Path Tracer: Add support to visualize mode for indirect lighting (only of emissives for now)

Default settings do not include indirect rays as they can be quite noisy

Add SER support as this seems to greatly speed up the performance of these traces.

#rnx

diff --git a/Engine/Shaders/Private/PathTracing/PathTracingDebug.usf b/Engine/Shaders/Private/PathTracing/PathTracingDebug.usf
--- a/Engine/Shaders/Private/PathTracing/PathTracingDebug.usf
+++ b/Engine/Shaders/Private/PathTracing/PathTracingDebug.usf
@@ -19,6 +19,12 @@
 #include "Light/PathTracingLightSampling.ush"
 #include "Light/PathTracingLightGrid.ush"
 
+#if PATH_TRACER_USE_SER
+
+#define NV_HITOBJECT_USE_MACRO_API
+#include "/Engine/Shared/ThirdParty/NVIDIA/nvHLSLExtns.h"
+
+#endif
 
 RWTexture2D<float4> RWSceneColor;
 RWTexture2D<float>  RWSceneDepth;
@@ -26,6 +32,7 @@ RaytracingAccelerationStructure TLAS;
 
 int DebugMode;
 int NumLightSamples;
+int NumIndirectSamples;
 
 float3 TraceShadowRay(FRayDesc Ray, float PathRoughness, float PreviousPathRoughness, uint MissShaderIndex, bool bCastShadows)
 {
@@ -214,7 +221,15 @@ PathTracingResult PathTracingDebugCore(FRayDesc Ray, uint2 LaunchIndex)
 			case PATH_TRACER_DEBUG_VIZ_IOR:							return MakeResult(PackedPayload.UnpackIor(), 0.0, Z);
 			case PATH_TRACER_DEBUG_VIZ_SHADING_MODEL:				return MakeResult(GetShadingModelColor(PackedPayload.GetShadingModelID()), 0.0, Z);
 			case PATH_TRACER_DEBUG_VIZ_LIGHTING_CHANNEL_MASK:		return MakeResult(GetLightingChannelMaskDebugColor(PackedPayload.GetPrimitiveLightingChannelMask()), 0.0, Z);
-			case PATH_TRACER_DEBUG_VIZ_WORLD_POSITION:				return MakeResult(TranslatedWorldPos * View.PreExposure, 0.0, Z);
+			case PATH_TRACER_DEBUG_VIZ_WORLD_POSITION: {
+				float3 AbsoluteWorldPosition = DFDemote(DFFastSubtract(TranslatedWorldPos, ResolvedView.PreViewTranslation));
+				const float CellWidth = 10.0f; // 10cm
+				int3 Coord = int3(round(AbsoluteWorldPosition / CellWidth));
+				uint HashIndex = Rand3DPCG32(Coord).x;
+				float3 Color = HUEtoRGB((HashIndex & 0xFFFFFF) * 5.96046447754e-08);
+				Color *= abs(dot(WorldNormal, Ray.Direction));
+				return MakeResult(Color * View.PreExposure, 0.0, Z);
+			}
 			case PATH_TRACER_DEBUG_VIZ_PRIMARY_RAYS: break; // fallthrough to implementation below
 			case PATH_TRACER_DEBUG_VIZ_WORLD_TANGENT: 				return MakeResult(PackedPayload.UnpackWorldTangent() * 0.5 + 0.5, 0.0, Z);
 			case PATH_TRACER_DEBUG_VIZ_LIGHT_GRID_COUNT:
@@ -380,8 +395,57 @@ PathTracingResult PathTracingDebugCore(FRayDesc Ray, uint2 LaunchIndex)
 			}
 #endif
 		}
-		Result.rgb /= NumLightSamples;
-		Result.rgb += PackedPayload.UnpackRadiance();
+		Result.rgb = Result.rgb * (NumLightSamples > 1 ? 1.0 / NumLightSamples : 1.0) + PackedPayload.UnpackRadiance();
+		float3 IndirectRadiance = 0;
+		for (int Sample = 0; Sample < NumIndirectSamples; Sample++)
+		{
+			RandomSequence RandSequence;
+			RandomSequence_Initialize(RandSequence, LaunchIndex, Sample, View.StateFrameIndex, NumIndirectSamples);
+			float3 RandSample = RandomSequence_GenerateSample3D(RandSequence);
+			FMaterialSample MaterialSample = SampleMaterial(-Ray.Direction, Payload, RandSample.xyz);
+
+			const uint RayFlags = 0;
+			const uint InstanceInclusionMask = PATHTRACER_MASK_INDIRECT;
+			const uint RayContributionToHitGroupIndex = RAY_TRACING_SHADER_SLOT_MATERIAL;
+			const uint MultiplierForGeometryContributionToShaderIndex = RAY_TRACING_NUM_SHADER_SLOTS;
+
+			FRayDesc Ray;
+			Ray.Origin = TranslatedWorldPos;
+			Ray.Direction = MaterialSample.Direction;
+			Ray.TMin = 0;
+			Ray.TMax = RAY_DEFAULT_T_MAX;
+
+			ApplyRayBias(Ray.Origin, PackedPayload.HitT, PackedPayload.UnpackWorldGeoNormal());
+
+			FPackedPathTracingPayload PackedPayloadIndirect = InitPathTracingPayload(MaterialSample.ScatterType, MaterialSample.Roughness);
+#if PATH_TRACER_USE_SER 
+			NvHitObject Hit;
+			NvTraceRayHitObject(
+				TLAS,
+				RayFlags,
+				InstanceInclusionMask,
+				RayContributionToHitGroupIndex,
+				MultiplierForGeometryContributionToShaderIndex,
+				MissShaderIndex,
+				Ray.GetNativeDesc(),
+				PackedPayloadIndirect,
+				Hit);
+			NvReorderThread(Hit);
+			NvInvokeHitObject(TLAS, Hit, PackedPayloadIndirect);
+#else
+			TraceRay(
+				TLAS,
+				RayFlags,
+				InstanceInclusionMask,
+				RayContributionToHitGroupIndex,
+				MultiplierForGeometryContributionToShaderIndex,
+				MissShaderIndex,
+				Ray.GetNativeDesc(),
+				PackedPayloadIndirect);
+#endif
+			IndirectRadiance += MaterialSample.Weight * PackedPayloadIndirect.UnpackRadiance();
+		}
+		Result.rgb += IndirectRadiance * (NumIndirectSamples > 1 ? 1.0 / NumIndirectSamples : 1.0);
 		Result.a = 0;
 	}
 	return MakeResult(Result.rgb * View.PreExposure, Result.a, Z);
@@ -393,6 +457,7 @@ RAY_TRACING_ENTRY_RAYGEN(PathTracingDebugRG)
 	uint2 DispatchIdx = DispatchRaysIndex().xy;
 	const uint2 DispatchDim = DispatchRaysDimensions().xy;
 
+	ResolvedView = ResolveView();
 
 	uint2 LaunchIndex = DispatchIdx + View.ViewRectMin.xy;
 
diff --git a/Engine/Source/Runtime/Renderer/Private/PathTracing.cpp b/Engine/Source/Runtime/Renderer/Private/PathTracing.cpp
--- a/Engine/Source/Runtime/Renderer/Private/PathTracing.cpp
+++ b/Engine/Source/Runtime/Renderer/Private/PathTracing.cpp
@@ -88,8 +88,6 @@ struct FPathTracingDebugModeInfo
 		REGISTER_VIZ_TYPE(TEXT("VolumeLightCount") 		, PATH_TRACER_DEBUG_VIZ_VOLUME_LIGHT_COUNT		);
 		REGISTER_VIZ_TYPE(TEXT("HitKind")            	, PATH_TRACER_DEBUG_VIZ_HITKIND             	);
 #undef REGISTER_VIZ_TYPE
-
-		CVarHelpText += TEXT("Default value: \"\"\n");
 	}
 
 	static FPathTracingDebugModeInfo& Get()
@@ -114,10 +112,16 @@ FAutoConsoleVariableRef CVarPathTracingVisualize(
 TAutoConsoleVariable<int32> CVarPathTracingVisualizeLightSamples(
 	TEXT("r.PathTracing.Visualize.LightSamples"),
 	1,
-	TEXT("Number of light samples for the path tracing light sampling debug mode (default: 1)"),
+	TEXT("Number of light samples for the path tracing light sampling visualize mode"),
 	ECVF_RenderThreadSafe
 );
 
+TAutoConsoleVariable<int32> CVarPathTracingVisualizeIndirectSamples(
+	TEXT("r.PathTracing.Visualize.IndirectSamples"),
+	0,
+	TEXT("Number of indirect samples for the path tracing visualize mode"),
+	ECVF_RenderThreadSafe
+);
 
 TAutoConsoleVariable<int32> CVarPathTracingCompactionDepth(
 	TEXT("r.PathTracing.CompactionDepth"),
@@ -1376,9 +1380,18 @@ class FPathTracingDebugRG : public FGlobalShader
 	DECLARE_GLOBAL_SHADER(FPathTracingDebugRG)
 	SHADER_USE_ROOT_PARAMETER_STRUCT(FPathTracingDebugRG, FGlobalShader)
 
+	class FUseSER : SHADER_PERMUTATION_BOOL("PATH_TRACER_USE_SER");
+	using FPermutationDomain = TShaderPermutationDomain<FUseSER>;
 
 	static bool ShouldCompilePermutation(const FGlobalShaderPermutationParameters& Parameters)
 	{
+		// Only compile SER extensions on platforms that could support it
+		FPermutationDomain PermutationVector(Parameters.PermutationId);
+		if (PermutationVector.Get<FUseSER>() && !FDataDrivenShaderPlatformInfo::GetSupportsShaderExecutionReordering(Parameters.Platform))
+		{
+			return false;
+		}
+
 		return PathTracing::ShouldCompilePathTracingShadersForProject(Parameters.Platform);
 	}
 
@@ -1405,13 +1418,16 @@ class FPathTracingDebugRG : public FGlobalShader
 		SHADER_PARAMETER_RDG_BUFFER_SRV(RaytracingAccelerationStructure, TLAS)
 		SHADER_PARAMETER_STRUCT_REF(FViewUniformShaderParameters, ViewUniformBuffer)
 		SHADER_PARAMETER(int32, DebugMode)
+		SHADER_PARAMETER(int32, NumLightSamples)
+		SHADER_PARAMETER(int32, NumIndirectSamples)
+
 
 		// scene lights
 		SHADER_PARAMETER_RDG_BUFFER_SRV(StructuredBuffer<FPathTracingLight>, SceneLights)
 		SHADER_PARAMETER(uint32, SceneLightCount)
 		SHADER_PARAMETER(uint32, SceneVisibleLightCount)
-		SHADER_PARAMETER(uint32, NumLightSamples)
 		SHADER_PARAMETER_STRUCT_INCLUDE(FPathTracingLightGrid, LightGridParameters)
+
 	
 		// scene decals
 		SHADER_PARAMETER_RDG_UNIFORM_BUFFER(FRayTracingDecals, DecalParameters)
@@ -3215,7 +3231,11 @@ void FDeferredShadingSceneRenderer::PreparePathTracing(const FViewInfo& View, co
 	{
 		if (GetPathTracingVisualizationMode() >= 0)
 		{
-			auto RayGenShader = GetGlobalShaderMap(ViewFamily.GetShaderPlatform())->GetShader<FPathTracingDebugRG>();
+			const bool bUseSER = GRHIGlobals.SupportsShaderExecutionReordering && CVarPathTracingShaderExecutionReordering.GetValueOnRenderThread();
+
+			FPathTracingDebugRG::FPermutationDomain PermutationVector;
+			PermutationVector.Set<FPathTracingDebugRG::FUseSER>(bUseSER);
+			auto RayGenShader = GetGlobalShaderMap(ViewFamily.GetShaderPlatform())->GetShader<FPathTracingDebugRG>(PermutationVector);
 			OutRayGenShaders.Add(RayGenShader.GetRayTracingShader());
 		}
 		else
@@ -3376,7 +3396,8 @@ void FDeferredShadingSceneRenderer::RenderPathTracing(
 	FRDGTextureRef SceneDepthOutputTexture,
 	FPathTracingResources& PathTracingResources)
 {
-	RDG_EVENT_SCOPE(GraphBuilder, "Path Tracing");
+	RDG_EVENT_SCOPE_STAT(GraphBuilder, PathTracing, "PathTracing");
+	RDG_GPU_STAT_SCOPE(GraphBuilder, PathTracing);
 
 	// To make the GPU profiler work for path tracing with multi-GPU, we need the root GPU profiling scope (marked as "Unaccounted") to be on all GPUs,
 	// as the profiler discards events where any event in the hierarchy wasn't on a given GPU.  So in the parent scene render code, we set the GPU mask
@@ -3405,7 +3426,6 @@ void FDeferredShadingSceneRenderer::RenderPathTracing(
 		const int32 DispatchResY = View.ViewRect.Size().Y;
 
 		// simplified pass for debugging purposes
-		RDG_GPU_STAT_SCOPE(GraphBuilder, PathTracing);
 		FRDGTextureDesc DepthDesc = FRDGTextureDesc::Create2D(
 			View.ViewRect.Size(),
 			PF_R32_FLOAT,
@@ -3421,7 +3441,8 @@ void FDeferredShadingSceneRenderer::RenderPathTracing(
 		PassParameters->ViewUniformBuffer = View.ViewUniformBuffer;
 		PassParameters->DebugMode = DebugMode;
 
-		PassParameters->NumLightSamples = FMath::Max(CVarPathTracingVisualizeLightSamples.GetValueOnRenderThread(), 1);
+		PassParameters->NumLightSamples = CVarPathTracingVisualizeLightSamples.GetValueOnRenderThread();
+		PassParameters->NumIndirectSamples = CVarPathTracingVisualizeIndirectSamples.GetValueOnRenderThread();
 		SetLightParameters(
 			GraphBuilder,
 			Scene,
@@ -3439,10 +3460,13 @@ void FDeferredShadingSceneRenderer::RenderPathTracing(
 		PassParameters->Scene = View.GetSceneUniforms().GetBuffer(GraphBuilder);
 		PassParameters->NaniteRayTracing = Nanite::GRayTracingManager.GetUniformBuffer();
 
-		TShaderMapRef<FPathTracingDebugRG> RayGenShader(View.ShaderMap);
+		const bool bUseSER = GRHIGlobals.SupportsShaderExecutionReordering && CVarPathTracingShaderExecutionReordering.GetValueOnRenderThread();
+		FPathTracingDebugRG::FPermutationDomain Permutation;
+		Permutation.Set<FPathTracingDebugRG::FUseSER>(bUseSER);
+		TShaderMapRef<FPathTracingDebugRG> RayGenShader(View.ShaderMap, Permutation);
 		ClearUnusedGraphResources(RayGenShader, PassParameters);
 		GraphBuilder.AddPass(
-			RDG_EVENT_NAME("Path Tracing"),
+			RDG_EVENT_NAME("Path Tracing Visualize"),
 		  	PassParameters,
 		  	ERDGPassFlags::Compute,
 			[PassParameters, RayGenShader, &View, DispatchResX, DispatchResY](FRHICommandList& RHICmdList)
@@ -3926,10 +3950,7 @@ void FDeferredShadingSceneRenderer::RenderPathTracing(
 			for (int32 GPUIndex : GPUMask)
 			{
 				RDG_GPU_MASK_SCOPE(GraphBuilder, FRHIGPUMask::FromIndex(GPUIndex));
-				RDG_EVENT_SCOPE_CONDITIONAL_STAT(GraphBuilder, NumGPUs > 1, PathTracing, "Path Tracing GPU%d", GPUIndex);
-#if WITH_MGPU
-				RDG_GPU_STAT_SCOPE(GraphBuilder, PathTracing);
-#endif
+				RDG_EVENT_SCOPE_CONDITIONAL(GraphBuilder, NumGPUs >  1, "Path Tracing GPU%d", GPUIndex);
 				for (int32 TileY = 0; TileY < DispatchResY; TileY += DispatchSize)
 				{
 					for (int32 TileX = 0; TileX < DispatchResX; TileX += DispatchSize)
