VolumetricCloud - fixed cloud shadow to not be traced when Cloud Shadow Strength = 0.

diff --git a/Engine/Source/Runtime/Renderer/Private/VolumetricCloudRendering.cpp b/Engine/Source/Runtime/Renderer/Private/VolumetricCloudRendering.cpp
--- a/Engine/Source/Runtime/Renderer/Private/VolumetricCloudRendering.cpp
+++ b/Engine/Source/Runtime/Renderer/Private/VolumetricCloudRendering.cpp
@@ -335,7 +335,7 @@ bool VolumetricCloudWantsToSampleLocalLights(const FScene* Scene, const FEngineS
 
 static bool ShouldRenderCloudShadowmap(const FLightSceneProxy* AtmosphericLight)
 {
-	return CVarVolumetricCloudShadowMap.GetValueOnRenderThread() > 0 && AtmosphericLight && AtmosphericLight->GetCastCloudShadows();
+	return CVarVolumetricCloudShadowMap.GetValueOnRenderThread() > 0 && AtmosphericLight && AtmosphericLight->GetCastCloudShadows() && AtmosphericLight->GetCloudShadowStrength() > 0;
 }
 
 bool ShouldVolumetricCloudTraceWithMinMaxDepth(const FViewInfo& ViewInfo)
@@ -3161,8 +3161,8 @@ bool LightMayCastCloudShadow(const FScene* Scene, const FViewInfo& View, const F
 
 	if (Scene->GetVolumetricCloudSceneInfo()
 		//&& (View.VolumetricCloudShadowRenderTarget[0] || View.VolumetricCloudShadowRenderTarget[1]) // We cannot check for RTs when this function might be called earlier in the frame (e.g. Lumen), before the RTs have been generated.
-		&& ((Light0 && Light0->Proxy == LightSceneInfo->Proxy && Light0->Proxy->GetCastCloudShadows() && Light0->Proxy->GetCloudShadowOnSurfaceStrength() > 0.f)
-			|| (Light1 && Light1->Proxy == LightSceneInfo->Proxy && Light1->Proxy->GetCastCloudShadows() && Light1->Proxy->GetCloudShadowOnSurfaceStrength() > 0.f)))
+		&& ((Light0 && Light0->Proxy == LightSceneInfo->Proxy && Light0->Proxy->GetCastCloudShadows() && Light0->Proxy->GetCloudShadowOnSurfaceStrength() > 0.f && Light0->Proxy->GetCloudShadowStrength() > 0.f)
+			|| (Light1 && Light1->Proxy == LightSceneInfo->Proxy && Light1->Proxy->GetCastCloudShadows() && Light1->Proxy->GetCloudShadowOnSurfaceStrength() > 0.f && Light1->Proxy->GetCloudShadowStrength() > 0.f)))
 	{
 		return true;
 	}
