Convert ray tracing tasks to use RDG setup tasks so that they respect RDG CVars.

diff --git a/Engine/Source/Runtime/Renderer/Private/DeferredShadingRenderer.cpp b/Engine/Source/Runtime/Renderer/Private/DeferredShadingRenderer.cpp
--- a/Engine/Source/Runtime/Renderer/Private/DeferredShadingRenderer.cpp
+++ b/Engine/Source/Runtime/Renderer/Private/DeferredShadingRenderer.cpp
@@ -1638,7 +1638,7 @@ void FDeferredShadingSceneRenderer::Render(FRDGBuilder& GraphBuilder, const FSce
 				GetViewPipelineState(ReferenceView).DiffuseIndirectMethod,
 				GetViewPipelineState(ReferenceView).ReflectionsMethod);
 
-			RayTracing::BeginGatherInstances(*InitViewTaskDatas.RayTracingGatherInstances, InitViewTaskDatas.VisibilityTaskData->GetFrustumCullTask());
+			RayTracing::BeginGatherInstances(GraphBuilder, *InitViewTaskDatas.RayTracingGatherInstances, InitViewTaskDatas.VisibilityTaskData->GetFrustumCullTask());
 		}
 #endif
 	}
diff --git a/Engine/Source/Runtime/Renderer/Private/RayTracing/RayTracing.cpp b/Engine/Source/Runtime/Renderer/Private/RayTracing/RayTracing.cpp
--- a/Engine/Source/Runtime/Renderer/Private/RayTracing/RayTracing.cpp
+++ b/Engine/Source/Runtime/Renderer/Private/RayTracing/RayTracing.cpp
@@ -1737,23 +1737,23 @@ namespace RayTracing
 		TaskData.DynamicRayTracingInstancesContext.AddInstancesToScene(RayTracingScene, RayTracingSBT, SharedBufferGenerationID);
 	}
 
-	void BeginGatherInstances(FGatherInstancesTaskData& TaskData, UE::Tasks::FTask FrustumCullTask)
+	void BeginGatherInstances(FRDGBuilder& GraphBuilder, FGatherInstancesTaskData& TaskData, UE::Tasks::FTask FrustumCullTask)
 	{
 		const bool bUsingReferenceBasedResidency = IsRayTracingUsingReferenceBasedResidency();
 
 		UE::Tasks::FTask CacheRayTracingPrimitivesTask = TaskData.Scene.GetCacheRayTracingPrimitivesTask();
 
-		TaskData.GatherRelevantPrimitivesTask = UE::Tasks::Launch(UE_SOURCE_LOCATION, [&TaskData, bUsingReferenceBasedResidency]
+		TaskData.GatherRelevantPrimitivesTask = GraphBuilder.AddSetupTask([&TaskData, bUsingReferenceBasedResidency]
 			{
 				FTaskTagScope Scope(ETaskTag::EParallelRenderingThread);
 				GatherRelevantPrimitives(TaskData, bUsingReferenceBasedResidency);
-			}, UE::Tasks::Prerequisites(CacheRayTracingPrimitivesTask, FrustumCullTask), UE::Tasks::ETaskPriority::High);
+			}, MakeArrayView({ CacheRayTracingPrimitivesTask, FrustumCullTask }), UE::Tasks::ETaskPriority::High);
 
 		static const auto ICVarStaticMeshLODDistanceScale = IConsoleManager::Get().FindConsoleVariable(TEXT("r.StaticMeshLODDistanceScale"));
 		const float LODScaleCVarValue = ICVarStaticMeshLODDistanceScale->GetFloat();
 		const int32 ForcedLODLevel = GetCVarForceLOD();
 
-		TaskData.GatherRelevantStaticPrimitivesTask = UE::Tasks::Launch(UE_SOURCE_LOCATION, [&TaskData, LODScaleCVarValue, ForcedLODLevel, bUsingReferenceBasedResidency]
+		TaskData.GatherRelevantStaticPrimitivesTask = GraphBuilder.AddSetupTask([&TaskData, LODScaleCVarValue, ForcedLODLevel, bUsingReferenceBasedResidency]
 			{
 				FTaskTagScope Scope(ETaskTag::EParallelRenderingThread);
 				GatherRelevantStaticPrimitives(TaskData, LODScaleCVarValue, ForcedLODLevel, bUsingReferenceBasedResidency);
@@ -1762,14 +1762,14 @@ namespace RayTracing
 		UE::Tasks::FTask GPUSkinCacheTask = TaskData.Scene.GetGPUSkinCacheTask();
 
 		// TODO: Could gather dynamic ray tracing instances using multiple tasks / FDynamicRayTracingInstancesContext
-		TaskData.GatherDynamicRayTracingInstancesTask = UE::Tasks::Launch(UE_SOURCE_LOCATION, [&TaskData]
+		TaskData.GatherDynamicRayTracingInstancesTask = GraphBuilder.AddSetupTask([&TaskData]
 			{
 				FTaskTagScope Scope(ETaskTag::EParallelRenderingThread);
 				TaskData.DynamicRayTracingInstancesContext.GatherDynamicRayTracingInstances(TaskData.DynamicPrimitivesIndices);
-			}, UE::Tasks::Prerequisites(TaskData.GatherRelevantPrimitivesTask, GPUSkinCacheTask));
+			}, MakeArrayView({ TaskData.GatherRelevantPrimitivesTask, GPUSkinCacheTask }));
 
 
-		TaskData.AddUsedStreamingHandlesTask = UE::Tasks::Launch(UE_SOURCE_LOCATION, [&TaskData, LODScaleCVarValue, ForcedLODLevel, bUsingReferenceBasedResidency]
+		TaskData.AddUsedStreamingHandlesTask = GraphBuilder.AddSetupTask([&TaskData, LODScaleCVarValue, ForcedLODLevel, bUsingReferenceBasedResidency]
 			{
 				FTaskTagScope Scope(ETaskTag::EParallelRenderingThread);
 
@@ -1826,8 +1826,8 @@ namespace RayTracing
 		TaskData.DynamicRayTracingInstancesContext.CollectRDGResources(GraphBuilder);
 
 		// This adds final instances to the ray tracing scene and must be done before FRayTracingScene::BuildInitializationData().
-		UE::Tasks::FTask AddInstancesTask = UE::Tasks::Launch(UE_SOURCE_LOCATION,
-			[&TaskData, &Scene, &View, &RayTracingScene, &RayTracingSBT, SharedBufferGenerationID]()
+		UE::Tasks::FTask AddInstancesTask = GraphBuilder.AddSetupTask(
+			[&TaskData, &Scene, &View, &RayTracingScene, &RayTracingSBT, SharedBufferGenerationID]
 			{
 				FTaskTagScope TaskTagScope(ETaskTag::EParallelRenderingThread);
 
@@ -1857,8 +1857,8 @@ namespace RayTracing
 			}, TaskData.GatherRelevantStaticPrimitivesTask);
 
 		// Scene init task can run only when all pre-init tasks are complete (including culling tasks that are spawned while adding instances)
-		View.RayTracingSceneInitTask = UE::Tasks::Launch(UE_SOURCE_LOCATION,
-			[&RayTracingScene]()
+		View.RayTracingSceneInitTask = GraphBuilder.AddSetupTask(
+			[&RayTracingScene]
 			{
 				FTaskTagScope TaskTagScope(ETaskTag::EParallelRenderingThread);
 
@@ -1867,8 +1867,8 @@ namespace RayTracing
 			}, AddInstancesTask);
 
 		// Finalizing VisibleRayTracingShaderBindings can run only when all pre-init tasks are complete (including culling tasks that are spawned while adding instances)
-		View.VisibleRayTracingShaderBindingsFinalizeTask = UE::Tasks::Launch(UE_SOURCE_LOCATION,
-			[&View, &RayTracingSBT]()
+		View.VisibleRayTracingShaderBindingsFinalizeTask = GraphBuilder.AddSetupTask(
+			[&View, &RayTracingSBT]
 			{
 				FTaskTagScope TaskTagScope(ETaskTag::EParallelRenderingThread);
 
diff --git a/Engine/Source/Runtime/Renderer/Private/RayTracing/RayTracing.h b/Engine/Source/Runtime/Renderer/Private/RayTracing/RayTracing.h
--- a/Engine/Source/Runtime/Renderer/Private/RayTracing/RayTracing.h
+++ b/Engine/Source/Runtime/Renderer/Private/RayTracing/RayTracing.h
@@ -52,7 +52,7 @@ namespace RayTracing
 		EDiffuseIndirectMethod DiffuseIndirectMethod,
 		EReflectionsMethod ReflectionsMethod);
 	
-	void BeginGatherInstances(FGatherInstancesTaskData& TaskData, UE::Tasks::FTask FrustumCullTask);
+	void BeginGatherInstances(FRDGBuilder& GraphBuilder, FGatherInstancesTaskData& TaskData, UE::Tasks::FTask FrustumCullTask);
 
 	// Fills RayTracingScene instance list for the given View and adds relevant ray tracing data to the view. Does not reset previous scene contents.
 	// This function must run on render thread
