Remove some unnecessary transition tracking at start of virtual texture update.
All page table update/transitions are handled by FinalizeRequests().
Also make sure that we call RemovePending() for virtual textures on all update paths including when we early out.
#rnx

diff --git a/Engine/Source/Runtime/Renderer/Private/VT/VirtualTextureSystem.cpp b/Engine/Source/Runtime/Renderer/Private/VT/VirtualTextureSystem.cpp
--- a/Engine/Source/Runtime/Renderer/Private/VT/VirtualTextureSystem.cpp
+++ b/Engine/Source/Runtime/Renderer/Private/VT/VirtualTextureSystem.cpp
@@ -2829,11 +2829,16 @@ TUniquePtr<FVirtualTextureUpdater> FVirtualTextureSystem::BeginUpdate(FRDGBuilde
 {
 	check(IsInRenderingThread());
 	check(!bUpdating);
-
 	checkf(Producers.HasPendingCallbacks() == false, TEXT("FVirtualTextureSystem::CallPendingCallbacks(), called in UpdateAllPrimitiveSceneInfos(), must run before FVirtualTextureSystem::BeginUpdate()"));
 
 	AllocateResources(GraphBuilder);
 
+	// Optional early out after AllocateResources().
+	if (!Settings.bEnablePageRequests)
+	{
+		return {};
+	}
+
 	RDG_CSV_STAT_EXCLUSIVE_SCOPE(GraphBuilder, VirtualTextureSystem_Update);
 	TRACE_CPUPROFILER_EVENT_SCOPE(FVirtualTextureSystem::BeginUpdate);
 	SCOPE_CYCLE_COUNTER(STAT_VirtualTextureSystem_Update);
@@ -2841,22 +2846,6 @@ TUniquePtr<FVirtualTextureUpdater> FVirtualTextureSystem::BeginUpdate(FRDGBuilde
 	RDG_EVENT_SCOPE_STAT(GraphBuilder, VirtualTexture, "VirtualTextureBeginUpdate");
 	RDG_GPU_STAT_SCOPE(GraphBuilder, VirtualTexture);
 
-	if (!Settings.bEnablePageRequests)
-	{
-		FRDGExternalAccessQueue ExternalAccessQueue;
-
-		for (uint32 ID = 0; ID < MaxSpaces; ID++)
-		{
-			if (Spaces[ID])
-			{
-				Spaces[ID]->FinalizeTextures(GraphBuilder, ExternalAccessQueue);
-			}
-		}
-
-		ExternalAccessQueue.Submit(GraphBuilder);
-		return {};
-	}
-
 	// Update Adaptive VTs. This can trigger allocation/destruction of VTs and must happen before the flush below.
 	{
 		SCOPE_CYCLE_COUNTER(STAT_ProcessRequests_UpdateAdaptiveAllocations);
@@ -2905,6 +2894,7 @@ TUniquePtr<FVirtualTextureUpdater> FVirtualTextureSystem::BeginUpdate(FRDGBuilde
 	if (AllocatedVTs.Num() == 0)
 	{
 		MappedTilesToProduce.Reset();
+		ReleasePendingSpaces();
 		return {};
 	}
 
