Mobile DFAO implementation
Mobile Deferred MultiPass has its own DFAO generation renderpass, samples the DFAO texture and composites it with other AO like SSAO
Mobile Forward and Deferred Single Pass will modulate the DFAO result into the SSAO texture and sample it.

diff --git a/Engine/Shaders/Private/DistanceFieldLightingPost.usf b/Engine/Shaders/Private/DistanceFieldLightingPost.usf
--- a/Engine/Shaders/Private/DistanceFieldLightingPost.usf
+++ b/Engine/Shaders/Private/DistanceFieldLightingPost.usf
@@ -368,8 +368,11 @@ void AOUpsamplePS(
 	,out float4 OutSceneColor : SV_Target0
 	)
 {
+#if SHADING_PATH_MOBILE
+	FGBufferData GBuffer = MobileFetchAndDecodeGBuffer(UVAndScreenPos.xy, UVAndScreenPos.zw);
+#else
 	FGBufferData GBuffer = GetGBufferData(UVAndScreenPos.xy);
-
+#endif
 	// The rectangle is drawn according to GBuffer rect,
 	// while the bent normal is always offset at (0, 0)
 	const float2 BentNormalUV = UVAndScreenPos.xy - View.ViewRectMin.xy * View.BufferSizeAndInvSize.zw;
@@ -381,4 +384,4 @@ void AOUpsamplePS(
 #else
 	OutSceneColor = float4(length(BentNormal).xxx, 1.0f);
 #endif
-}
+}
\ No newline at end of file
diff --git a/Engine/Shaders/Private/DistanceFieldScreenGridLighting.usf b/Engine/Shaders/Private/DistanceFieldScreenGridLighting.usf
--- a/Engine/Shaders/Private/DistanceFieldScreenGridLighting.usf
+++ b/Engine/Shaders/Private/DistanceFieldScreenGridLighting.usf
@@ -26,7 +26,12 @@ void ComputeDistanceFieldNormalPS(
 #if SUBTRATE_GBUFFER_FORMAT==1 
 	float3 WorldNormal = SubstrateUnpackTopLayerData(Substrate.TopLayerTexture.Load(uint3(PixelCoord, 0))).WorldNormal;
 #else
-	float3 WorldNormal = GetGBufferData(ScreenUV).WorldNormal.xyz;
+#if SHADING_PATH_MOBILE
+	FGBufferData GBuffer = MobileFetchAndDecodeGBuffer(ScreenUV, PixelCoord);
+#else
+	FGBufferData GBuffer = GetGBufferData(ScreenUV);
+#endif
+	float3 WorldNormal = GBuffer.WorldNormal.xyz;
 #endif
 
 	OutColor = EncodeDownsampledGBuffer(WorldNormal, SceneDepth);
@@ -48,7 +53,12 @@ void ComputeDistanceFieldNormalCS(
 #if SUBTRATE_GBUFFER_FORMAT==1 
 	float3 WorldNormal = SubstrateUnpackTopLayerData(Substrate.TopLayerTexture.Load(uint3(PixelCoord, 0))).WorldNormal;
 #else
-	float3 WorldNormal = GetGBufferData(ScreenUV).WorldNormal.xyz;
+#if SHADING_PATH_MOBILE
+	FGBufferData GBuffer = MobileFetchAndDecodeGBuffer(ScreenUV, PixelCoord);
+#else
+	FGBufferData GBuffer = GetGBufferData(ScreenUV);
+#endif
+	float3 WorldNormal = GBuffer.WorldNormal.xyz;
 #endif
 
 	RWDistanceFieldNormal[DispatchThreadId.xy] = EncodeDownsampledGBuffer(WorldNormal, SceneDepth);
@@ -378,9 +388,11 @@ void CombineConeVisibilityCS(
 	float SceneDepth;
 	GetDownsampledGBuffer(BaseLevelScreenUV, WorldNormal, SceneDepth);
 
+#if ROTATE_WITH_TANGENT
 	float3 TangentX;
 	float3 TangentY;
 	FindBestAxisVectors2(WorldNormal, TangentX, TangentY);
+#endif
 
 	uint InputBaseIndex = InputCoordinate.y * ScreenGridConeVisibilitySize.x + InputCoordinate.x;
 
@@ -396,7 +408,11 @@ void CombineConeVisibilityCS(
 		float ConeVisibility = asfloat(ScreenGridConeVisibility[ConeIndex * ScreenGridConeVisibilitySize.x * ScreenGridConeVisibilitySize.y + InputBaseIndex]);
 		
 		float3 ConeDirection = SampleDirections[ConeIndex].xyz;
+#if ROTATE_WITH_TANGENT
 		float3 RotatedConeDirection = ConeDirection.x * TangentX + ConeDirection.y * TangentY + ConeDirection.z * WorldNormal;
+#else
+		float3 RotatedConeDirection = ConeDirection;
+#endif
 		UnoccludedDirection += ConeVisibility * RotatedConeDirection;
 	}
 
diff --git a/Engine/Shaders/Private/MobileDeferredShading.usf b/Engine/Shaders/Private/MobileDeferredShading.usf
--- a/Engine/Shaders/Private/MobileDeferredShading.usf
+++ b/Engine/Shaders/Private/MobileDeferredShading.usf
@@ -34,6 +34,11 @@
 #endif
 #include "HairStrands/HairStrandsEnvironmentLightingCommon.ush"
 
+#if APPLY_SKY_SHADOWING
+#include "SkyLightingDiffuseShared.ush"
+#include "DistanceFieldAOShared.ush"
+#endif
+
 float4x4 TranslatedWorldToLight;
 float2 LightFunctionParameters2;
 
@@ -109,6 +114,7 @@ void ReflectionEnvironmentSkyLighting(
 	float3 TranslatedWorldPosition,
 	half3 ReflectionVector,
 	uint GridIndex,
+	half GatheredAmbientOcclusion,
 	inout FLightAccumulator DirectLighting)
 {
 	half IndirectIrradiance = GBuffer.IndirectIrradiance;
@@ -118,13 +124,24 @@ void ReflectionEnvironmentSkyLighting(
 	half3 N = GBuffer.WorldNormal;
 	half3 V = -CameraVector;
 	half NoV = saturate(abs(dot(N, V)) + 1e-5);
+#if APPLY_SKY_SHADOWING
+	half3 BentNormal  = GBuffer.WorldNormal;
+	const float2 BentNormalUV = (SvPosition.xy - View.ViewRectMin.xy) * View.BufferSizeAndInvSize.zw;
+	BentNormal = UpsampleDFAO(BentNormalUV, GBuffer.Depth, GBuffer.WorldNormal);
+#endif
 
 	half3 DiffuseColor = GBuffer.DiffuseColor;
 	half3 SpecularColor = GBuffer.SpecularColor;
 	RemapClearCoatDiffuseAndSpecularColor(GBuffer, NoV, DiffuseColor, SpecularColor);
-
+#if APPLY_SKY_SHADOWING
+	//const float3 DiffuseColor = GBuffer.BaseColor; // not used
+	const float2 BufferUV = SvPosition * View.BufferSizeAndInvSize.zw; // not used
+	float2 ScreenPosition = SvPositionToScreenPosition(SvPosition).xy;
+	// Sky lighting (already pre-exposed)
+	half3 SkyLighting = SkyLightDiffuse(GBuffer, GatheredAmbientOcclusion, BufferUV, ScreenPosition, BentNormal, DiffuseColor);
+#else
 	half3 SkyLighting = SkyLightDiffuseMobile(GBuffer, DiffuseColor, V, IndirectIrradiance);
-
+#endif
 	const bool bNeedsSeparateSubsurfaceLightAccumulation = UseSubsurfaceProfile(GBuffer.ShadingModelID);
 	LightAccumulator_Add(DirectLighting, SkyLighting, SkyLighting, 1.0f, bNeedsSeparateSubsurfaceLightAccumulation);
 #endif // ENABLE_SKY_LIGHT 
@@ -208,7 +225,7 @@ void MobileDirectionalLightPS(
 
 #if ENABLE_CLUSTERED_REFLECTION || ENABLE_SKY_LIGHT || ENABLE_PLANAR_REFLECTION || MOBILE_SSR_ENABLED
 	// If we have a single directional light, apply relfection and sky contrubution here
-	ReflectionEnvironmentSkyLighting(GBuffer, SvPosition, CameraVector, TranslatedWorldPosition, ReflectionVector, GridIndex, DirectLighting);
+	ReflectionEnvironmentSkyLighting(GBuffer, SvPosition, CameraVector, TranslatedWorldPosition, ReflectionVector, GridIndex, GatheredAmbientOcclusion, DirectLighting);
 #endif
 
 	half3 Color = DirectLighting.TotalLight;
@@ -398,7 +415,7 @@ void MobileReflectionEnvironmentSkyLightingPS(
 #endif
 	
 	FLightAccumulator DirectLighting = (FLightAccumulator)0;
-	ReflectionEnvironmentSkyLighting(GBuffer, SvPosition, CameraVector, TranslatedWorldPosition, ReflectionVector, GridIndex, DirectLighting);
+	ReflectionEnvironmentSkyLighting(GBuffer, SvPosition, CameraVector, TranslatedWorldPosition, ReflectionVector, GridIndex, GatheredAmbientOcclusion, DirectLighting);
 	half3 ReflectionAndSky = DirectLighting.TotalLight;
 
 	ReflectionAndSky *= View.PreExposure;
diff --git a/Engine/Source/Runtime/Engine/Private/SceneManagement.cpp b/Engine/Source/Runtime/Engine/Private/SceneManagement.cpp
--- a/Engine/Source/Runtime/Engine/Private/SceneManagement.cpp
+++ b/Engine/Source/Runtime/Engine/Private/SceneManagement.cpp
@@ -48,7 +48,7 @@ bool DoesPlatformSupportDistanceFieldShadowing(EShaderPlatform Platform)
 
 bool DoesPlatformSupportDistanceFieldAO(EShaderPlatform Platform)
 {
-	return DoesPlatformSupportDistanceFields(Platform);
+	return DoesPlatformSupportDistanceFields(Platform) && (!IsMobilePlatform(Platform) || IsMobileDistanceFieldAOEnabled(Platform));
 }
 
 bool DoesProjectSupportDistanceFields()
diff --git a/Engine/Source/Runtime/RenderCore/Private/RenderUtils.cpp b/Engine/Source/Runtime/RenderCore/Private/RenderUtils.cpp
--- a/Engine/Source/Runtime/RenderCore/Private/RenderUtils.cpp
+++ b/Engine/Source/Runtime/RenderCore/Private/RenderUtils.cpp
@@ -550,7 +550,7 @@ RENDERCORE_API ShaderPlatformMaskType GMobileAmbientOcclusionPlatformMask;
 
 RENDERCORE_API bool IsMobileAmbientOcclusionEnabled(const FStaticShaderPlatform Platform)
 {
-	return IsMobilePlatform(Platform) && GMobileAmbientOcclusionPlatformMask[(int)Platform];
+	return IsMobilePlatform(Platform) && (GMobileAmbientOcclusionPlatformMask[(int)Platform] || IsMobileDistanceFieldAOEnabled(Platform));
 }
 
 RENDERCORE_API bool AreMobileScreenSpaceReflectionsEnabled(const FStaticShaderPlatform Platform)
@@ -1565,6 +1565,12 @@ bool IsWaterSeparateMainDirLightEnabled(const FStaticShaderPlatform Platform)
 	return IsWaterDistanceFieldShadowEnabled(Platform) || IsWaterVirtualShadowMapFilteringEnabled(Platform);
 }
 
+bool IsMobileDistanceFieldAOEnabled(const FStaticShaderPlatform Platform)
+{
+	static FShaderPlatformCachedIniValue<int32> MobileDistanceFieldAO(TEXT("r.Mobile.DistanceFieldAO"));
+	return MobileDistanceFieldAO.Get(Platform) > 0;
+}
+
 bool UseGPUScene(const FStaticShaderPlatform Platform, const FStaticFeatureLevel FeatureLevel)
 {
 	if (FeatureLevel == ERHIFeatureLevel::ES3_1)
diff --git a/Engine/Source/Runtime/RenderCore/Public/RenderUtils.h b/Engine/Source/Runtime/RenderCore/Public/RenderUtils.h
--- a/Engine/Source/Runtime/RenderCore/Public/RenderUtils.h
+++ b/Engine/Source/Runtime/RenderCore/Public/RenderUtils.h
@@ -361,6 +361,9 @@ RENDERCORE_API bool IsWaterDistanceFieldShadowEnabled(const FStaticShaderPlatfor
 /** Returns if water needs a separate main directional light texture. This is for a platofrm so can be used at cook time or at runtime. */
 RENDERCORE_API bool IsWaterSeparateMainDirLightEnabled(const FStaticShaderPlatform Platform);
 
+/** Returns if Mobile Distance Field AO Enabled. */
+RENDERCORE_API bool IsMobileDistanceFieldAOEnabled(const FStaticShaderPlatform Platform);
+
 RENDERCORE_API bool UseGPUScene(const FStaticShaderPlatform Platform, const FStaticFeatureLevel FeatureLevel);
 
 RENDERCORE_API bool UseGPUScene(const FStaticShaderPlatform Platform);
diff --git a/Engine/Source/Runtime/Renderer/Private/CapsuleShadowRendering.cpp b/Engine/Source/Runtime/Renderer/Private/CapsuleShadowRendering.cpp
--- a/Engine/Source/Runtime/Renderer/Private/CapsuleShadowRendering.cpp
+++ b/Engine/Source/Runtime/Renderer/Private/CapsuleShadowRendering.cpp
@@ -1202,10 +1202,9 @@ bool ShouldPrepareForDFInsetIndirectShadow(EShaderPlatform ShaderPlatform, const
 	return SupportsCapsuleIndirectShadows(ShaderPlatform) && EngineShowFlags.CapsuleShadows;
 }
 
-void FDeferredShadingSceneRenderer::RenderCapsuleShadowsForMovableSkylight(
+void FSceneRenderer::RenderCapsuleShadowsForMovableSkylight(
 	FRDGBuilder& GraphBuilder,
 	const FViewInfo& View,
-	TRDGUniformBufferRef<FSceneTextureUniformParameters> SceneTexturesUniformBuffer,
 	FRDGTextureRef& BentNormalOutput) const
 {
 	if (SupportsCapsuleIndirectShadows(ShaderPlatform)
diff --git a/Engine/Source/Runtime/Renderer/Private/DeferredShadingRenderer.h b/Engine/Source/Runtime/Renderer/Private/DeferredShadingRenderer.h
--- a/Engine/Source/Runtime/Renderer/Private/DeferredShadingRenderer.h
+++ b/Engine/Source/Runtime/Renderer/Private/DeferredShadingRenderer.h
@@ -710,27 +710,6 @@ private:
 		const FSceneTextures& SceneTextures,
 		TArray<FRDGTextureRef>& DynamicBentNormalAOTextures);
 
-	/** Render Ambient Occlusion using mesh distance fields and the surface cache, which supports dynamic rigid meshes. */
-	void RenderDistanceFieldLighting(
-		FRDGBuilder& GraphBuilder,
-		const FSceneTextures& SceneTextures,
-		const class FDistanceFieldAOParameters& Parameters,
-		TArray<FRDGTextureRef>& OutDynamicBentNormalAOTextures,
-		bool bModulateToSceneColor,
-		bool bVisualizeAmbientOcclusion);
-
-	/** Render Ambient Occlusion using mesh distance fields on a screen based grid. */
-	void RenderDistanceFieldAOScreenGrid(
-		FRDGBuilder& GraphBuilder,
-		const FSceneTextures& SceneTextures,
-		const FViewInfo& View,
-		const FDistanceFieldCulledObjectBufferParameters& CulledObjectBufferParameters,
-		FRDGBufferRef ObjectTilesIndirectArguments,
-		const FTileIntersectionParameters& TileIntersectionParameters,
-		const FDistanceFieldAOParameters& Parameters,
-		FRDGTextureRef DistanceFieldNormal,
-		FRDGTextureRef& OutDynamicBentNormalAO);
-
 	void RenderMegaLights(
 		FRDGBuilder& GraphBuilder,
 		const FSceneTextures& SceneTextures,
@@ -935,12 +914,7 @@ private:
 	/** Renders indirect shadows from capsules modulated onto scene color. */
 	void RenderIndirectCapsuleShadows(FRDGBuilder& GraphBuilder, const FSceneTextures& SceneTextures) const;
 
-	/** Renders capsule shadows for movable skylights, using the cone of visibility (bent normal) from DFAO. */
-	void RenderCapsuleShadowsForMovableSkylight(
-		FRDGBuilder& GraphBuilder,
-		const FViewInfo& View,
-		TRDGUniformBufferRef<FSceneTextureUniformParameters> SceneTexturesUniformBuffer,
-		FRDGTextureRef& BentNormalOutput) const;
+	
 
 	void RenderDeferredShadowProjections(
 		FRDGBuilder& GraphBuilder,
diff --git a/Engine/Source/Runtime/Renderer/Private/DistanceFieldAmbientOcclusion.cpp b/Engine/Source/Runtime/Renderer/Private/DistanceFieldAmbientOcclusion.cpp
--- a/Engine/Source/Runtime/Renderer/Private/DistanceFieldAmbientOcclusion.cpp
+++ b/Engine/Source/Runtime/Renderer/Private/DistanceFieldAmbientOcclusion.cpp
@@ -31,6 +31,14 @@ FAutoConsoleVariableRef CVarDistanceFieldAO(
 	ECVF_Scalability | ECVF_RenderThreadSafe
 	);
 
+int32 GMobileDistanceFieldAO = 0;
+FAutoConsoleVariableRef CVarMobileDistanceFieldAO(
+	TEXT("r.Mobile.DistanceFieldAO"),
+	GMobileDistanceFieldAO,
+	TEXT("Whether the distance field AO feature is allowed for Mobile Platforms, which is used to implement shadows of Movable sky lights from static meshes."),
+	ECVF_ReadOnly | ECVF_RenderThreadSafe
+);
+
 int32 GDistanceFieldAOMultiView = 1;
 FAutoConsoleVariableRef CVarDistanceFieldAOMultiView(
 	TEXT("r.DistanceFieldAO.MultiView"),
@@ -326,7 +334,7 @@ FAutoConsoleCommandWithWorld ListMeshDistanceFieldsMemoryConsoleCommand(
 
 bool ShouldCompileDFNormalShaders(EShaderPlatform ShaderPlatform)
 {
-	return ShouldCompileDistanceFieldShaders(ShaderPlatform) && !IsMobilePlatform(ShaderPlatform);
+	return ShouldCompileDistanceFieldShaders(ShaderPlatform);
 }
 
 class FComputeDistanceFieldNormalPS : public FGlobalShader
@@ -337,7 +345,7 @@ public:
 
 	BEGIN_SHADER_PARAMETER_STRUCT(FParameters, )
 		SHADER_PARAMETER_STRUCT_REF(FViewUniformShaderParameters, View)
-		SHADER_PARAMETER_RDG_UNIFORM_BUFFER(FSceneTextureUniformParameters, SceneTextures)
+		SHADER_PARAMETER_STRUCT_INCLUDE(FSceneTextureShaderParameters, SceneTextures)
 		SHADER_PARAMETER_RDG_UNIFORM_BUFFER(FSubstrateGlobalUniformParameters, Substrate)
 		SHADER_PARAMETER_STRUCT_INCLUDE(FAOParameters, AOParameters)
 		RENDER_TARGET_BINDING_SLOTS()
@@ -366,7 +374,7 @@ public:
 
 	BEGIN_SHADER_PARAMETER_STRUCT(FParameters, )
 		SHADER_PARAMETER_STRUCT_REF(FViewUniformShaderParameters, View)
-		SHADER_PARAMETER_RDG_UNIFORM_BUFFER(FSceneTextureUniformParameters, SceneTextures)
+		SHADER_PARAMETER_STRUCT_INCLUDE(FSceneTextureShaderParameters, SceneTextures)
 		SHADER_PARAMETER_RDG_UNIFORM_BUFFER(FSubstrateGlobalUniformParameters, Substrate)
 		SHADER_PARAMETER_STRUCT_INCLUDE(FAOParameters, AOParameters)
 		SHADER_PARAMETER_RDG_TEXTURE_UAV(RWTexture2D<float4>, RWDistanceFieldNormal)
@@ -390,7 +398,7 @@ IMPLEMENT_GLOBAL_SHADER(FComputeDistanceFieldNormalCS, "/Engine/Private/Distance
 void ComputeDistanceFieldNormal(
 	FRDGBuilder& GraphBuilder,
 	const FViewInfo& View,
-	TRDGUniformBufferRef<FSceneTextureUniformParameters> SceneTexturesUniformBuffer,
+	const FSceneTextures& SceneTextures,
 	FRDGTextureRef DistanceFieldNormal,
 	const FDistanceFieldAOParameters& Parameters)
 {
@@ -403,7 +411,7 @@ void ComputeDistanceFieldNormal(
 
 		auto* PassParameters = GraphBuilder.AllocParameters<FComputeDistanceFieldNormalCS::FParameters>();
 		PassParameters->View = View.ViewUniformBuffer;
-		PassParameters->SceneTextures = SceneTexturesUniformBuffer;
+		PassParameters->SceneTextures = SceneTextures.GetSceneTextureShaderParameters(View.FeatureLevel);;
 		PassParameters->Substrate = Substrate::BindSubstrateGlobalUniformParameters(View);
 		PassParameters->AOParameters = DistanceField::SetupAOShaderParameters(Parameters);
 		PassParameters->RWDistanceFieldNormal = GraphBuilder.CreateUAV(DistanceFieldNormal);
@@ -416,7 +424,7 @@ void ComputeDistanceFieldNormal(
 	{
 		auto* PassParameters = GraphBuilder.AllocParameters<FComputeDistanceFieldNormalPS::FParameters>();
 		PassParameters->View = View.ViewUniformBuffer;
-		PassParameters->SceneTextures = SceneTexturesUniformBuffer;
+		PassParameters->SceneTextures = SceneTextures.GetSceneTextureShaderParameters(View.FeatureLevel);;
 		PassParameters->Substrate = Substrate::BindSubstrateGlobalUniformParameters(View);
 		PassParameters->AOParameters = DistanceField::SetupAOShaderParameters(Parameters);
 		PassParameters->RenderTargets[0] = FRenderTargetBinding(DistanceFieldNormal, ERenderTargetLoadAction::EClear);
@@ -692,7 +700,6 @@ bool ShouldRenderDeferredDynamicSkyLight(const FScene* Scene, const FSceneViewFa
 		&& !Scene->SkyLight->bWantsStaticShadowing
 		&& !Scene->SkyLight->bHasStaticLighting
 		&& ViewFamily.EngineShowFlags.SkyLighting
-		&& Scene->GetFeatureLevel() >= ERHIFeatureLevel::SM5
 		&& !IsForwardShadingEnabled(Scene->GetShaderPlatform())
 		&& !ViewFamily.EngineShowFlags.VisualizeLightCulling
 		&& !ShouldRenderRayTracingSkyLight(Scene->SkyLight, Scene->GetShaderPlatform()); // Disable diffuse sky contribution if evaluated by RT Sky;
@@ -875,13 +882,14 @@ void FDeferredShadingSceneRenderer::RenderDFAOAsIndirectShadowing(
 	}
 }
 
-void FDeferredShadingSceneRenderer::RenderDistanceFieldLighting(
+void FSceneRenderer::RenderDistanceFieldLighting(
 	FRDGBuilder& GraphBuilder,
 	const FSceneTextures& SceneTextures,
 	const FDistanceFieldAOParameters& Parameters,
 	TArray<FRDGTextureRef>& OutDynamicBentNormalAOTextures,
 	bool bModulateToSceneColor,
-	bool bVisualizeAmbientOcclusion)
+	bool bVisualizeAmbientOcclusion,
+	bool bModulateToScreenSpaceAO)
 {
 	check(ShouldRenderDistanceFieldLighting(Scene->DistanceFieldSceneData, Views));
 	check(!Scene->DistanceFieldSceneData.HasPendingOperations());
@@ -920,7 +928,7 @@ void FDeferredShadingSceneRenderer::RenderDistanceFieldLighting(
 			DistanceFieldNormal = GraphBuilder.CreateTexture(Desc, TEXT("DistanceFieldNormal"));
 		}
 
-		ComputeDistanceFieldNormal(GraphBuilder, View, SceneTextures.UniformBuffer, DistanceFieldNormal, Parameters);
+		ComputeDistanceFieldNormal(GraphBuilder, View, SceneTextures, DistanceFieldNormal, Parameters);
 
 		const FIntPoint TileListGroupSize = GetTileListGroupSizeForView(View);
 		const int32 MaxSceneObjects = FMath::DivideAndRoundUp(Scene->DistanceFieldSceneData.NumObjectsInBuffer, 256) * 256;
@@ -944,7 +952,7 @@ void FDeferredShadingSceneRenderer::RenderDistanceFieldLighting(
 
 			CullObjectsToView(GraphBuilder, *Scene, View, Parameters, CulledObjectBufferParameters);
 
-			BuildTileObjectLists(GraphBuilder, *Scene, View, SceneTextures.UniformBuffer, ObjectIndirectArguments, CulledObjectBufferParameters, TileIntersectionParameters, DistanceFieldNormal, Parameters);
+			BuildTileObjectLists(GraphBuilder, *Scene, View, SceneTextures, ObjectIndirectArguments, CulledObjectBufferParameters, TileIntersectionParameters, DistanceFieldNormal, Parameters);
 		}
 
 		// Render to a per-view BentNormal first because it also needs to be stored by the view history.
@@ -962,12 +970,12 @@ void FDeferredShadingSceneRenderer::RenderDistanceFieldLighting(
 			DistanceFieldNormal,
 			PerViewBentNormal);
 
-		RenderCapsuleShadowsForMovableSkylight(GraphBuilder, View, SceneTextures.UniformBuffer, PerViewBentNormal);
+		RenderCapsuleShadowsForMovableSkylight(GraphBuilder, View, PerViewBentNormal);
 
 		// Upsample to full resolution, write to output in case of debug AO visualization or scene color modulation (standard upsampling is done later together with sky lighting and reflection environment)
-		if (bModulateToSceneColor || bVisualizeAmbientOcclusion)
+		if (bModulateToSceneColor || bModulateToScreenSpaceAO || bVisualizeAmbientOcclusion)
 		{
-			UpsampleBentNormalAO(GraphBuilder, View, SceneTextures.UniformBuffer, SceneTextures.Color.Target, PerViewBentNormal, bModulateToSceneColor && !bVisualizeAmbientOcclusion);
+			UpsampleBentNormalAO(GraphBuilder, View, SceneTextures, bModulateToScreenSpaceAO ? SceneTextures.ScreenSpaceAO: SceneTextures.Color.Target, PerViewBentNormal, bModulateToSceneColor && !bVisualizeAmbientOcclusion, bModulateToScreenSpaceAO && !bVisualizeAmbientOcclusion);
 		}
 
 		OutDynamicBentNormalAOTextures.Add(PerViewBentNormal);
diff --git a/Engine/Source/Runtime/Renderer/Private/DistanceFieldAmbientOcclusion.h b/Engine/Source/Runtime/Renderer/Private/DistanceFieldAmbientOcclusion.h
--- a/Engine/Source/Runtime/Renderer/Private/DistanceFieldAmbientOcclusion.h
+++ b/Engine/Source/Runtime/Renderer/Private/DistanceFieldAmbientOcclusion.h
@@ -17,6 +17,7 @@ class FScene;
 class FSceneTextureUniformParameters;
 class FSceneViewFamily;
 class FViewInfo;
+struct FSceneTextures;
 struct FShaderCompilerEnvironment;
 using FRDGBufferRef = FRDGBuffer*;
 using FRDGTextureRef = FRDGTexture*;
@@ -110,7 +111,7 @@ extern void CullObjectsToView(FRDGBuilder& GraphBuilder, const FScene& Scene, co
 extern void BuildTileObjectLists(FRDGBuilder& GraphBuilder,
 	const FScene& Scene,
 	const FViewInfo& View,
-	TRDGUniformBufferRef<FSceneTextureUniformParameters> SceneTexturesUniformBuffer,
+	const FSceneTextures& SceneTextures,
 	FRDGBufferRef ObjectIndirectArguments,
 	const FDistanceFieldCulledObjectBufferParameters& CulledObjectBufferParameters,
 	FTileIntersectionParameters TileIntersectionParameters,
diff --git a/Engine/Source/Runtime/Renderer/Private/DistanceFieldLightingPost.cpp b/Engine/Source/Runtime/Renderer/Private/DistanceFieldLightingPost.cpp
--- a/Engine/Source/Runtime/Renderer/Private/DistanceFieldLightingPost.cpp
+++ b/Engine/Source/Runtime/Renderer/Private/DistanceFieldLightingPost.cpp
@@ -71,7 +71,7 @@ bool UseAOHistoryStabilityPass()
 
 bool ShouldCompileDFLightingPostShaders(EShaderPlatform ShaderPlatform)
 {
-	return ShouldCompileDistanceFieldShaders(ShaderPlatform) && !IsMobilePlatform(ShaderPlatform);
+	return ShouldCompileDistanceFieldShaders(ShaderPlatform);
 }
 
 BEGIN_SHADER_PARAMETER_STRUCT(FGeometryAwareUpsampleParameters, )
@@ -127,7 +127,7 @@ public:
 
 	BEGIN_SHADER_PARAMETER_STRUCT(FParameters, )
 		SHADER_PARAMETER_STRUCT_REF(FViewUniformShaderParameters, View)
-		SHADER_PARAMETER_RDG_UNIFORM_BUFFER(FSceneTextureUniformParameters, SceneTextures)
+		SHADER_PARAMETER_STRUCT_INCLUDE(FSceneTextureShaderParameters, SceneTextures)
 		SHADER_PARAMETER_STRUCT_INCLUDE(FAOParameters, AOParameters)
 		SHADER_PARAMETER_STRUCT_INCLUDE(FGeometryAwareUpsampleParameters, GeometryAwareUpsampleParameters)
 		SHADER_PARAMETER_RDG_TEXTURE(Texture2D, BentNormalHistoryTexture)
@@ -163,7 +163,7 @@ public:
 
 	BEGIN_SHADER_PARAMETER_STRUCT(FParameters, )
 		SHADER_PARAMETER_STRUCT_REF(FViewUniformShaderParameters, View)
-		SHADER_PARAMETER_RDG_UNIFORM_BUFFER(FSceneTextureUniformParameters, SceneTextures)
+		SHADER_PARAMETER_STRUCT_INCLUDE(FSceneTextureShaderParameters, SceneTextures)
 		SHADER_PARAMETER_RDG_TEXTURE(Texture2D, BentNormalAOTexture)
 		SHADER_PARAMETER_SAMPLER(SamplerState, BentNormalAOSampler)
 		SHADER_PARAMETER(FVector2f, BentNormalAOTexelSize)
@@ -230,7 +230,7 @@ void GeometryAwareUpsample(FRDGBuilder& GraphBuilder, const FViewInfo& View, FRD
 	auto* PassParameters = GraphBuilder.AllocParameters<FGeometryAwareUpsamplePS::FParameters>();
 	PassParameters->View = View.ViewUniformBuffer;
 	PassParameters->AOParameters = DistanceField::SetupAOShaderParameters(Parameters);
-	PassParameters->GeometryAwareUpsampleParameters = SetupGeometryAwareUpsampleParameters(View, DistanceFieldNormal, BentNormalInterpolation);
+	PassParameters->GeometryAwareUpsampleParameters = SetupGeometryAwareUpsampleParameters(View, DistanceFieldNormal, BentNormalInterpolation ? BentNormalInterpolation : GSystemTextures.GetWhiteDummy(GraphBuilder));
 	PassParameters->RenderTargets[0] = FRenderTargetBinding(DistanceFieldAOBentNormal, ERenderTargetLoadAction::ELoad);
 
 	auto VertexShader = View.ShaderMap->GetShader<FPostProcessVS>();
@@ -309,7 +309,7 @@ void UpdateHistory(
 	FRDGBuilder& GraphBuilder,
 	const FViewInfo& View, 
 	const TCHAR* BentNormalHistoryRTName,
-	TRDGUniformBufferRef<FSceneTextureUniformParameters> SceneTexturesUniformBuffer,
+	const FSceneTextures& SceneTextures,
 	FRDGTextureRef VelocityTexture,
 	FRDGTextureRef DistanceFieldNormal,
 	FRDGTextureRef BentNormalInterpolation,
@@ -346,7 +346,7 @@ void UpdateHistory(
 
 				auto* PassParameters = GraphBuilder.AllocParameters<FUpdateHistoryDepthRejectionPS::FParameters>();
 				PassParameters->View = View.ViewUniformBuffer;
-				PassParameters->SceneTextures = SceneTexturesUniformBuffer;
+				PassParameters->SceneTextures = SceneTextures.GetSceneTextureShaderParameters(View.FeatureLevel);;
 				PassParameters->AOParameters = DistanceField::SetupAOShaderParameters(Parameters);
 				PassParameters->GeometryAwareUpsampleParameters = SetupGeometryAwareUpsampleParameters(View, DistanceFieldNormal, BentNormalInterpolation);
 				PassParameters->BentNormalHistoryTexture = BentNormalHistoryTexture;
@@ -447,7 +447,7 @@ void UpdateHistory(
 
 				auto* PassParameters = GraphBuilder.AllocParameters<FFilterHistoryPS::FParameters>();
 				PassParameters->View = View.ViewUniformBuffer;
-				PassParameters->SceneTextures = SceneTexturesUniformBuffer;
+				PassParameters->SceneTextures = SceneTextures.GetSceneTextureShaderParameters(View.FeatureLevel);;
 				PassParameters->BentNormalAOTexture = NewBentNormalHistory;
 				PassParameters->BentNormalAOSampler = TStaticSamplerState<SF_Bilinear, AM_Clamp, AM_Clamp, AM_Clamp>::GetRHI();
 				PassParameters->BentNormalAOTexelSize = FVector2f(1.0f / AOBufferSize.X, 1.0f / AOBufferSize.Y);
@@ -536,7 +536,7 @@ public:
 
 	BEGIN_SHADER_PARAMETER_STRUCT(FParameters, )
 		SHADER_PARAMETER_STRUCT_REF(FViewUniformShaderParameters, View)
-		SHADER_PARAMETER_RDG_UNIFORM_BUFFER(FSceneTextureUniformParameters, SceneTextures)
+		SHADER_PARAMETER_STRUCT_INCLUDE(FSceneTextureShaderParameters, SceneTextures)
 		SHADER_PARAMETER_STRUCT_INCLUDE(FDFAOUpsampleParameters, DFAOUpsampleParameters)
 		SHADER_PARAMETER(float, MinIndirectDiffuseOcclusion)
 		RDG_TEXTURE_ACCESS(DistanceFieldAOBentNormal, ERHIAccess::SRVGraphics)
@@ -557,10 +557,11 @@ IMPLEMENT_GLOBAL_SHADER(FDistanceFieldAOUpsamplePS, "/Engine/Private/DistanceFie
 void UpsampleBentNormalAO(
 	FRDGBuilder& GraphBuilder,
 	const FViewInfo& View,
-	TRDGUniformBufferRef<FSceneTextureUniformParameters> SceneTexturesUniformBuffer,
+	const FSceneTextures& SceneTextures,
 	FRDGTextureRef SceneColorTexture,
 	FRDGTextureRef DistanceFieldAOBentNormal,
-	bool bModulateSceneColor)
+	bool bModulateSceneColor,
+	bool bModulateToScreenSpaceAO)
 {
 	FScene* Scene = (FScene*)View.Family->Scene;
 
@@ -568,7 +569,7 @@ void UpsampleBentNormalAO(
 
 	auto* PassParameters = GraphBuilder.AllocParameters<FDistanceFieldAOUpsamplePS::FParameters>();
 	PassParameters->View = View.ViewUniformBuffer;
-	PassParameters->SceneTextures = SceneTexturesUniformBuffer;
+	PassParameters->SceneTextures = SceneTextures.GetSceneTextureShaderParameters(View.FeatureLevel);;
 	PassParameters->DFAOUpsampleParameters = DistanceField::SetupAOUpsampleParameters(View, DistanceFieldAOBentNormal ? DistanceFieldAOBentNormal : GSystemTextures.GetWhiteDummy(GraphBuilder));
 	PassParameters->MinIndirectDiffuseOcclusion = Scene->SkyLight ? Scene->SkyLight->MinOcclusion : 0;
 	PassParameters->RenderTargets[0] = FRenderTargetBinding(SceneColorTexture, ERenderTargetLoadAction::ELoad);
@@ -577,14 +578,14 @@ void UpsampleBentNormalAO(
 	TShaderMapRef<FPostProcessVS> VertexShader(View.ShaderMap);
 
 	FDistanceFieldAOUpsamplePS::FPermutationDomain PermutationVector;
-	PermutationVector.Set<FDistanceFieldAOUpsamplePS::FModulateToSceneColorDim>(bModulateSceneColor);
+	PermutationVector.Set<FDistanceFieldAOUpsamplePS::FModulateToSceneColorDim>(bModulateSceneColor || bModulateToScreenSpaceAO);
 	TShaderMapRef<FDistanceFieldAOUpsamplePS> PixelShader(View.ShaderMap, PermutationVector);
 
 	GraphBuilder.AddPass(
 		RDG_EVENT_NAME("UpsampleAO"),
 		PassParameters,
 		ERDGPassFlags::Raster,
-		[VertexShader, PixelShader, PassParameters, &View, DistanceFieldAOBentNormal, bModulateSceneColor](FRDGAsyncTask, FRHICommandList& RHICmdList)
+		[VertexShader, PixelShader, PassParameters, &View, DistanceFieldAOBentNormal, bModulateSceneColor, bModulateToScreenSpaceAO](FRDGAsyncTask, FRHICommandList& RHICmdList)
 		{
 			RHICmdList.SetViewport(View.ViewRect.Min.X, View.ViewRect.Min.Y, 0.0f, View.ViewRect.Max.X, View.ViewRect.Max.Y, 1.0f);
 
@@ -594,7 +595,7 @@ void UpsampleBentNormalAO(
 			GraphicsPSOInit.DepthStencilState = TStaticDepthStencilState<false, CF_Always>::GetRHI();
 			GraphicsPSOInit.PrimitiveType = PT_TriangleList;
 
-			if (bModulateSceneColor)
+			if (bModulateSceneColor || bModulateToScreenSpaceAO)
 			{
 				GraphicsPSOInit.BlendState = TStaticBlendState<CW_RGB, BO_Add, BF_DestColor, BF_Zero, BO_Add, BF_Zero, BF_One>::GetRHI();
 			}
@@ -609,14 +610,15 @@ void UpsampleBentNormalAO(
 			SetGraphicsPipelineState(RHICmdList, GraphicsPSOInit, 0);
 			SetShaderParameters(RHICmdList, PixelShader, PixelShader.GetPixelShader(), *PassParameters);
 
+			const int32 DownsampleFactor = bModulateToScreenSpaceAO ? 1.f : GAODownsampleFactor;
 			DrawRectangle(
 				RHICmdList,
 				0, 0,
 				View.ViewRect.Width(), View.ViewRect.Height(),
-				View.ViewRect.Min.X / GAODownsampleFactor, View.ViewRect.Min.Y / GAODownsampleFactor,
-				View.ViewRect.Width() / GAODownsampleFactor, View.ViewRect.Height() / GAODownsampleFactor,
+				View.ViewRect.Min.X / DownsampleFactor, View.ViewRect.Min.Y / DownsampleFactor,
+				View.ViewRect.Width() / DownsampleFactor, View.ViewRect.Height() / DownsampleFactor,
 				FIntPoint(View.ViewRect.Width(), View.ViewRect.Height()),
 				GetBufferSizeForAO(View),
 				VertexShader);
 		});
-}
+}
\ No newline at end of file
diff --git a/Engine/Source/Runtime/Renderer/Private/DistanceFieldLightingPost.h b/Engine/Source/Runtime/Renderer/Private/DistanceFieldLightingPost.h
--- a/Engine/Source/Runtime/Renderer/Private/DistanceFieldLightingPost.h
+++ b/Engine/Source/Runtime/Renderer/Private/DistanceFieldLightingPost.h
@@ -17,7 +17,7 @@ extern void UpdateHistory(
 	FRDGBuilder& GraphBuilder,
 	const FViewInfo& View,
 	const TCHAR* BentNormalHistoryRTName,
-	TRDGUniformBufferRef<FSceneTextureUniformParameters> SceneTexturesUniformBuffer,
+	const FSceneTextures& SceneTextures,
 	FRDGTextureRef VelocityTexture,
 	FRDGTextureRef BentNormalInterpolation,
 	FRDGTextureRef DistanceFieldNormal,
@@ -31,7 +31,15 @@ extern void UpdateHistory(
 extern void UpsampleBentNormalAO(
 	FRDGBuilder& GraphBuilder,
 	const FViewInfo& View,
-	TRDGUniformBufferRef<FSceneTextureUniformParameters> SceneTexturesUniformBuffer,
+	const FSceneTextures& SceneTextures,
 	FRDGTextureRef SceneColorTexture,
 	FRDGTextureRef DistanceFieldAOBentNormal,
-	bool bModulateSceneColor);
+	bool bModulateSceneColor,
+	bool bModulateToScreenSpaceAO);
+
+extern void MobileAggregateDFAOToSSAO(
+	FRDGBuilder& GraphBuilder,
+	const FViewInfo& View,
+	const FSceneTextures& SceneTextures,
+	FRDGTextureRef SceneColorTexture,
+	FRDGTextureRef DistanceFieldAOBentNormal);
diff --git a/Engine/Source/Runtime/Renderer/Private/DistanceFieldObjectCulling.cpp b/Engine/Source/Runtime/Renderer/Private/DistanceFieldObjectCulling.cpp
--- a/Engine/Source/Runtime/Renderer/Private/DistanceFieldObjectCulling.cpp
+++ b/Engine/Source/Runtime/Renderer/Private/DistanceFieldObjectCulling.cpp
@@ -117,7 +117,7 @@ public:
 		SHADER_PARAMETER_STRUCT_REF(FViewUniformShaderParameters, View)
 		SHADER_PARAMETER_RDG_BUFFER_UAV(RWStructuredBuffer<FVector4f>, RWTileConeAxisAndCos)
 		SHADER_PARAMETER_RDG_BUFFER_UAV(RWStructuredBuffer<FVector4f>, RWTileConeDepthRanges)
-		SHADER_PARAMETER_RDG_UNIFORM_BUFFER(FSceneTextureUniformParameters, SceneTextures)
+		SHADER_PARAMETER_STRUCT_INCLUDE(FSceneTextureShaderParameters, SceneTextures)
 		SHADER_PARAMETER_RDG_TEXTURE(Texture2D, DistanceFieldNormalTexture)
 		SHADER_PARAMETER_SAMPLER(SamplerState, DistanceFieldNormalSampler)
 		SHADER_PARAMETER_STRUCT_INCLUDE(FAOParameters, AOParameters)
@@ -206,7 +206,7 @@ IMPLEMENT_GLOBAL_SHADER(FObjectCullPS, "/Engine/Private/DistanceFieldObjectCulli
 BEGIN_SHADER_PARAMETER_STRUCT(FObjectCullParameters, )
 	SHADER_PARAMETER_STRUCT_INCLUDE(FObjectCullVS::FParameters, VS)
 	SHADER_PARAMETER_STRUCT_INCLUDE(FObjectCullPS::FParameters, PS)
-	SHADER_PARAMETER_RDG_UNIFORM_BUFFER(FSceneTextureUniformParameters, SceneTextures)
+	SHADER_PARAMETER_STRUCT_INCLUDE(FSceneTextureShaderParameters, SceneTextures)
 	RDG_BUFFER_ACCESS(ObjectIndirectArguments, ERHIAccess::IndirectArgs)
 	RENDER_TARGET_BINDING_SLOTS()
 END_SHADER_PARAMETER_STRUCT()
@@ -226,7 +226,7 @@ public:
 		SHADER_PARAMETER_STRUCT_INCLUDE(FTileIntersectionParameters, TileIntersectionParameters)
 		SHADER_PARAMETER_STRUCT_INCLUDE(FDistanceFieldCulledObjectBufferParameters, DistanceFieldCulledObjectBuffers)
 		SHADER_PARAMETER_STRUCT_INCLUDE(FDistanceFieldAtlasParameters, DistanceFieldAtlas)
-		SHADER_PARAMETER_RDG_UNIFORM_BUFFER(FSceneTextureUniformParameters, SceneTextures)
+		SHADER_PARAMETER_STRUCT_INCLUDE(FSceneTextureShaderParameters, SceneTextures)
 	END_SHADER_PARAMETER_STRUCT()
 
 	static bool ShouldCompilePermutation(const FGlobalShaderPermutationParameters& Parameters)
@@ -255,7 +255,7 @@ void ScatterTilesToObjects(
 	FRDGBufferRef ObjectIndirectArguments,
 	const FDistanceFieldCulledObjectBufferParameters& CulledObjectBufferParameters,
 	const FTileIntersectionParameters& TileIntersectionParameters,
-	TRDGUniformBufferRef<FSceneTextureUniformParameters> SceneTexturesUniformBuffer)
+	const FSceneTextures& SceneTextures)
 {
 	FDistanceFieldObjectBufferParameters DistanceFieldObjectBuffers = DistanceField::SetupObjectBufferParameters(GraphBuilder, DistanceFieldSceneData);
 
@@ -288,7 +288,7 @@ void ScatterTilesToObjects(
 	PassParameters->PS.AOParameters = DistanceField::SetupAOShaderParameters(Parameters);
 	PassParameters->PS.NumGroups = FVector2f(TileListGroupSize.X, TileListGroupSize.Y);
 
-	PassParameters->SceneTextures = SceneTexturesUniformBuffer;
+	PassParameters->SceneTextures = SceneTextures.GetSceneTextureShaderParameters(View.FeatureLevel);
 	PassParameters->ObjectIndirectArguments = ObjectIndirectArguments;
 
 	if (GRHIRequiresRenderTargetForPixelShaderUAVs)
@@ -347,7 +347,7 @@ void BuildTileObjectLists(
 	FRDGBuilder& GraphBuilder,
 	const FScene& Scene,
 	const FViewInfo& View,
-	TRDGUniformBufferRef<FSceneTextureUniformParameters> SceneTexturesUniformBuffer,
+	const FSceneTextures& SceneTextures,
 	FRDGBufferRef ObjectIndirectArguments,
 	const FDistanceFieldCulledObjectBufferParameters& CulledObjectBufferParameters,
 	FTileIntersectionParameters TileIntersectionParameters,
@@ -366,7 +366,7 @@ void BuildTileObjectLists(
 		PassParameters->View = View.ViewUniformBuffer;
 		PassParameters->RWTileConeAxisAndCos = TileIntersectionParameters.RWTileConeAxisAndCos;
 		PassParameters->RWTileConeDepthRanges = TileIntersectionParameters.RWTileConeDepthRanges;
-		PassParameters->SceneTextures = SceneTexturesUniformBuffer;
+		PassParameters->SceneTextures = SceneTextures.GetSceneTextureShaderParameters(View.FeatureLevel);
 		PassParameters->DistanceFieldNormalTexture = DistanceFieldNormal;
 		PassParameters->DistanceFieldNormalSampler = TStaticSamplerState<SF_Point, AM_Clamp, AM_Clamp, AM_Clamp>::GetRHI();
 		PassParameters->AOParameters = DistanceField::SetupAOShaderParameters(Parameters);
@@ -384,7 +384,7 @@ void BuildTileObjectLists(
 	AddClearUAVPass(GraphBuilder, TileIntersectionParameters.RWObjectTilesIndirectArguments, 0);
 
 	// Rasterize object bounding shapes and intersect with screen tiles to compute how many tiles intersect each object
-	ScatterTilesToObjects(GraphBuilder, true, View, Scene.DistanceFieldSceneData, TileListGroupSize, Parameters, ObjectIndirectArguments, CulledObjectBufferParameters, TileIntersectionParameters, SceneTexturesUniformBuffer);
+	ScatterTilesToObjects(GraphBuilder, true, View, Scene.DistanceFieldSceneData, TileListGroupSize, Parameters, ObjectIndirectArguments, CulledObjectBufferParameters, TileIntersectionParameters, SceneTextures);
 
 	{
 		auto* PassParameters = GraphBuilder.AllocParameters<FComputeCulledTilesStartOffsetCS::FParameters>();
@@ -392,7 +392,7 @@ void BuildTileObjectLists(
 		PassParameters->TileIntersectionParameters = TileIntersectionParameters;
 		PassParameters->DistanceFieldCulledObjectBuffers = CulledObjectBufferParameters;
 		PassParameters->DistanceFieldAtlas = DistanceField::SetupAtlasParameters(GraphBuilder, Scene.DistanceFieldSceneData);
-		PassParameters->SceneTextures = SceneTexturesUniformBuffer;
+		PassParameters->SceneTextures = SceneTextures.GetSceneTextureShaderParameters(View.FeatureLevel);
 
 		auto ComputeShader = View.ShaderMap->GetShader<FComputeCulledTilesStartOffsetCS>();
 		const FIntVector GroupCount = FComputeShaderUtils::GetGroupCountWrapped(Scene.DistanceFieldSceneData.NumObjectsInBuffer, ComputeStartOffsetGroupSize);
@@ -409,5 +409,5 @@ void BuildTileObjectLists(
 	AddClearUAVPass(GraphBuilder, TileIntersectionParameters.RWNumCulledTilesArray, 0);
 
 	// Rasterize object bounding shapes and intersect with screen tiles, and write out intersecting tile indices for the cone tracing pass
-	ScatterTilesToObjects(GraphBuilder, false, View, Scene.DistanceFieldSceneData, TileListGroupSize, Parameters, ObjectIndirectArguments, CulledObjectBufferParameters, TileIntersectionParameters, SceneTexturesUniformBuffer);
+	ScatterTilesToObjects(GraphBuilder, false, View, Scene.DistanceFieldSceneData, TileListGroupSize, Parameters, ObjectIndirectArguments, CulledObjectBufferParameters, TileIntersectionParameters, SceneTextures);
 }
diff --git a/Engine/Source/Runtime/Renderer/Private/DistanceFieldScreenGridLighting.cpp b/Engine/Source/Runtime/Renderer/Private/DistanceFieldScreenGridLighting.cpp
--- a/Engine/Source/Runtime/Renderer/Private/DistanceFieldScreenGridLighting.cpp
+++ b/Engine/Source/Runtime/Renderer/Private/DistanceFieldScreenGridLighting.cpp
@@ -70,7 +70,7 @@ extern bool DistanceFieldAOUseHistory(const FViewInfo& View);
 
 bool ShouldCompileDFScreenGridLightingShaders(EShaderPlatform ShaderPlatform)
 {
-	return ShouldCompileDistanceFieldShaders(ShaderPlatform) && !IsMobilePlatform(ShaderPlatform);
+	return ShouldCompileDistanceFieldShaders(ShaderPlatform);
 }
 
 FVector2f GetJitterOffset(const FViewInfo& View)
@@ -152,7 +152,7 @@ public:
 		SHADER_PARAMETER_STRUCT_INCLUDE(FScreenGridParameters, ScreenGridParameters)
 		SHADER_PARAMETER_STRUCT_INCLUDE(FGlobalDistanceFieldParameters2, GlobalDistanceFieldParameters)
 		SHADER_PARAMETER_STRUCT_INCLUDE(FAOSampleParameters, AOSampleParameters)
-		SHADER_PARAMETER_RDG_UNIFORM_BUFFER(FSceneTextureUniformParameters, SceneTextures)
+		SHADER_PARAMETER_STRUCT_INCLUDE(FSceneTextureShaderParameters, SceneTextures)
 		SHADER_PARAMETER(float, TanConeHalfAngle)
 		RDG_BUFFER_ACCESS(ObjectTilesIndirectArguments, ERHIAccess::IndirectArgs)
 	END_SHADER_PARAMETER_STRUCT()
@@ -196,7 +196,7 @@ public:
 		SHADER_PARAMETER_STRUCT_INCLUDE(FGlobalDistanceFieldParameters2, GlobalDistanceFieldParameters)
 		SHADER_PARAMETER_STRUCT_INCLUDE(FAOSampleParameters, AOSampleParameters)
 		SHADER_PARAMETER_RDG_BUFFER_SRV(Buffer<FVector4f>, TileConeDepthRanges)
-		SHADER_PARAMETER_RDG_UNIFORM_BUFFER(FSceneTextureUniformParameters, SceneTextures)
+		SHADER_PARAMETER_STRUCT_INCLUDE(FSceneTextureShaderParameters, SceneTextures)
 		SHADER_PARAMETER(FIntPoint, TileListGroupSize)
 		SHADER_PARAMETER(float, TanConeHalfAngle)
 	END_SHADER_PARAMETER_STRUCT()
@@ -253,6 +253,15 @@ public:
 		OutEnvironment.SetDefine(TEXT("COMBINE_CONES_SIZEX"), GCombineConesSizeX);
 		OutEnvironment.SetDefine(TEXT("TRACE_DOWNSAMPLE_FACTOR"), GConeTraceDownsampleFactor);
 
+		if (IsMobilePlatform(Parameters.Platform) && !IsMobileDeferredShadingEnabled(Parameters.Platform))
+		{
+			OutEnvironment.SetDefine(TEXT("ROTATE_WITH_TANGENT"), 0);
+		}
+		else
+		{
+			OutEnvironment.SetDefine(TEXT("ROTATE_WITH_TANGENT"), 1);
+		}
+
 		// To reduce shader compile time of compute shaders with shared memory, doesn't have an impact on generated code with current compiler (June 2010 DX SDK)
 		OutEnvironment.CompilerFlags.Add(CFLAG_StandardOptimization);
 	}
@@ -264,7 +273,7 @@ void PostProcessBentNormalAOScreenGrid(
 	FRDGBuilder& GraphBuilder, 
 	const FDistanceFieldAOParameters& Parameters, 
 	const FViewInfo& View,
-	TRDGUniformBufferRef<FSceneTextureUniformParameters> SceneTexturesUniformBuffer,
+	const FSceneTextures& SceneTextures,
 	FRDGTextureRef VelocityTexture,
 	FRDGTextureRef BentNormalInterpolation,
 	FRDGTextureRef DistanceFieldNormal,
@@ -278,7 +287,7 @@ void PostProcessBentNormalAOScreenGrid(
 		GraphBuilder,
 		View,
 		TEXT("DistanceFieldAOHistory"),
-		SceneTexturesUniformBuffer,
+		SceneTextures,
 		VelocityTexture,
 		DistanceFieldNormal,
 		BentNormalInterpolation,
@@ -288,7 +297,7 @@ void PostProcessBentNormalAOScreenGrid(
 		Parameters);
 }
 
-void FDeferredShadingSceneRenderer::RenderDistanceFieldAOScreenGrid(
+void FSceneRenderer::RenderDistanceFieldAOScreenGrid(
 	FRDGBuilder& GraphBuilder,
 	const FSceneTextures& SceneTextures,
 	const FViewInfo& View,
@@ -338,7 +347,7 @@ void FDeferredShadingSceneRenderer::RenderDistanceFieldAOScreenGrid(
 		PassParameters->GlobalDistanceFieldParameters = SetupGlobalDistanceFieldParameters(View.GlobalDistanceFieldInfo.ParameterData);
 		PassParameters->AOSampleParameters = SetupAOSampleParameters(View.Family->FrameNumber);
 		PassParameters->TileConeDepthRanges = TileIntersectionParameters.TileConeDepthRanges;
-		PassParameters->SceneTextures = SceneTextures.UniformBuffer;
+		PassParameters->SceneTextures = SceneTextures.GetSceneTextureShaderParameters(View.FeatureLevel);;
 		PassParameters->TileListGroupSize = TileListGroupSize;
 		PassParameters->TanConeHalfAngle = FMath::Tan(GAOConeHalfAngle);
 
@@ -371,7 +380,7 @@ void FDeferredShadingSceneRenderer::RenderDistanceFieldAOScreenGrid(
 			PassParameters->GlobalDistanceFieldParameters = SetupGlobalDistanceFieldParameters(View.GlobalDistanceFieldInfo.ParameterData);
 		}
 		PassParameters->AOSampleParameters = SetupAOSampleParameters(View.Family->FrameNumber);
-		PassParameters->SceneTextures = SceneTextures.UniformBuffer;
+		PassParameters->SceneTextures = SceneTextures.GetSceneTextureShaderParameters(View.FeatureLevel);;
 		PassParameters->TanConeHalfAngle = FMath::Tan(GAOConeHalfAngle);
 		PassParameters->ObjectTilesIndirectArguments = ObjectTilesIndirectArguments;
 
@@ -418,7 +427,7 @@ void FDeferredShadingSceneRenderer::RenderDistanceFieldAOScreenGrid(
 		GraphBuilder,
 		Parameters,
 		View,
-		SceneTextures.UniformBuffer,
+		SceneTextures,
 		SceneTextures.Velocity,
 		DownsampledBentNormal,
 		DistanceFieldNormal,
diff --git a/Engine/Source/Runtime/Renderer/Private/MobileDeferredShadingPass.cpp b/Engine/Source/Runtime/Renderer/Private/MobileDeferredShadingPass.cpp
--- a/Engine/Source/Runtime/Renderer/Private/MobileDeferredShadingPass.cpp
+++ b/Engine/Source/Runtime/Renderer/Private/MobileDeferredShadingPass.cpp
@@ -14,6 +14,9 @@
 #include "LocalLightSceneProxy.h"
 #include "Materials/MaterialRenderProxy.h"
 #include "MobileSSR.h"
+#include "DistanceFieldLightingShared.h"
+#include "DistanceFieldAmbientOcclusion.h"
+#include "DistanceFieldLightingShared.h"
 
 DECLARE_GPU_STAT(DeferredShading);
 
@@ -66,7 +69,8 @@ class FMobileDirectionalLightFunctionPS : public FMaterialShader
 	class FEnableCSM					: SHADER_PERMUTATION_BOOL("ENABLE_MOBILE_CSM");
 	class FShadowQuality				: SHADER_PERMUTATION_RANGE_INT("MOBILE_SHADOW_QUALITY", 1, 3); // not using Quality=0
 	class FMobileSSRQuality 			: SHADER_PERMUTATION_ENUM_CLASS("MOBILE_SSR_QUALITY", EMobileSSRQuality);
-	
+	class FSkyShadowing					: SHADER_PERMUTATION_BOOL("APPLY_SKY_SHADOWING");
+
 	using FPermutationDomain = TShaderPermutationDomain<
 		FEnableShadingModelSupport,
 		FEnableClustredLights, 
@@ -75,7 +79,8 @@ class FMobileDirectionalLightFunctionPS : public FMaterialShader
 		FEnableSkyLight,
 		FEnableCSM, 
 		FShadowQuality,
-		FMobileSSRQuality>;
+		FMobileSSRQuality,
+		FSkyShadowing>;
 
 	BEGIN_SHADER_PARAMETER_STRUCT(FParameters, )
 		SHADER_PARAMETER_STRUCT_REF(FMobileDirectionalLightShaderParameters, MobileDirectionalLight)
@@ -86,6 +91,8 @@ class FMobileDirectionalLightFunctionPS : public FMaterialShader
 		SHADER_PARAMETER(FVector3f, CameraRelativeLightPosition)
 		SHADER_PARAMETER_TEXTURE(Texture2D, ScreenSpaceShadowMaskTexture)
 		SHADER_PARAMETER_SAMPLER(SamplerState, ScreenSpaceShadowMaskSampler)
+		SHADER_PARAMETER_STRUCT_INCLUDE(FDFAOUpsampleParameters, DFAOUpsampleParameters)
+		SHADER_PARAMETER_STRUCT_INCLUDE(FSkyDiffuseLightingParameters, SkyDiffuseLighting)
 	END_SHADER_PARAMETER_STRUCT()
 
 	static void ModifyCompilationEnvironment(const FMaterialShaderPermutationParameters& Parameters, FShaderCompilerEnvironment& OutEnvironment)
@@ -127,6 +134,11 @@ class FMobileDirectionalLightFunctionPS : public FMaterialShader
 			PermutationVector.Set<FMobileSSRQuality>(EMobileSSRQuality::Disabled);
 		}
 
+		if (!IsMobileDistanceFieldAOEnabled(Platform))
+		{
+			PermutationVector.Set<FSkyShadowing>(false);
+		}
+		
 		return PermutationVector;
 	}
 
@@ -149,7 +161,15 @@ class FMobileDirectionalLightFunctionPS : public FMaterialShader
 		return true;
 	}
 
-	static FPermutationDomain BuildPermutationVector(const FViewInfo& View, bool bInlineReflectionAndSky, bool bShadingModelSupport, bool bDynamicShadows, bool bSkyLight, bool bPlanarReflection, EMobileSSRQuality MobileSSRQuality)
+	static FPermutationDomain BuildPermutationVector(
+		const FViewInfo& View, 
+		bool bInlineReflectionAndSky, 
+		bool bShadingModelSupport, 
+		bool bDynamicShadows, 
+		bool bSkyLight, 
+		bool bPlanarReflection, 
+		EMobileSSRQuality MobileSSRQuality,
+		bool bApplySkyShadowing)
 	{
 		bool bUseClusteredLights = UseClusteredDeferredShading(View.GetShaderPlatform());
 		bool bClustredReflection = bInlineReflectionAndSky && (View.NumBoxReflectionCaptures + View.NumSphereReflectionCaptures) > 0;
@@ -166,6 +186,8 @@ class FMobileDirectionalLightFunctionPS : public FMaterialShader
 		PermutationVector.Set<FMobileDirectionalLightFunctionPS::FEnableCSM>(ShadowQuality > 0);
 		PermutationVector.Set<FMobileDirectionalLightFunctionPS::FShadowQuality>(FMath::Clamp(ShadowQuality, 1, 3));
 		PermutationVector.Set<FMobileDirectionalLightFunctionPS::FMobileSSRQuality>(bInlineReflectionAndSky ? MobileSSRQuality : EMobileSSRQuality::Disabled);
+		extern bool UseDistanceFieldAO();
+		PermutationVector.Set<FMobileDirectionalLightFunctionPS::FSkyShadowing>(bInlineReflectionAndSky && bApplySkyShadowing);
 		return PermutationVector;
 	}
 
@@ -280,18 +302,22 @@ public:
 	class FEnablePlanarReflection		: SHADER_PERMUTATION_BOOL("ENABLE_PLANAR_REFLECTION");
 	class FEnableSkyLight				: SHADER_PERMUTATION_BOOL("ENABLE_SKY_LIGHT");
 	class FMobileSSRQuality				: SHADER_PERMUTATION_ENUM_CLASS("MOBILE_SSR_QUALITY", EMobileSSRQuality);
-	
+	class FSkyShadowing					: SHADER_PERMUTATION_BOOL("APPLY_SKY_SHADOWING");
+
 	using FPermutationDomain = TShaderPermutationDomain<
 		FEnableShadingModelSupport, 
 		FEnableClustredReflection, 
 		FEnablePlanarReflection,
 		FEnableSkyLight,
-		FMobileSSRQuality
+		FMobileSSRQuality,
+		FSkyShadowing
 	>;
 
 	BEGIN_SHADER_PARAMETER_STRUCT(FParameters, )
 		SHADER_PARAMETER_STRUCT_REF(FViewUniformShaderParameters, View)
 		SHADER_PARAMETER_STRUCT_REF(FMobileReflectionCaptureShaderData, MobileReflectionCaptureData)
+		SHADER_PARAMETER_STRUCT_INCLUDE(FDFAOUpsampleParameters, DFAOUpsampleParameters)
+		SHADER_PARAMETER_STRUCT_INCLUDE(FSkyDiffuseLightingParameters, SkyDiffuseLighting)
 	END_SHADER_PARAMETER_STRUCT()
 
 	static bool ShouldCompilePermutation(const FGlobalShaderPermutationParameters& Parameters)
@@ -313,6 +339,11 @@ public:
 			return false;
 		}
 
+		if (PermutationVector.Get<FSkyShadowing>() && !IsMobileDistanceFieldAOEnabled(Parameters.Platform))
+		{
+			return false;
+		}
+
 		return true;
 	}
 
@@ -405,7 +436,12 @@ constexpr uint8 PassShadingModelStencilMask(bool bEnableShadingModelSupport)
 	return bEnableShadingModelSupport ? GET_STENCIL_MOBILE_SM_MASK(0xff) : STENCIL_MOBILE_DEFAULTLIT_MASK;
 }
 
-void RenderReflectionEnvironmentSkyLighting(FRHICommandList& RHICmdList, const FScene& Scene, const FViewInfo& View, const EMobileSSRQuality MobileSSRQuality)
+void RenderReflectionEnvironmentSkyLighting(
+	FRHICommandList& RHICmdList, 
+	const FScene& Scene, 
+	const FViewInfo& View, 
+	const EMobileSSRQuality MobileSSRQuality,
+	FRDGTextureRef DynamicBentNormalAOTexture)
 {
 	// Skylights with static lighting already had their diffuse contribution baked into lightmaps
 	const bool bDynamicSkyLight = Scene.SkyLight && (!Scene.SkyLight->bHasStaticLighting || !IsStaticLightingAllowed());
@@ -447,6 +483,14 @@ void RenderReflectionEnvironmentSkyLighting(FRHICommandList& RHICmdList, const F
 	PassParameters.View = GetShaderBinding(View.ViewUniformBuffer);
 	PassParameters.MobileReflectionCaptureData = GetShaderBinding(View.MobileReflectionCaptureUniformBuffer);
 
+	// DFAO
+	if(DynamicBentNormalAOTexture != nullptr)
+	{
+		PassParameters.DFAOUpsampleParameters = DistanceField::SetupAOUpsampleParameters(View, DynamicBentNormalAOTexture);
+		float DynamicBentNormalAO = 1.0f;
+		PassParameters.SkyDiffuseLighting = GetSkyDiffuseLightingParameters(Scene.SkyLight, DynamicBentNormalAO);
+	}
+
 	TShaderMapRef<FPostProcessVS> VertexShader(View.ShaderMap);
 	
 	uint32 PassEnableShadingModelSupport = 0;
@@ -488,6 +532,8 @@ void RenderReflectionEnvironmentSkyLighting(FRHICommandList& RHICmdList, const F
 		PermutationVector.Set<FMobileReflectionEnvironmentSkyLightingPS::FEnablePlanarReflection>(bPlanarReflection);
 		PermutationVector.Set<FMobileReflectionEnvironmentSkyLightingPS::FEnableSkyLight>(bEnableSkyLight);
 		PermutationVector.Set<FMobileReflectionEnvironmentSkyLightingPS::FMobileSSRQuality>((PassEnableSSR & (1 << PassIndex)) ? MobileSSRQuality : EMobileSSRQuality::Disabled);
+		extern bool UseDistanceFieldAO();
+		PermutationVector.Set<FMobileReflectionEnvironmentSkyLightingPS::FSkyShadowing>(DynamicBentNormalAOTexture && UseDistanceFieldAO() && IsMobileDistanceFieldAOEnabled(View.GetShaderPlatform()));
 		TShaderMapRef<FMobileReflectionEnvironmentSkyLightingPS> PixelShader(View.ShaderMap, PermutationVector);
 		
 		GraphicsPSOInit.BoundShaderState.VertexDeclarationRHI = GFilterVertexDeclaration.VertexDeclarationRHI;
@@ -551,7 +597,16 @@ static void SetDirectionalLightDepthStencilState(FGraphicsPipelineStateInitializ
 	}
 }
 
-static void RenderDirectionalLight(FRHICommandList& RHICmdList, const FScene& Scene, const FViewInfo& View, const FCachedLightMaterial& DefaultLightMaterial, const FLightSceneInfo& DirectionalLight, uint32 LightingChannel, bool bInlineReflectionAndSky, EMobileSSRQuality MobileSSRQuality)
+static void RenderDirectionalLight(
+	FRHICommandList& RHICmdList, 
+	const FScene& Scene, 
+	const FViewInfo& View, 
+	const FCachedLightMaterial& DefaultLightMaterial,
+	const FLightSceneInfo& DirectionalLight, 
+	uint32 LightingChannel, 
+	bool bInlineReflectionAndSky, 
+	EMobileSSRQuality MobileSSRQuality,
+	FRDGTextureRef DynamicBentNormalAOTexture)
 {
 	FString LightNameWithLevel;
 	FSceneRenderer::GetLightNameForDrawEvent(DirectionalLight.Proxy, LightNameWithLevel);
@@ -597,6 +652,14 @@ static void RenderDirectionalLight(FRHICommandList& RHICmdList, const FScene& Sc
 		PassParameters.TranslatedWorldToLight = FMatrix44f(FTranslationMatrix(-View.ViewMatrices.GetPreViewTranslation()) * WorldToLight);
 	}
 
+	// DFAO
+	if (DynamicBentNormalAOTexture != nullptr)
+	{
+		PassParameters.DFAOUpsampleParameters = DistanceField::SetupAOUpsampleParameters(View, DynamicBentNormalAOTexture);
+		float DynamicBentNormalAO = 1.0f;
+		PassParameters.SkyDiffuseLighting = GetSkyDiffuseLightingParameters(Scene.SkyLight, DynamicBentNormalAO);
+	}
+
 	// Skylights with static lighting already had their diffuse contribution baked into lightmaps
 	const bool bDynamicSkyLight = Scene.SkyLight && (!Scene.SkyLight->bHasStaticLighting || !IsStaticLightingAllowed());
 	const bool bEnableSkyLight = bDynamicSkyLight && View.Family->EngineShowFlags.SkyLighting;
@@ -652,7 +715,8 @@ static void RenderDirectionalLight(FRHICommandList& RHICmdList, const FScene& Sc
 			bDynamicShadows,
 			bEnableSkyLight,
 			bPlanarReflection, 
-			(PassEnableSSR & (1<<PassIndex)) ? MobileSSRQuality : EMobileSSRQuality::Disabled
+			(PassEnableSSR & (1<<PassIndex)) ? MobileSSRQuality : EMobileSSRQuality::Disabled,
+			DynamicBentNormalAOTexture && UseDistanceFieldAO() && IsMobileDistanceFieldAOEnabled(View.GetShaderPlatform())
 		);
 		FCachedLightMaterial LightMaterial;
 		TShaderRef<FMobileDirectionalLightFunctionPS> PixelShader;
@@ -682,7 +746,13 @@ static void RenderDirectionalLight(FRHICommandList& RHICmdList, const FScene& Sc
 	}
 }
 
-static int RenderDirectionalLights(FRHICommandList& RHICmdList, const FScene& Scene, const FViewInfo& View, const FCachedLightMaterial& DefaultLightMaterial, EMobileSSRQuality MobileSSRQuality)
+static int RenderDirectionalLights(
+	FRHICommandList& RHICmdList, 
+	const FScene& Scene, 
+	const FViewInfo& View, 
+	const FCachedLightMaterial& DefaultLightMaterial, 
+	EMobileSSRQuality MobileSSRQuality,
+	FRDGTextureRef DynamicBentNormalAOTexture)
 {
 	uint32 NumLights = 0;
 	for (uint32 ChannelIdx = 0; ChannelIdx < UE_ARRAY_COUNT(Scene.MobileDirectionalLights); ChannelIdx++)
@@ -695,7 +765,7 @@ static int RenderDirectionalLights(FRHICommandList& RHICmdList, const FScene& Sc
 
 	if (!bInlineReflectionAndSky)
 	{
-		RenderReflectionEnvironmentSkyLighting(RHICmdList, Scene, View, MobileSSRQuality);
+		RenderReflectionEnvironmentSkyLighting(RHICmdList, Scene, View, MobileSSRQuality, DynamicBentNormalAOTexture);
 	}
 
 	for (uint32 ChannelIdx = 0; ChannelIdx < UE_ARRAY_COUNT(Scene.MobileDirectionalLights); ChannelIdx++)
@@ -703,7 +773,7 @@ static int RenderDirectionalLights(FRHICommandList& RHICmdList, const FScene& Sc
 		FLightSceneInfo* DirectionalLight = Scene.MobileDirectionalLights[ChannelIdx];
 		if (DirectionalLight)
 		{
-			RenderDirectionalLight(RHICmdList, Scene, View, DefaultLightMaterial, *DirectionalLight, ChannelIdx, bInlineReflectionAndSky, MobileSSRQuality);
+			RenderDirectionalLight(RHICmdList, Scene, View, DefaultLightMaterial, *DirectionalLight, ChannelIdx, bInlineReflectionAndSky, MobileSSRQuality, DynamicBentNormalAOTexture);
 		}
 	}
 	return NumLights;
@@ -1070,7 +1140,8 @@ void MobileDeferredShadingPass(
 	const FScene& Scene, 
 	const FSortedLightSetSceneInfo& SortedLightSet,
 	const TArray<FVisibleLightInfo, SceneRenderingAllocator>& VisibleLightInfos,
-	EMobileSSRQuality MobileSSRQuality)
+	EMobileSSRQuality MobileSSRQuality,
+	FRDGTextureRef DynamicBentNormalAOTexture)
 {
 	RHI_BREADCRUMB_EVENT_STAT(RHICmdList, DeferredShading, "DeferredShading");
 	SCOPED_GPU_STAT(RHICmdList, DeferredShading);
@@ -1083,7 +1154,7 @@ void MobileDeferredShadingPass(
 	DefaultMaterial.Material = DefaultMaterial.MaterialProxy->GetMaterialNoFallback(ERHIFeatureLevel::ES3_1);
 	check(DefaultMaterial.Material);
 
-	int NumDirLights = RenderDirectionalLights(RHICmdList, Scene, View, DefaultMaterial, MobileSSRQuality);
+	int NumDirLights = RenderDirectionalLights(RHICmdList, Scene, View, DefaultMaterial, MobileSSRQuality, DynamicBentNormalAOTexture);
 	
 	const bool bMobileUseClusteredDeferredShading = UseClusteredDeferredShading(View.GetShaderPlatform()) && NumDirLights > 0;
 	if (!bMobileUseClusteredDeferredShading)
diff --git a/Engine/Source/Runtime/Renderer/Private/MobileDeferredShadingPass.h b/Engine/Source/Runtime/Renderer/Private/MobileDeferredShadingPass.h
--- a/Engine/Source/Runtime/Renderer/Private/MobileDeferredShadingPass.h
+++ b/Engine/Source/Runtime/Renderer/Private/MobileDeferredShadingPass.h
@@ -25,4 +25,5 @@ void MobileDeferredShadingPass(
 	const FScene& Scene, 
 	const FSortedLightSetSceneInfo& SortedLightSet,
 	const TArray<FVisibleLightInfo, SceneRenderingAllocator>& VisibleLightInfos,
-	EMobileSSRQuality MobileSSRQuality);
+	EMobileSSRQuality MobileSSRQuality,
+	FRDGTextureRef DynamicBentNormalAOTexture = nullptr);
diff --git a/Engine/Source/Runtime/Renderer/Private/MobileShadingRenderer.cpp b/Engine/Source/Runtime/Renderer/Private/MobileShadingRenderer.cpp
--- a/Engine/Source/Runtime/Renderer/Private/MobileShadingRenderer.cpp
+++ b/Engine/Source/Runtime/Renderer/Private/MobileShadingRenderer.cpp
@@ -78,6 +78,8 @@
 #include "GenerateMips.h"
 #include "MobileSSR.h"
 #include "ViewData.h"
+#include "DistanceFieldAmbientOcclusion.h"
+#include "DistanceFieldLightingShared.h"
 
 uint32 GetShadowQuality();
 
@@ -230,6 +232,7 @@ BEGIN_SHADER_PARAMETER_STRUCT(FMobileRenderPassParameters,)
 	SHADER_PARAMETER_RDG_BUFFER_SRV(Buffer<uint>, LocalFogVolumeTileDataBuffer)
 	SHADER_PARAMETER_RDG_TEXTURE_SRV(Texture2D<float4>, HalfResLocalFogVolumeViewSRV)
 	SHADER_PARAMETER_RDG_TEXTURE_SRV(Texture2D<float>, HalfResLocalFogVolumeDepthSRV)
+	SHADER_PARAMETER_RDG_TEXTURE(Texture2D, BentNormalAOTexture)
 	RDG_TEXTURE_ACCESS(ColorGradingLUT, ERHIAccess::SRVGraphics)
 	RENDER_TARGET_BINDING_SLOTS()
 END_SHADER_PARAMETER_STRUCT()
@@ -312,6 +315,7 @@ FMobileSceneRenderer::FMobileSceneRenderer(const FSceneViewFamily* InViewFamily,
 	bShouldRenderCustomDepth = false;
 	bRequiresAmbientOcclusionPass = false;
 	bRequiresShadowProjections = false;
+	bEnableDistanceFieldAO = false;
 	bIsFullDepthPrepassEnabled = (Scene->EarlyZPassMode == DDM_AllOpaque || Scene->EarlyZPassMode == DDM_AllOpaqueNoVelocity);
 	bIsMaskedOnlyDepthPrepassEnabled = Scene->EarlyZPassMode == DDM_MaskedOnly;
 	bEnableClusteredLocalLights = MobileForwardEnableLocalLights(ShaderPlatform);
@@ -614,6 +618,18 @@ void FMobileSceneRenderer::InitViews(
 	// never keep MSAA depth if SceneDepthAux is enabled
 	bKeepDepthContent = ((NumMSAASamples > 1) && bRequiresSceneDepthAux) ? false : bKeepDepthContent;
 
+	extern int32 GDistanceFieldAOApplyToStaticIndirect;
+	bEnableDistanceFieldAO =
+		Scene->SkyLight && Scene->SkyLight->bCastShadows
+		&& ShouldRenderDeferredDynamicSkyLight(Scene, ViewFamily)
+		&& AnyViewHasGIMethodSupportingDFAO()
+		&& (Views[0].GlobalDistanceFieldInfo.Clipmaps.Num() > 0)
+		&& !GDistanceFieldAOApplyToStaticIndirect
+		&& ShouldRenderDistanceFieldAO(Views, ViewFamily.EngineShowFlags)
+		&& ShouldRenderDistanceFieldLighting(Scene->DistanceFieldSceneData, Views)
+		&& ViewFamily.EngineShowFlags.AmbientOcclusion
+		&& !Views[0].bIsReflectionCapture;
+	
 	// Depth is needed for Editor Primitives
 	if (bShouldCompositeEditorPrimitives)
 	{
@@ -1411,6 +1427,13 @@ void FMobileSceneRenderer::Render(FRDGBuilder& GraphBuilder, const FSceneRenderU
 			RenderAmbientOcclusion(GraphBuilder, SceneTextures.Depth.Resolve, SceneTextures.ScreenSpaceAO);
 		}
 
+		if (bEnableDistanceFieldAO && (!bDeferredShading || !bRequiresMultiPass))
+		{
+			TArray<FRDGTextureRef> DynamicBentNormalAOTextures;
+			const float OcclusionMaxDistance = Scene->SkyLight && !Scene->SkyLight->bWantsStaticShadowing ? Scene->SkyLight->OcclusionMaxDistance : Scene->DefaultMaxDistanceFieldOcclusionDistance;
+			RenderDistanceFieldLighting(GraphBuilder, SceneTextures, FDistanceFieldAOParameters(OcclusionMaxDistance), DynamicBentNormalAOTextures, false, false, true);
+		}
+
 		// When renderer is in ERendererOutput::DepthPrepassOnly mode, bRequiresShadowProjections is set to false in InitViews()
 		if (bRequiresShadowProjections)
 		{
@@ -2353,17 +2376,30 @@ void FMobileSceneRenderer::RenderDeferredMultiPass(FRDGBuilder& GraphBuilder, FS
 		}
 	}
 
+	TArray<FRDGTextureRef> DynamicBentNormalAOTextures;
+	if(bEnableDistanceFieldAO)
+	{
+		SceneTextures.MobileSetupMode = 
+			EMobileSceneTextureSetupMode::SceneDepth |
+			EMobileSceneTextureSetupMode::GBuffers;
+		SceneTextures.MobileUniformBuffer = CreateMobileSceneTextureUniformBuffer(GraphBuilder, &SceneTextures, SceneTextures.MobileSetupMode);
+		FDistanceFieldAOParameters Parameters(Scene->SkyLight->OcclusionMaxDistance, Scene->SkyLight->Contrast);
+		RenderDistanceFieldLighting(GraphBuilder, SceneTextures, Parameters, DynamicBentNormalAOTextures, false, false);
+	}
+
 	// Lighting and translucency
+	uint32 ViewIndex = 0;
 	for (FRenderViewContext& ViewContext : RenderViews)
 	{
 		FViewInfo& View = *ViewContext.ViewInfo;
-
+		const uint32 CurrentViewIndex = ViewIndex++;
 		SCOPED_GPU_MASK(GraphBuilder.RHICmdList, !View.IsInstancedStereoPass() ? View.GPUMask : (View.GPUMask | View.GetInstancedView()->GPUMask));
 		SCOPED_CONDITIONAL_DRAW_EVENTF(GraphBuilder.RHICmdList, EventView, RenderViews.Num() > 1, TEXT("View%d"), ViewContext.ViewIndex);
 
 		View.BeginRenderView();
 		UpdateDirectionalLightUniformBuffers(GraphBuilder, View);
 
+		FRDGTextureRef DynamicBentNormalAOTexture = DynamicBentNormalAOTextures.IsEmpty() ? nullptr : DynamicBentNormalAOTextures[CurrentViewIndex];;
 		auto* PassParameters = GraphBuilder.AllocParameters<FMobileRenderPassParameters>();
 		PassParameters->View = View.GetShaderParameters();
 
@@ -2379,6 +2415,7 @@ void FMobileSceneRenderer::RenderDeferredMultiPass(FRDGBuilder& GraphBuilder, FS
 		PassParameters->LocalFogVolumeTileDataBuffer = View.LocalFogVolumeViewData.GPUTileDataBufferSRV;
 		PassParameters->HalfResLocalFogVolumeViewSRV = View.LocalFogVolumeViewData.HalfResLocalFogVolumeViewSRV;
 		PassParameters->HalfResLocalFogVolumeDepthSRV = View.LocalFogVolumeViewData.HalfResLocalFogVolumeDepthSRV;
+		PassParameters->BentNormalAOTexture = DynamicBentNormalAOTexture;
 		// Only SceneColor and Depth
 		PassParameters->RenderTargets[0] = BasePassRenderTargets[0];
 		PassParameters->RenderTargets.DepthStencil = BasePassRenderTargets.DepthStencil;
@@ -2398,7 +2435,7 @@ void FMobileSceneRenderer::RenderDeferredMultiPass(FRDGBuilder& GraphBuilder, FS
 			{
 				FViewInfo& View = *ViewContext.ViewInfo;
 
-				MobileDeferredShadingPass(RHICmdList, ViewContext.ViewIndex, Views.Num(), View, *Scene, SortedLightSet, VisibleLightInfos, MobileSSRQuality);
+				MobileDeferredShadingPass(RHICmdList, ViewContext.ViewIndex, Views.Num(), View, *Scene, SortedLightSet, VisibleLightInfos, MobileSSRQuality, PassParameters->BentNormalAOTexture);
 				RenderFog(RHICmdList, View);
 
 				// Draw translucency.
diff --git a/Engine/Source/Runtime/Renderer/Private/SceneRendering.h b/Engine/Source/Runtime/Renderer/Private/SceneRendering.h
--- a/Engine/Source/Runtime/Renderer/Private/SceneRendering.h
+++ b/Engine/Source/Runtime/Renderer/Private/SceneRendering.h
@@ -95,7 +95,9 @@ struct FScreenMessageWriter;
 struct FVolumetricFogIntegrationParameterData;
 class FLumenHardwareRayTracingUniformBufferParameters;
 struct FDBufferTextures;
-
+class FDistanceFieldCulledObjectBufferParameters;
+class FTileIntersectionParameters;
+class FDistanceFieldAOParameters;
 namespace Nanite
 {
 	struct FInstanceDraw;
@@ -2293,6 +2295,34 @@ public:
 		TConstArrayView<const FLightSceneInfo*> RayTracedLightsToInject,
 		FRDGTexture*& OutLocalShadowedLightScattering);
 
+	/** Renders capsule shadows for movable skylights, using the cone of visibility (bent normal) from DFAO. */
+	void RenderCapsuleShadowsForMovableSkylight(
+		FRDGBuilder& GraphBuilder,
+		const FViewInfo& View,
+		FRDGTextureRef& BentNormalOutput) const;
+
+	/** Render Ambient Occlusion using mesh distance fields on a screen based grid. */
+	void RenderDistanceFieldAOScreenGrid(
+		FRDGBuilder& GraphBuilder,
+		const FSceneTextures& SceneTextures,
+		const FViewInfo& View,
+		const FDistanceFieldCulledObjectBufferParameters& CulledObjectBufferParameters,
+		FRDGBufferRef ObjectTilesIndirectArguments,
+		const FTileIntersectionParameters& TileIntersectionParameters,
+		const FDistanceFieldAOParameters& Parameters,
+		FRDGTextureRef DistanceFieldNormal,
+		FRDGTextureRef& OutDynamicBentNormalAO);
+
+	/** Render Ambient Occlusion using mesh distance fields and the surface cache, which supports dynamic rigid meshes. */
+	void RenderDistanceFieldLighting(
+		FRDGBuilder& GraphBuilder,
+		const FSceneTextures& SceneTextures,
+		const class FDistanceFieldAOParameters& Parameters,
+		TArray<FRDGTextureRef>& OutDynamicBentNormalAOTextures,
+		bool bModulateToSceneColor,
+		bool bVisualizeAmbientOcclusion,
+		bool bModulateToScreenSpaceAO = false);
+
 	void SetupVolumetricFog();
 
 	bool ShouldRenderVolumetricFog() const;
@@ -2300,7 +2330,7 @@ public:
 
 	virtual bool IsLumenEnabled(const FViewInfo& View) const { return false; }
 	virtual bool IsLumenGIEnabled(const FViewInfo& View) const { return false; }
-	virtual bool AnyViewHasGIMethodSupportingDFAO() const { return false; }
+	virtual bool AnyViewHasGIMethodSupportingDFAO() const { return true; }
 
 	/** Gets a readable light name for use with a draw event. */
 	static void GetLightNameForDrawEvent(const FLightSceneProxy* LightProxy, FString& LightNameWithLevel);
@@ -2919,7 +2949,7 @@ private:
 	bool bEnableClusteredReflections;
 	bool bRequiresShadowProjections;
 	bool bAdrenoOcclusionMode;
-
+	bool bEnableDistanceFieldAO;
 	ETranslucencyPass::Type StandardTranslucencyPass;
 	EMeshPass::Type StandardTranslucencyMeshPass;
 
