Added dynamic mesh element support to RVT rendering

diff --git a/Engine/Source/Runtime/Engine/Public/MeshElementCollector.h b/Engine/Source/Runtime/Engine/Public/MeshElementCollector.h
--- a/Engine/Source/Runtime/Engine/Public/MeshElementCollector.h
+++ b/Engine/Source/Runtime/Engine/Public/MeshElementCollector.h
@@ -17,6 +17,8 @@ struct FMeshBatchAndRelevance;
 
 ////////////////////////////////////////////////////////////////////////////////////////////////////
 
+namespace RuntimeVirtualTexture { class FDynamicMeshCollector; }
+
 /** 
  * Encapsulates the gathering of meshes from the various FPrimitiveSceneProxy classes. 
  */
@@ -228,6 +230,7 @@ protected:
 	friend class FShadowMeshCollector;
 	friend class FDynamicMeshElementContext;
 	friend struct FRayTracingMaterialGatheringContext;
+	friend class RuntimeVirtualTexture::FDynamicMeshCollector;
 	friend FSceneRenderingBulkObjectAllocator;
 };
 
diff --git a/Engine/Source/Runtime/GeometryFramework/Private/Components/BaseDynamicMeshSceneProxy.cpp b/Engine/Source/Runtime/GeometryFramework/Private/Components/BaseDynamicMeshSceneProxy.cpp
--- a/Engine/Source/Runtime/GeometryFramework/Private/Components/BaseDynamicMeshSceneProxy.cpp
+++ b/Engine/Source/Runtime/GeometryFramework/Private/Components/BaseDynamicMeshSceneProxy.cpp
@@ -67,6 +67,9 @@ FBaseDynamicMeshSceneProxy::FBaseDynamicMeshSceneProxy(UBaseDynamicMeshComponent
 	bAffectDistanceFieldLighting = bWillHaveDistanceField;
 	// note whether lumen is enabled will depend on the distance field flags (in some cases)
 	UpdateVisibleInLumenScene();
+
+	// Dynamic meshes can write to runtime virtual texture if they are set to do so.
+	bSupportsRuntimeVirtualTexture = true;
 }
 
 // Note: deprecation warnings disabled due to removal of distance field support
@@ -480,6 +483,21 @@ void FBaseDynamicMeshSceneProxy::DrawBatch(FMeshElementCollector& Collector, con
 	Mesh.DepthPriorityGroup = DepthPriority;
 	// if this is a wireframe draw pass then we do not want to apply View Mode Overrides
 	Mesh.bCanApplyViewModeOverrides = (bWireframe) ? false : this->bEnableViewModeOverrides;
+
+	for (ERuntimeVirtualTextureMaterialType MaterialType : RuntimeVirtualTextureMaterialTypes)
+	{
+		FMeshBatch& VirtualMesh = Collector.AllocateMesh();
+		VirtualMesh = Mesh;
+		VirtualMesh.CastShadow = 0;
+		VirtualMesh.bUseAsOccluder = 0;
+		VirtualMesh.bUseForDepthPass = 0;
+		VirtualMesh.bUseForMaterial = 0;
+		VirtualMesh.bDitheredLODTransition = 0;
+		VirtualMesh.bRenderToVirtualTexture = 1;
+		VirtualMesh.RuntimeVirtualTextureMaterialType = (uint32)MaterialType;
+		Collector.AddMesh(ViewIndex, VirtualMesh);
+	}
+	
 	Collector.AddMesh(ViewIndex, Mesh);
 }
 
@@ -552,10 +570,12 @@ void FBaseDynamicMeshSceneProxy::DrawStaticElements(FStaticPrimitiveDrawInterfac
 	bool bDrawSecondaryBuffers = ParentBaseComponent->GetSecondaryBuffersVisibility();
 
 	ESceneDepthPriorityGroup DepthPriority = SDPG_World;
+	
+	const int32 NumRuntimeVirtualTextureTypes = RuntimeVirtualTextureMaterialTypes.Num();
 
 	TArray<FMeshRenderBufferSet*> Buffers;
 	GetActiveRenderBufferSets(Buffers);
-	PDI->ReserveMemoryForMeshes(Buffers.Num());
+	PDI->ReserveMemoryForMeshes(Buffers.Num() * (1 + NumRuntimeVirtualTextureTypes));
 
 	// Draw the mesh.
 	int32 SectionIndexCounter = 0;
@@ -603,6 +623,23 @@ void FBaseDynamicMeshSceneProxy::DrawStaticElements(FStaticPrimitiveDrawInterfac
 		MeshBatch.bDitheredLODTransition = false;
 		MeshBatch.bUseForDepthPass = true;
 		MeshBatch.bUseAsOccluder = ShouldUseAsOccluder();
+		
+		if (NumRuntimeVirtualTextureTypes > 0)
+		{
+			FMeshBatch VirtualMeshBatch(MeshBatch);
+			VirtualMeshBatch.CastShadow = 0;
+			VirtualMeshBatch.bUseAsOccluder = 0;
+			VirtualMeshBatch.bUseForDepthPass = 0;
+			VirtualMeshBatch.bUseForMaterial = 0;
+			VirtualMeshBatch.bDitheredLODTransition = 0;
+			VirtualMeshBatch.bRenderToVirtualTexture = 1;
+			
+			for (ERuntimeVirtualTextureMaterialType MaterialType : RuntimeVirtualTextureMaterialTypes)
+			{
+				VirtualMeshBatch.RuntimeVirtualTextureMaterialType = (uint32)MaterialType;
+				PDI->DrawMesh(VirtualMeshBatch, FLT_MAX);
+			}
+		}
 
 		PDI->DrawMesh(MeshBatch, FLT_MAX);
 	}
diff --git a/Engine/Source/Runtime/Renderer/Private/GPUScene.h b/Engine/Source/Runtime/Renderer/Private/GPUScene.h
--- a/Engine/Source/Runtime/Renderer/Private/GPUScene.h
+++ b/Engine/Source/Runtime/Renderer/Private/GPUScene.h
@@ -286,6 +286,7 @@ public:
 	int32 GetNumPrimitives() const { return DynamicPrimitivesOffset; }
 	int32 GetMaxLightId() const { return NumGPULights; }
 	int32 GetNumLightmapDataItems() const { return LightmapDataAllocator.GetMaxSize(); }
+	FGPUSceneDynamicContext* GetCurrentDynamicContext() const { return CurrentDynamicContext; }
 	/**
 	 * Returns the highest instance ID that is represented in the GPU scene (which may be lower than the host allocated IDs due to various limits)
 	 * Never larger than MAX_INSTANCE_ID, see Engine\Shaders\Shared\SceneDefinitions.h
diff --git a/Engine/Source/Runtime/Renderer/Private/VT/RuntimeVirtualTextureRender.cpp b/Engine/Source/Runtime/Renderer/Private/VT/RuntimeVirtualTextureRender.cpp
--- a/Engine/Source/Runtime/Renderer/Private/VT/RuntimeVirtualTextureRender.cpp
+++ b/Engine/Source/Runtime/Renderer/Private/VT/RuntimeVirtualTextureRender.cpp
@@ -83,6 +83,8 @@ namespace RuntimeVirtualTexture
 	END_GLOBAL_SHADER_PARAMETER_STRUCT()
 
 	IMPLEMENT_GLOBAL_SHADER_PARAMETER_STRUCT(FEtcParameters, "EtcParameters");
+	
+	static TGlobalResource<FGlobalDynamicReadBuffer> RuntimeVirtualReadBuffer;
 
 	class FEtcParametersUniformBuffer : public TUniformBuffer<FEtcParameters>
 	{
@@ -1558,6 +1560,86 @@ namespace RuntimeVirtualTexture
 		FMeshPassProcessorRenderState DrawRenderState;
 	};
 
+	/** Mesh collector for all dynamic mesh batches */
+	class FDynamicMeshCollector
+	{
+	public:
+		FDynamicMeshCollector(FRHICommandList& RHICmdList, FSceneRenderingBulkObjectAllocator& Allocator, FScene* Scene, FViewInfo* View)
+			: View(View)
+			, Scene(Scene)
+			, Collector(GMaxRHIFeatureLevel, Allocator)
+		{
+			DynamicVertexBuffer.Init(RHICmdList);
+			DynamicIndexBuffer.Init(RHICmdList);
+			Collector.Start(RHICmdList, DynamicVertexBuffer, DynamicIndexBuffer, RuntimeVirtualReadBuffer);
+
+			// Create a new primitive collector for this page set only
+			if (Scene->GPUScene.IsEnabled())
+			{
+				View->DynamicPrimitiveCollector = FGPUScenePrimitiveCollector(Scene->GPUScene.GetCurrentDynamicContext());
+			}
+
+			Collector.AddViewMeshArrays(
+				View,
+				&MeshBatchAndRelevances, 
+				&SimpleElements,
+				&View->DynamicPrimitiveCollector
+			);
+		}
+
+		/** Commit and submit batches */
+		void Submit(FRDGBuilder& GraphBuilder, FRuntimeVirtualTextureMeshProcessor& Processor, ERuntimeVirtualTextureMaterialType MaterialType)
+		{
+			Collector.Finish();
+			DynamicVertexBuffer.Commit();
+			DynamicIndexBuffer.Commit();
+
+			if (Scene->GPUScene.IsEnabled())
+			{
+				View->DynamicPrimitiveCollector.Commit();
+			}
+
+			// Avoid uploads on empty submissions
+			if (MeshBatchAndRelevances.IsEmpty())
+			{
+				return;
+			}
+
+			Scene->GPUScene.UploadDynamicPrimitiveShaderDataForView(GraphBuilder, *View);
+
+			// Process all batches marked for virtual texturing
+			for (FMeshBatchAndRelevance& MeshBatch : MeshBatchAndRelevances)
+			{
+				if (MeshBatch.Mesh->bRenderToVirtualTexture && MeshBatch.Mesh->RuntimeVirtualTextureMaterialType == static_cast<uint32>(MaterialType))
+				{
+					Processor.AddMeshBatch(*MeshBatch.Mesh, ~0ull, MeshBatch.PrimitiveSceneProxy);
+				}
+			}
+		}
+
+		void SetPrimitive(FPrimitiveSceneProxy* Proxy, FHitProxyId HitProxyId)
+		{
+			Collector.SetPrimitive(Proxy, HitProxyId);
+		}
+
+		FMeshElementCollector& GetCollector()
+		{
+			return Collector;
+		}
+
+	private:
+		FViewInfo* View;
+		FScene*    Scene;
+
+		/** Collection states, simple elements ignored */
+		FMeshElementCollector      Collector;
+		FGlobalDynamicVertexBuffer DynamicVertexBuffer;
+		FGlobalDynamicIndexBuffer  DynamicIndexBuffer;
+		FSimpleElementCollector    SimpleElements;
+
+		/** All collected batches */
+		TArray<FMeshBatchAndRelevance, SceneRenderingAllocator> MeshBatchAndRelevances;
+	};
 
 	/** Registration for virtual texture command caching pass */
 	FMeshPassProcessor* CreateRuntimeVirtualTexturePassProcessor(ERHIFeatureLevel::Type FeatureLevel, const FScene* Scene, const FSceneView* InViewIfDynamicMeshCommand, FMeshPassDrawListContext* InDrawListContext)
@@ -1570,7 +1652,7 @@ namespace RuntimeVirtualTexture
 
 
 	/** Collect meshes to draw. */
-	void GatherMeshesToDraw(FDynamicPassMeshDrawListContext* DynamicMeshPassContext, FScene const* Scene, FViewInfo* View, ERuntimeVirtualTextureMaterialType MaterialType, int32 RuntimeVirtualTextureId, uint8 vLevel, uint8 MaxLevel, bool bAllowCachedMeshDrawCommands)
+	void GatherMeshesToDraw(FDynamicPassMeshDrawListContext* DynamicMeshPassContext, FRDGBuilder& GraphBuilder, FScene* Scene, FViewInfo* View, ERuntimeVirtualTextureMaterialType MaterialType, int32 RuntimeVirtualTextureId, uint8 vLevel, uint8 MaxLevel, bool bAllowCachedMeshDrawCommands)
 	{
 		// Cached draw command collectors
 		const FCachedPassMeshDrawList& SceneDrawList = Scene->CachedDrawLists[EMeshPass::VirtualTexture];
@@ -1588,6 +1670,12 @@ namespace RuntimeVirtualTexture
 			SceneExtension->GetPrimitivesForRuntimeVirtualTexture(Scene, RuntimeVirtualTextureId, PrimitiveIndices);
 		}
 
+		// Lazily created on first usage
+		FDynamicMeshCollector* DynamicCollector = nullptr;
+
+		// Set of views for dynamic collection
+		TArray<const FSceneView*> Views{View};
+
 		for (const int32 PrimitiveIndex : PrimitiveIndices)
  		{
 			//todo[vt]: In our case we know that frustum is an oriented box so investigate cheaper test for intersecting that
@@ -1672,6 +1760,27 @@ namespace RuntimeVirtualTexture
 					}
 				}
 			}
+
+			// It's not entirely accurate to equate the dynamic view relevance here, so, assume dynamic
+			// mesh batches if no static were supplied. For now at least.
+			if (PrimitiveSceneInfo->StaticMeshes.IsEmpty())
+			{
+				// Lazy create
+				if (!DynamicCollector)
+				{
+					DynamicCollector = View->Allocator.Create<FDynamicMeshCollector>(GraphBuilder.RHICmdList, View->Allocator, Scene, View);
+				}
+				
+				// Collect all dynamic batches
+				DynamicCollector->SetPrimitive(PrimitiveSceneInfo->Proxy, PrimitiveSceneInfo->DefaultDynamicHitProxyId);
+				PrimitiveSceneInfo->Proxy->GetDynamicMeshElements(Views, *View->Family, 0x1, DynamicCollector->GetCollector());
+			}
+		}
+
+		// Process collected batches
+		if (DynamicCollector)
+		{
+			DynamicCollector->Submit(GraphBuilder, MeshProcessor, MaterialType);
 		}
 	}
 
@@ -2122,9 +2231,9 @@ namespace RuntimeVirtualTexture
 			const bool bAllowCachedMeshDrawCommands = BatchRenderContext.bAllowCachedMeshDrawCommands;
 
 			AddSimpleMeshPass(GraphBuilder, PassParameters, Scene, *View, nullptr, RDG_EVENT_NAME("VirtualTextureDraw"), ERDGPassFlags::Raster | ERDGPassFlags::NeverMerge,
-	    	[Scene, View, MaterialType, RuntimeVirtualTextureId, vLevel, MaxLevel, bAllowCachedMeshDrawCommands](FDynamicPassMeshDrawListContext* DynamicMeshPassContext)
+	    	[Scene, &GraphBuilder, View, MaterialType, RuntimeVirtualTextureId, vLevel, MaxLevel, bAllowCachedMeshDrawCommands](FDynamicPassMeshDrawListContext* DynamicMeshPassContext)
 	    	{
-	    		GatherMeshesToDraw(DynamicMeshPassContext, Scene, View, MaterialType, RuntimeVirtualTextureId, vLevel, MaxLevel, bAllowCachedMeshDrawCommands);
+	    		GatherMeshesToDraw(DynamicMeshPassContext, GraphBuilder, Scene, View, MaterialType, RuntimeVirtualTextureId, vLevel, MaxLevel, bAllowCachedMeshDrawCommands);
 	    	}, 
 			[ViewRect = View->ViewRect, PageIndex, ShaderPlatform = Scene->GetShaderPlatform()](FRHICommandList& RHICmdList)
 			{
