SkyAtmospehre - Optimisation for SkyAtmosphereLUT CPU thread cost.
Basically we no longer create a new view + UB, instead we set the variable we are interested in and a special permutation is used to drive parameter fetching.

#rnx
#tests pc editor + preflight

diff --git a/Engine/Shaders/Private/SkyAtmosphere.usf b/Engine/Shaders/Private/SkyAtmosphere.usf
--- a/Engine/Shaders/Private/SkyAtmosphere.usf
+++ b/Engine/Shaders/Private/SkyAtmosphere.usf
@@ -162,20 +162,6 @@ float4 GetScreenTranslatedWorldPos(float4 SVPos, float DeviceZ)
 	return float4(SvPositionToTranslatedWorld(float4(SVPos.xy, DeviceZ, 1.0)), 1.0);
 }
 
-// This is the world position of the camera. It is also force to be at the top of the virutal planet surface. 
-// This is to always see the sky even when the camera is buried into the virtual planet.
-float3 GetCameraTranslatedWorldPos()
-{
-	return View.SkyCameraTranslatedWorldOrigin;
-}
-
-// This is the camera position relative to the virtual planet center.
-// This is convenient because for all the math in this file using world position relative to the virtual planet center.
-float3 GetTranslatedCameraPlanetPos()
-{
-	return (GetCameraTranslatedWorldPos() - View.SkyPlanetTranslatedWorldCenterAndViewHeight.xyz)*CM_TO_SKY_UNIT;
-}
-
 bool MoveToTopAtmosphere(inout float3 WorldPos, in float3 WorldDir, in float AtmosphereTopRadius)
 {
 	float ViewHeight = length(WorldPos);
@@ -262,7 +248,6 @@ void UvToSkyViewLutParams(out float3 ViewDir, in float ViewHeight, in float2 UV)
 }
 
 
-
 ////////////////////////////////////////////////////////////
 // Utilities
 ////////////////////////////////////////////////////////////
@@ -280,6 +265,65 @@ void UvToSkyViewLutParams(out float3 ViewDir, in float ViewHeight, in float2 UV)
 #endif
 
 
+////////////////////////////////////////////////////////////
+// Real time reflection capture overriden parameters
+////////////////////////////////////////////////////////////
+
+#if PERMUTATION_REALTIME_REFLECTION_LUT
+
+#undef OutputPreExposure
+#define OutputPreExposure View.PreExposure
+
+float3x3 GetSkyViewLutReferentialParameter()
+{
+	return GetSkyViewLutReferential(SkyAtmosphereRealTimeReflectionLUTParameters.SkyViewLutReferential);
+}
+
+float3 GetSkyPlanetTranslatedWorldCenterAndViewHeightParameter()
+{
+	return SkyAtmosphereRealTimeReflectionLUTParameters.SkyPlanetTranslatedWorldCenterAndViewHeight.xyz;
+}
+
+float3 GetSkyCameraTranslatedWorldOriginParameter()
+{
+	return SkyAtmosphereRealTimeReflectionLUTParameters.SkyCameraTranslatedWorldOrigin;
+}
+
+#else // PERMUTATION_REALTIME_REFLECTION_LUT
+
+float3x3 GetSkyViewLutReferentialParameter()
+{
+	return GetSkyViewLutReferential(View.SkyViewLutReferential);
+}
+
+float3 GetSkyPlanetTranslatedWorldCenterAndViewHeightParameter()
+{
+	return View.SkyPlanetTranslatedWorldCenterAndViewHeight.xyz;
+}
+
+float3 GetSkyCameraTranslatedWorldOriginParameter()
+{
+	return View.SkyCameraTranslatedWorldOrigin;
+}
+
+#endif // PERMUTATION_REALTIME_REFLECTION_LUT
+
+// This is the world position of the camera. It is also force to be at the top of the virutal planet surface. 
+// This is to always see the sky even when the camera is buried into the virtual planet.
+float3 GetCameraTranslatedWorldPos()
+{
+	return GetSkyCameraTranslatedWorldOriginParameter();
+}
+
+// This is the camera position relative to the virtual planet center.
+// This is convenient because for all the math in this file using world position relative to the virtual planet center.
+float3 GetTranslatedCameraPlanetPos()
+{
+	return (GetCameraTranslatedWorldPos() - GetSkyPlanetTranslatedWorldCenterAndViewHeightParameter()) * CM_TO_SKY_UNIT;
+}
+
+
+
 SamplerState TransmittanceLutTextureSampler;
 SamplerState MultiScatteredLuminanceLutTextureSampler;
 SamplerState SkyViewLutTextureSampler;
@@ -487,7 +531,7 @@ SingleScatteringResult IntegrateSingleScatteredLuminance(
 	if (DeviceZ != FarDepthValue)
 	{
 		const float3 DepthBufferTranslatedWorldPosKm = GetScreenTranslatedWorldPos(SVPos, DeviceZ).xyz * CM_TO_SKY_UNIT;
-		const float3 TraceStartTranslatedWorldPosKm  = WorldPos + View.SkyPlanetTranslatedWorldCenterAndViewHeight.xyz * CM_TO_SKY_UNIT; // apply planet offset to go back to world from planet local referencial.
+		const float3 TraceStartTranslatedWorldPosKm  = WorldPos + GetSkyPlanetTranslatedWorldCenterAndViewHeightParameter() * CM_TO_SKY_UNIT; // apply planet offset to go back to world from planet local referencial.
 		const float3 TraceStartToSurfaceWorldKm = DepthBufferTranslatedWorldPosKm - TraceStartTranslatedWorldPosKm;
 		float tDepth = length(TraceStartToSurfaceWorldKm);
 		if (tDepth < tMax)
@@ -555,7 +599,7 @@ SingleScatteringResult IntegrateSingleScatteredLuminance(
 //#if SAMPLE_OPAQUE_SHADOW
 	// Get the referencial when rendering the SkyView lut being in a special Z-top space
 #if SKYVIEWLUT_PASS
-	float3x3 LocalReferencial = GetSkyViewLutReferential(View.SkyViewLutReferential);
+	float3x3 LocalReferencial = GetSkyViewLutReferentialParameter();
 #endif
 //#endif
 
@@ -666,7 +710,7 @@ SingleScatteringResult IntegrateSingleScatteredLuminance(
 #endif
 #if SAMPLE_OPAQUE_SHADOW
 		{
-			float3 ShadowSampleWorldPosition0 = ShadowP0 * SKY_UNIT_TO_CM + View.SkyPlanetTranslatedWorldCenterAndViewHeight.xyz;
+			float3 ShadowSampleWorldPosition0 = ShadowP0 * SKY_UNIT_TO_CM + GetSkyPlanetTranslatedWorldCenterAndViewHeightParameter();
 			PlanetShadow0 *= ComputeLight0VolumeShadowing(ShadowSampleWorldPosition0 /* - DFHackToFloat(PrimaryView.PreViewTranslation)*/, false, false, bUnused);
 
 		#if VIRTUAL_SHADOW_MAP
@@ -680,12 +724,12 @@ SingleScatteringResult IntegrateSingleScatteredLuminance(
 #endif
 #if SAMPLE_CLOUD_SKYAO
 		float OutOpticalDepth = 0.0f;
-		MultiScatteredLuminance0 *= GetCloudVolumetricShadow(ShadowP0 * SKY_UNIT_TO_CM + View.SkyPlanetTranslatedWorldCenterAndViewHeight.xyz, VolumetricCloudCommonParameters.CloudSkyAOTranslatedWorldToLightClipMatrix,
+		MultiScatteredLuminance0 *= GetCloudVolumetricShadow(ShadowP0 * SKY_UNIT_TO_CM + GetSkyPlanetTranslatedWorldCenterAndViewHeightParameter(), VolumetricCloudCommonParameters.CloudSkyAOTranslatedWorldToLightClipMatrix,
 			VolumetricCloudCommonParameters.CloudSkyAOFarDepthKm, VolumetricCloudSkyAOTexture, VolumetricCloudSkyAOTextureSampler, OutOpticalDepth);
 #endif
 #if SAMPLE_CLOUD_SHADOW
 		float OutOpticalDepth2 = 0.0f;
-		PlanetShadow0 *= saturate(lerp(1.0f, GetCloudVolumetricShadow(ShadowP0 * SKY_UNIT_TO_CM + View.SkyPlanetTranslatedWorldCenterAndViewHeight.xyz, VolumetricCloudCommonParameters.CloudShadowmapTranslatedWorldToLightClipMatrix[0],
+		PlanetShadow0 *= saturate(lerp(1.0f, GetCloudVolumetricShadow(ShadowP0 * SKY_UNIT_TO_CM + GetSkyPlanetTranslatedWorldCenterAndViewHeightParameter(), VolumetricCloudCommonParameters.CloudShadowmapTranslatedWorldToLightClipMatrix[0],
 			VolumetricCloudCommonParameters.CloudShadowmapFarDepthKm[0].x, VolumetricCloudShadowMapTexture0, VolumetricCloudShadowMapTexture0Sampler, OutOpticalDepth2), VolumetricCloudShadowStrength0));
 #endif
 		// MultiScatteredLuminance is already pre-exposed, atmospheric light contribution needs to be pre exposed
@@ -703,7 +747,7 @@ SingleScatteringResult IntegrateSingleScatteredLuminance(
 		ShadowP1 = GetTranslatedCameraPlanetPos() + t * mul(LocalReferencial, WorldDir); // Inverse of the local SkyViewLUT referencial transform
 	#endif
 		{
-			float3 ShadowSampleWorldPosition1 = ShadowP1 * SKY_UNIT_TO_CM + View.SkyPlanetTranslatedWorldCenterAndViewHeight.xyz;
+			float3 ShadowSampleWorldPosition1 = ShadowP1 * SKY_UNIT_TO_CM + GetSkyPlanetTranslatedWorldCenterAndViewHeightParameter();
 			PlanetShadow1 *= ComputeLight1VolumeShadowing(ShadowSampleWorldPosition1/* - DFHackToFloat(PrimaryView.PreViewTranslation)*/, false, false, bUnused);
 		#if VIRTUAL_SHADOW_MAP
 			if (VirtualShadowMapId1 != INDEX_NONE)
@@ -716,7 +760,7 @@ SingleScatteringResult IntegrateSingleScatteredLuminance(
 #endif // SAMPLE_OPAQUE_SHADOW
 #if SAMPLE_CLOUD_SHADOW
 		float OutOpticalDepth3 = 0.0f;
-		PlanetShadow1 *= saturate(lerp(1.0f, GetCloudVolumetricShadow(ShadowP1 * SKY_UNIT_TO_CM + View.SkyPlanetTranslatedWorldCenterAndViewHeight.xyz, VolumetricCloudCommonParameters.CloudShadowmapTranslatedWorldToLightClipMatrix[1],
+		PlanetShadow1 *= saturate(lerp(1.0f, GetCloudVolumetricShadow(ShadowP1 * SKY_UNIT_TO_CM + GetSkyPlanetTranslatedWorldCenterAndViewHeightParameter(), VolumetricCloudCommonParameters.CloudShadowmapTranslatedWorldToLightClipMatrix[1],
 			VolumetricCloudCommonParameters.CloudShadowmapFarDepthKm[1].x, VolumetricCloudShadowMapTexture1, VolumetricCloudShadowMapTexture1Sampler, OutOpticalDepth3), VolumetricCloudShadowStrength1));
 #endif
 		//  Multi-scattering can work for the second light but it is disabled for the sake of performance.
@@ -954,7 +998,7 @@ void RenderSkyAtmosphereRayMarchingPS(
 		float2 UV;
 
 		// The referencial used to build the Sky View lut
-		float3x3 LocalReferencial = GetSkyViewLutReferential(View.SkyViewLutReferential);
+		float3x3 LocalReferencial = GetSkyViewLutReferentialParameter();
 
 		// Input vectors expressed in this referencial: Up is always Z. Also note that ViewHeight is unchanged in this referencial.
 		float3 WorldPosLocal = float3(0.0, 0.0, ViewHeight);
@@ -1286,7 +1330,7 @@ void RenderSkyViewLutCS(uint3 ThreadId : SV_DispatchThreadID)
 	// into a referential with UP being perpendicular to the ground. And with origin at the planet center.
 
 	// This is the local referencial
-	float3x3 LocalReferencial = GetSkyViewLutReferential(View.SkyViewLutReferential);
+	float3x3 LocalReferencial = GetSkyViewLutReferentialParameter();
 
 	// This is the LUT camera height and position in the local referential
 	float ViewHeight = length(WorldPos);
@@ -1674,7 +1718,7 @@ void RenderSkyAtmosphereDebugPS(
 	const float AerialPespectiveViewDistanceScale = 1.0f;
 
 	// We position the camera at the same height as the view port camera, along z-up vector
-	float3 WorldPos = float3(0.0f, 0.0f, length((View.TranslatedWorldCameraOrigin - View.SkyPlanetTranslatedWorldCenterAndViewHeight.xyz) * CM_TO_SKY_UNIT));
+	float3 WorldPos = float3(0.0f, 0.0f, length((View.TranslatedWorldCameraOrigin - GetSkyPlanetTranslatedWorldCenterAndViewHeightParameter()) * CM_TO_SKY_UNIT));
 
 	// Some organisation constants
 	const float Margin = 2.0f;
diff --git a/Engine/Source/Runtime/Renderer/Private/SkyAtmosphereRendering.cpp b/Engine/Source/Runtime/Renderer/Private/SkyAtmosphereRendering.cpp
--- a/Engine/Source/Runtime/Renderer/Private/SkyAtmosphereRendering.cpp
+++ b/Engine/Source/Runtime/Renderer/Private/SkyAtmosphereRendering.cpp
@@ -258,13 +258,31 @@ BEGIN_GLOBAL_SHADER_PARAMETER_STRUCT(FSkyAtmosphereInternalCommonParameters, )
 
 	SHADER_PARAMETER(FVector3f, SkyLuminanceFactor)
 	SHADER_PARAMETER(FVector3f, SkyAndAerialPerspectiveLuminanceFactor)
-
 END_GLOBAL_SHADER_PARAMETER_STRUCT()
 
 IMPLEMENT_GLOBAL_SHADER_PARAMETER_STRUCT(FAtmosphereUniformShaderParameters, "Atmosphere");
 IMPLEMENT_GLOBAL_SHADER_PARAMETER_STRUCT(FSkyAtmosphereInternalCommonParameters, "SkyAtmosphere");
 
 
+BEGIN_GLOBAL_SHADER_PARAMETER_STRUCT(FSkyAtmosphereRealTimeReflectionLUTParameters, )
+	SHADER_PARAMETER(FMatrix44f, SkyViewLutReferential)
+	SHADER_PARAMETER(FVector4f, SkyPlanetTranslatedWorldCenterAndViewHeight)
+	SHADER_PARAMETER(FVector3f, SkyCameraTranslatedWorldOrigin)
+	SHADER_PARAMETER(float, RealTimeReflectionCapture)
+END_GLOBAL_SHADER_PARAMETER_STRUCT()
+IMPLEMENT_GLOBAL_SHADER_PARAMETER_STRUCT(FSkyAtmosphereRealTimeReflectionLUTParameters, "SkyAtmosphereRealTimeReflectionLUTParameters");
+
+auto GetDefaultSkyAtmosphereRealTimeReflectionLUTUniformBuffer()
+{
+	FSkyAtmosphereRealTimeReflectionLUTParameters Parameters;
+	Parameters.SkyViewLutReferential = FMatrix44f::Identity;
+	Parameters.SkyPlanetTranslatedWorldCenterAndViewHeight = FVector4f::Zero();
+	Parameters.SkyCameraTranslatedWorldOrigin = FVector3f::Zero();
+	Parameters.RealTimeReflectionCapture = 0.0f;
+	return TUniformBufferRef<FSkyAtmosphereRealTimeReflectionLUTParameters>::CreateUniformBufferImmediate(Parameters, UniformBuffer_SingleDraw);
+}
+
+
 #define GET_VALID_DATA_FROM_CVAR \
 	auto ValidateLUTResolution = [](int32 Value) \
 	{ \
@@ -635,6 +653,7 @@ class FSampleCloudSkyAO : SHADER_PERMUTATION_BOOL("SAMPLE_CLOUD_SKYAO");
 class FAtmosphereOnClouds : SHADER_PERMUTATION_BOOL("SAMPLE_ATMOSPHERE_ON_CLOUDS");
 class FMSAASampleCount : SHADER_PERMUTATION_SPARSE_INT("MSAA_SAMPLE_COUNT", 1, 2, 4, 8);
 class FSeparateMieAndRayleighScattering : SHADER_PERMUTATION_BOOL("SEPARATE_MIE_RAYLEIGH_SCATTERING");
+class FRealTimeReflectionLUT : SHADER_PERMUTATION_BOOL("PERMUTATION_REALTIME_REFLECTION_LUT");
 
 }
 
@@ -891,7 +910,7 @@ class FRenderSkyViewLutCS : public FGlobalShader
 	DECLARE_GLOBAL_SHADER(FRenderSkyViewLutCS);
 	SHADER_USE_PARAMETER_STRUCT(FRenderSkyViewLutCS, FGlobalShader);
 
-	using FPermutationDomain = TShaderPermutationDomain<FSampleCloudSkyAO, FSecondAtmosphereLight, FSampleOpaqueShadow, FSampleCloudShadow>;
+	using FPermutationDomain = TShaderPermutationDomain<FSampleCloudSkyAO, FSecondAtmosphereLight, FSampleOpaqueShadow, FSampleCloudShadow, FRealTimeReflectionLUT>;
 
 public:
 	const static uint32 GroupSize = 8;
@@ -900,6 +919,7 @@ public:
 		SHADER_PARAMETER_STRUCT_REF(FAtmosphereUniformShaderParameters, Atmosphere)
 		SHADER_PARAMETER_STRUCT_REF(FSkyAtmosphereInternalCommonParameters, SkyAtmosphere)
 		SHADER_PARAMETER_STRUCT_REF(FViewUniformShaderParameters, ViewUniformBuffer)
+		SHADER_PARAMETER_STRUCT_REF(FSkyAtmosphereRealTimeReflectionLUTParameters, SkyAtmosphereRealTimeReflectionLUTUniformBuffer)
 		SHADER_PARAMETER_RDG_TEXTURE_UAV(RWTexture2D<float3>, SkyViewLutUAV)
 		SHADER_PARAMETER_RDG_TEXTURE(Texture2D<float3>, TransmittanceLutTexture)
 		SHADER_PARAMETER_RDG_TEXTURE(Texture2D<float3>, MultiScatteredLuminanceLutTexture)
@@ -948,7 +968,7 @@ class FRenderCameraAerialPerspectiveVolumeCS : public FGlobalShader
 	DECLARE_GLOBAL_SHADER(FRenderCameraAerialPerspectiveVolumeCS);
 	SHADER_USE_PARAMETER_STRUCT(FRenderCameraAerialPerspectiveVolumeCS, FGlobalShader);
 
-	using FPermutationDomain = TShaderPermutationDomain<FSampleCloudSkyAO, FSecondAtmosphereLight, FSampleOpaqueShadow, FSampleCloudShadow, FSeparateMieAndRayleighScattering>;
+	using FPermutationDomain = TShaderPermutationDomain<FSampleCloudSkyAO, FSecondAtmosphereLight, FSampleOpaqueShadow, FSampleCloudShadow, FSeparateMieAndRayleighScattering, FRealTimeReflectionLUT>;
 
 public:
 	const static uint32 GroupSize = 4;
@@ -957,6 +977,7 @@ public:
 		SHADER_PARAMETER_STRUCT_REF(FAtmosphereUniformShaderParameters, Atmosphere)
 		SHADER_PARAMETER_STRUCT_REF(FSkyAtmosphereInternalCommonParameters, SkyAtmosphere)
 		SHADER_PARAMETER_STRUCT_REF(FViewUniformShaderParameters, ViewUniformBuffer)
+		SHADER_PARAMETER_STRUCT_REF(FSkyAtmosphereRealTimeReflectionLUTParameters, SkyAtmosphereRealTimeReflectionLUTUniformBuffer)
 		SHADER_PARAMETER_RDG_TEXTURE_UAV(RWTexture3D<float4>, CameraAerialPerspectiveVolumeUAV)
 		SHADER_PARAMETER_RDG_TEXTURE_UAV(RWTexture3D<float4>, CameraAerialPerspectiveVolumeMieOnlyUAV)
 		SHADER_PARAMETER_RDG_TEXTURE_UAV(RWTexture3D<float4>, CameraAerialPerspectiveVolumeRayOnlyUAV)
@@ -1469,23 +1490,8 @@ void FSceneRenderer::RenderSkyAtmosphereLookUpTables(FRDGBuilder& GraphBuilder,
 	const bool bRealTimeReflectionCapture360APLutTexture = Views.Num() > 0 && Scene->SkyLight && Scene->SkyLight->bRealTimeCaptureEnabled && (Views[0].bSceneHasSkyMaterial || Scene->HasVolumetricCloud());
 	if (bRealTimeReflectionCaptureSkyAtmosphereViewLutTexture || bRealTimeReflectionCapture360APLutTexture)
 	{
-		FViewInfo& View = Views[0];
-		const float AerialPerspectiveStartDepthInCm = GetValidAerialPerspectiveStartDepthInCm(View, SkyAtmosphereSceneProxy);
-
-		FViewUniformShaderParameters ReflectionViewParameters = *View.CachedViewUniformShaderParameters;
-		FViewMatrices ViewMatrices = View.ViewMatrices;
-		ViewMatrices.HackRemoveTemporalAAProjectionJitter();
-		ViewMatrices.UpdateViewMatrix(Scene->SkyLight->CapturePosition, FRotator(EForceInit::ForceInitToZero));
-		View.SetupCommonViewUniformBufferParameters(
-			ReflectionViewParameters,
-			View.ViewRect.Size(), 1,
-			View.ViewRect,
-			ViewMatrices,
-			ViewMatrices
-		);
-
-		// LUTs still need to be pre-exposed as usual so we set reflection to 0.
-		ReflectionViewParameters.RealTimeReflectionCapture = 0.0f;
+		FViewInfo& ViewForReflection = Views[0];
+		const float AerialPerspectiveStartDepthInCm = GetValidAerialPerspectiveStartDepthInCm(ViewForReflection, SkyAtmosphereSceneProxy);
 
 		// Setup a constant referential for each of the faces of the dynamic reflection capture.
 		const FAtmosphereSetup& AtmosphereSetup = SkyAtmosphereSceneProxy.GetAtmosphereSetup();
@@ -1495,22 +1501,28 @@ void FSceneRenderer::RenderSkyAtmosphereLookUpTables(FRDGBuilder& GraphBuilder,
 		FMatrix44f SkyViewLutReferential;
 		FVector4f TempSkyPlanetData;
 		AtmosphereSetup.ComputeViewData(
-			Scene->SkyLight->CapturePosition, View.ViewMatrices.GetPreViewTranslation(), SkyViewLutReferentialForward, SkyViewLutReferentialRight,
+			Scene->SkyLight->CapturePosition, ViewForReflection.ViewMatrices.GetPreViewTranslation(), SkyViewLutReferentialForward, SkyViewLutReferentialRight,
 			SkyCameraTranslatedWorldOrigin, TempSkyPlanetData, SkyViewLutReferential);
-		ReflectionViewParameters.SkyPlanetTranslatedWorldCenterAndViewHeight = TempSkyPlanetData;
-		ReflectionViewParameters.SkyCameraTranslatedWorldOrigin = SkyCameraTranslatedWorldOrigin;
-		ReflectionViewParameters.SkyViewLutReferential = SkyViewLutReferential;
+
+		// Here, it would be convenient to create a new ViewUB using Scene->SkyLight->CapturePosition,
+		// while setting up correct view parameters such as SkyPlanetTranslatedWorldCenterAndViewHeight, SkyViewLutReferential, etc.
+		// But this is expenssive and can increase the cost of that function by 3 on low end platform.
+		// As an optimisation to avoid creating a new View UB, shaders will fetch from constants what they need.
+		FSkyAtmosphereRealTimeReflectionLUTParameters Parameters;
+		Parameters.SkyViewLutReferential = SkyViewLutReferential;
+		Parameters.SkyPlanetTranslatedWorldCenterAndViewHeight = TempSkyPlanetData;
+		Parameters.SkyCameraTranslatedWorldOrigin = SkyCameraTranslatedWorldOrigin;
+		Parameters.RealTimeReflectionCapture = 0.0f;
+		TUniformBufferRef<FSkyAtmosphereRealTimeReflectionLUTParameters> SkyAtmosphereRealTimeReflectionLUTUB = TUniformBufferRef<FSkyAtmosphereRealTimeReflectionLUTParameters>::CreateUniformBufferImmediate(Parameters, UniformBuffer_SingleFrame);
 
 		FVolumetricCloudRenderSceneInfo* CloudInfo = Scene->GetVolumetricCloudSceneInfo();
 		FCloudShadowAOData CloudShadowAOData;
-		GetCloudShadowAOData(CloudInfo, View, GraphBuilder, CloudShadowAOData);
+		GetCloudShadowAOData(CloudInfo, ViewForReflection, GraphBuilder, CloudShadowAOData);
 		const bool bShouldSampleCloudShadow = CloudShadowAOData.bShouldSampleCloudShadow && (CloudShadowOnAtmosphereStrength0 > 0.0f || CloudShadowOnAtmosphereStrength1 > 0.0f);
 
-		TUniformBufferRef<FViewUniformShaderParameters> ReflectionViewUniformBuffer = TUniformBufferRef<FViewUniformShaderParameters>::CreateUniformBufferImmediate(ReflectionViewParameters, UniformBuffer_SingleFrame);
-
 		if (bRealTimeReflectionCaptureSkyAtmosphereViewLutTexture)
 		{
-			FIntVector SkyViewLutSize = View.SkyAtmosphereViewLutTexture->GetDesc().GetSize();
+			FIntVector SkyViewLutSize = ViewForReflection.SkyAtmosphereViewLutTexture->GetDesc().GetSize();
 			FRDGTextureRef RealTimeReflectionCaptureSkyAtmosphereViewLutTexture = GraphBuilder.CreateTexture(
 				FRDGTextureDesc::Create2D(FIntPoint(SkyViewLutSize.X, SkyViewLutSize.Y), GetSkyLutTextureFormat(Scene->GetFeatureLevel()), FClearValueBinding::None,
 					TexCreate_ShaderResource | TexCreate_UAV), TEXT("SkyAtmosphere.RealTimeSkyCapViewLut"));
@@ -1521,12 +1533,14 @@ void FSceneRenderer::RenderSkyAtmosphereLookUpTables(FRDGBuilder& GraphBuilder,
 			PermutationVector.Set<FSecondAtmosphereLight>(bSecondAtmosphereLightEnabled);
 			PermutationVector.Set<FSampleOpaqueShadow>(false); // bShouldSampleOpaqueShadow); Off for now to not have to generate Light0Shadow and Light1Shadow
 			PermutationVector.Set<FSampleCloudShadow>(bShouldSampleCloudShadow);
-			TShaderMapRef<FRenderSkyViewLutCS> ComputeShader(View.ShaderMap, PermutationVector);
+			PermutationVector.Set<FRealTimeReflectionLUT>(true);
+			TShaderMapRef<FRenderSkyViewLutCS> ComputeShader(ViewForReflection.ShaderMap, PermutationVector);
 
 			FRenderSkyViewLutCS::FParameters* PassParameters = GraphBuilder.AllocParameters<FRenderSkyViewLutCS::FParameters>();
 			PassParameters->Atmosphere = Scene->GetSkyAtmosphereSceneInfo()->GetAtmosphereUniformBuffer();
 			PassParameters->SkyAtmosphere = SkyInfo.GetInternalCommonParametersUniformBuffer();
-			PassParameters->ViewUniformBuffer = ReflectionViewUniformBuffer;
+			PassParameters->ViewUniformBuffer = ViewForReflection.ViewUniformBuffer;
+			PassParameters->SkyAtmosphereRealTimeReflectionLUTUniformBuffer = SkyAtmosphereRealTimeReflectionLUTUB;
 			PassParameters->TransmittanceLutTextureSampler = SamplerLinearClamp;
 			PassParameters->MultiScatteredLuminanceLutTextureSampler = SamplerLinearClamp;
 			PassParameters->VolumetricCloudShadowMapTexture0Sampler = SamplerLinearClamp;
@@ -1556,7 +1570,7 @@ void FSceneRenderer::RenderSkyAtmosphereLookUpTables(FRDGBuilder& GraphBuilder,
 
 		if (bRealTimeReflectionCapture360APLutTexture)
 		{
-			FIntVector CameraAPLutSize = View.SkyAtmosphereCameraAerialPerspectiveVolume->GetDesc().GetSize();
+			FIntVector CameraAPLutSize = ViewForReflection.SkyAtmosphereCameraAerialPerspectiveVolume->GetDesc().GetSize();
 			EPixelFormat TextureAerialLUTFormat = (CVarSkyAtmosphereLUT32.GetValueOnAnyThread() != 0) ? PF_A32B32G32R32F : PF_FloatRGBA;
 			FRDGTextureRef RealTimeReflectionCaptureCamera360APLutTexture = GraphBuilder.CreateTexture(
 				FRDGTextureDesc::Create3D(CameraAPLutSize, TextureAerialLUTFormat, FClearValueBinding::None, 
@@ -1569,12 +1583,14 @@ void FSceneRenderer::RenderSkyAtmosphereLookUpTables(FRDGBuilder& GraphBuilder,
 			PermutationVector.Set<FSampleOpaqueShadow>(false); // bShouldSampleOpaqueShadow); Off for now to not have to generate Light0Shadow and Light1Shadow
 			PermutationVector.Set<FSampleCloudShadow>(bShouldSampleCloudShadow);
 			PermutationVector.Set<FSeparateMieAndRayleighScattering>(false);	// Not for reflections for now
-			TShaderMapRef<FRenderCameraAerialPerspectiveVolumeCS> ComputeShader(View.ShaderMap, PermutationVector);
+			PermutationVector.Set<FRealTimeReflectionLUT>(true);
+			TShaderMapRef<FRenderCameraAerialPerspectiveVolumeCS> ComputeShader(ViewForReflection.ShaderMap, PermutationVector);
 
 			FRenderCameraAerialPerspectiveVolumeCS::FParameters* PassParameters = GraphBuilder.AllocParameters<FRenderCameraAerialPerspectiveVolumeCS::FParameters>();
 			PassParameters->Atmosphere = Scene->GetSkyAtmosphereSceneInfo()->GetAtmosphereUniformBuffer();
 			PassParameters->SkyAtmosphere = SkyInfo.GetInternalCommonParametersUniformBuffer();
-			PassParameters->ViewUniformBuffer = ReflectionViewUniformBuffer;
+			PassParameters->ViewUniformBuffer = ViewForReflection.ViewUniformBuffer;
+			PassParameters->SkyAtmosphereRealTimeReflectionLUTUniformBuffer = SkyAtmosphereRealTimeReflectionLUTUB;
 			PassParameters->TransmittanceLutTextureSampler = SamplerLinearClamp;
 			PassParameters->MultiScatteredLuminanceLutTextureSampler = SamplerLinearClamp;
 			PassParameters->VolumetricCloudShadowMapTexture0Sampler = SamplerLinearClamp;
@@ -1646,12 +1662,14 @@ void FSceneRenderer::RenderSkyAtmosphereLookUpTables(FRDGBuilder& GraphBuilder,
 			PermutationVector.Set<FSecondAtmosphereLight>(bSecondAtmosphereLightEnabled);
 			PermutationVector.Set<FSampleOpaqueShadow>(bShouldSampleOpaqueShadow);
 			PermutationVector.Set<FSampleCloudShadow>(bShouldSampleCloudShadow);
+			PermutationVector.Set<FRealTimeReflectionLUT>(false);
 			TShaderMapRef<FRenderSkyViewLutCS> ComputeShader(View.ShaderMap, PermutationVector);
 
 			FRenderSkyViewLutCS::FParameters * PassParameters = GraphBuilder.AllocParameters<FRenderSkyViewLutCS::FParameters>();
 			PassParameters->Atmosphere = Scene->GetSkyAtmosphereSceneInfo()->GetAtmosphereUniformBuffer();
 			PassParameters->SkyAtmosphere = SkyInfo.GetInternalCommonParametersUniformBuffer();
 			PassParameters->ViewUniformBuffer = View.ViewUniformBuffer;
+			PassParameters->SkyAtmosphereRealTimeReflectionLUTUniformBuffer = GetDefaultSkyAtmosphereRealTimeReflectionLUTUniformBuffer();
 			PassParameters->TransmittanceLutTextureSampler = SamplerLinearClamp;
 			PassParameters->MultiScatteredLuminanceLutTextureSampler = SamplerLinearClamp;
 			PassParameters->VolumetricCloudShadowMapTexture0Sampler = SamplerLinearClamp;
@@ -1690,12 +1708,14 @@ void FSceneRenderer::RenderSkyAtmosphereLookUpTables(FRDGBuilder& GraphBuilder,
 			PermutationVector.Set<FSampleOpaqueShadow>(bShouldSampleOpaqueShadow);
 			PermutationVector.Set<FSampleCloudShadow>(bShouldSampleCloudShadow);
 			PermutationVector.Set<FSeparateMieAndRayleighScattering>(bSeparatedAtmosphereMieRayLeigh);
+			PermutationVector.Set<FRealTimeReflectionLUT>(false);
 			TShaderMapRef<FRenderCameraAerialPerspectiveVolumeCS> ComputeShader(View.ShaderMap, PermutationVector);
 
 			FRenderCameraAerialPerspectiveVolumeCS::FParameters * PassParameters = GraphBuilder.AllocParameters<FRenderCameraAerialPerspectiveVolumeCS::FParameters>();
 			PassParameters->Atmosphere = Scene->GetSkyAtmosphereSceneInfo()->GetAtmosphereUniformBuffer();
 			PassParameters->SkyAtmosphere = SkyInfo.GetInternalCommonParametersUniformBuffer();
 			PassParameters->ViewUniformBuffer = View.ViewUniformBuffer;
+			PassParameters->SkyAtmosphereRealTimeReflectionLUTUniformBuffer = GetDefaultSkyAtmosphereRealTimeReflectionLUTUniformBuffer();
 			PassParameters->TransmittanceLutTextureSampler = SamplerLinearClamp;
 			PassParameters->MultiScatteredLuminanceLutTextureSampler = SamplerLinearClamp;
 			PassParameters->VolumetricCloudShadowMapTexture0Sampler = SamplerLinearClamp;
