[HWRT] Improvements to dynamic ray tracing instance gathering.

Context:
- currently FLandscapeComponentSceneProxy::GetDynamicRayTracingInstances(...) can't run on parallel tasks because it depends on PerViewStateCachedSectionLODValues.
- PerViewStateCachedSectionLODValues is calculated during LandscapeSetupTask which is waited on inside PreVisibilityFrameSetup(...).
- on rasterization side, we only start gathering dynamic mesh elements after PreVisibilityFrameSetup(...), and doing the same for dynamic ray tracing instances will prevent the issue with FLandscapeComponentSceneProxy.

Changes:
- Only start gathering dynamic ray tracing instances after PreVisibilityFrameSetup(...) finishes.
- Re-enabled bSupportsParallelGDRTI for FLandscapeComponentSceneProxy.

diff --git a/Engine/Source/Runtime/Landscape/Private/LandscapeRender.cpp b/Engine/Source/Runtime/Landscape/Private/LandscapeRender.cpp
--- a/Engine/Source/Runtime/Landscape/Private/LandscapeRender.cpp
+++ b/Engine/Source/Runtime/Landscape/Private/LandscapeRender.cpp
@@ -1422,9 +1422,6 @@ FLandscapeComponentSceneProxy::FLandscapeComponentSceneProxy(ULandscapeComponent
 	// Enabled just so the GVarDumpLandscapeLODsCurrentFrame debug feature only runs once per render proxy in GetDynamicMeshElements
 	bSinglePassGDME = true;
 
-	// Disable GDRTI since there's seems to be a race condition in FLandscapeRenderSystem::GetSectionLODValue(...) otherwise
-	bSupportsParallelGDRTI = false;
-
 	bNaniteActive = InComponent->IsNaniteActive();
 	bUsesLandscapeCulling = (!bNaniteActive && Culling::UseCulling(GetScene().GetShaderPlatform()));
 
diff --git a/Engine/Source/Runtime/Renderer/Private/RayTracing/RayTracing.cpp b/Engine/Source/Runtime/Renderer/Private/RayTracing/RayTracing.cpp
--- a/Engine/Source/Runtime/Renderer/Private/RayTracing/RayTracing.cpp
+++ b/Engine/Source/Runtime/Renderer/Private/RayTracing/RayTracing.cpp
@@ -316,6 +316,9 @@ namespace RayTracing
 		// This task must complete before PostRenderAllViewports().
 		UE::Tasks::FTask AddUsedStreamingHandlesTask;
 
+		UE::Tasks::FTaskEvent GatherDynamicRayTracingInstancesPrerequisites{ UE_SOURCE_LOCATION };
+		bool bGatherDynamicRayTracingInstancesPrerequisitesTriggered = false;
+
 		// Array of primitives that should update their cached ray tracing instances via FPrimitiveSceneInfo::UpdateCachedRaytracingData()
 		TArray<FPrimitiveSceneInfo*> DirtyCachedRayTracingPrimitives; // TODO: remove this since it seems to be transient
 
@@ -1759,14 +1762,15 @@ namespace RayTracing
 				GatherRelevantStaticPrimitives(TaskData, LODScaleCVarValue, ForcedLODLevel, bUsingReferenceBasedResidency);
 			}, TaskData.GatherRelevantPrimitivesTask);
 
-		UE::Tasks::FTask GPUSkinCacheTask = TaskData.Scene.GetGPUSkinCacheTask();
+		TaskData.GatherDynamicRayTracingInstancesPrerequisites.AddPrerequisites(TaskData.GatherRelevantPrimitivesTask);
+		TaskData.GatherDynamicRayTracingInstancesPrerequisites.AddPrerequisites(TaskData.Scene.GetGPUSkinCacheTask());
 
 		// TODO: Could gather dynamic ray tracing instances using multiple tasks / FDynamicRayTracingInstancesContext
 		TaskData.GatherDynamicRayTracingInstancesTask = GraphBuilder.AddSetupTask([&TaskData]
 			{
 				FTaskTagScope Scope(ETaskTag::EParallelRenderingThread);
 				TaskData.DynamicRayTracingInstancesContext.GatherDynamicRayTracingInstances(TaskData.DynamicPrimitivesIndices);
-			}, MakeArrayView({ TaskData.GatherRelevantPrimitivesTask, GPUSkinCacheTask }));
+			}, TaskData.GatherDynamicRayTracingInstancesPrerequisites);
 
 
 		TaskData.AddUsedStreamingHandlesTask = GraphBuilder.AddSetupTask([&TaskData, LODScaleCVarValue, ForcedLODLevel, bUsingReferenceBasedResidency]
@@ -1784,6 +1788,15 @@ namespace RayTracing
 		TaskData.bValid = true;
 	}
 
+	void BeginGatherDynamicRayTracingInstances(FGatherInstancesTaskData& TaskData)
+	{
+		if (!TaskData.bGatherDynamicRayTracingInstancesPrerequisitesTriggered)
+		{
+			TaskData.GatherDynamicRayTracingInstancesPrerequisites.Trigger();
+			TaskData.bGatherDynamicRayTracingInstancesPrerequisitesTriggered = true;
+		}
+	}
+
 	bool FinishGatherInstances(
 		FRDGBuilder& GraphBuilder,
 		FGatherInstancesTaskData& TaskData,
diff --git a/Engine/Source/Runtime/Renderer/Private/RayTracing/RayTracing.h b/Engine/Source/Runtime/Renderer/Private/RayTracing/RayTracing.h
--- a/Engine/Source/Runtime/Renderer/Private/RayTracing/RayTracing.h
+++ b/Engine/Source/Runtime/Renderer/Private/RayTracing/RayTracing.h
@@ -54,6 +54,8 @@ namespace RayTracing
 	
 	void BeginGatherInstances(FRDGBuilder& GraphBuilder, FGatherInstancesTaskData& TaskData, UE::Tasks::FTask FrustumCullTask);
 
+	void BeginGatherDynamicRayTracingInstances(FGatherInstancesTaskData& TaskData);
+
 	// Fills RayTracingScene instance list for the given View and adds relevant ray tracing data to the view. Does not reset previous scene contents.
 	// This function must run on render thread
 	bool FinishGatherInstances(
diff --git a/Engine/Source/Runtime/Renderer/Private/SceneVisibility.cpp b/Engine/Source/Runtime/Renderer/Private/SceneVisibility.cpp
--- a/Engine/Source/Runtime/Renderer/Private/SceneVisibility.cpp
+++ b/Engine/Source/Runtime/Renderer/Private/SceneVisibility.cpp
@@ -47,6 +47,7 @@
 #include "InstanceCulling/InstanceCullingOcclusionQuery.h"
 #include "PostProcess/TemporalAA.h"
 #include "RayTracing/RayTracingInstanceCulling.h"
+#include "RayTracing/RayTracing.h"
 #include "HeterogeneousVolumes/HeterogeneousVolumes.h"
 #include "RendererModule.h"
 #include "SceneViewExtension.h"
@@ -5682,6 +5683,14 @@ void FDeferredShadingSceneRenderer::BeginInitViews(
 	// Start processing dynamic mesh elements tasks early enough to overlap with shadows and GPU scene update.
 	TaskDatas.VisibilityTaskData->StartGatherDynamicMeshElements();
 
+#if RHI_RAYTRACING
+	// Start processing dynamic ray tracing instances early enough to overlap with shadows and GPU scene update.
+	if (TaskDatas.RayTracingGatherInstances != nullptr)
+	{
+		RayTracing::BeginGatherDynamicRayTracingInstances(*TaskDatas.RayTracingGatherInstances);
+	}
+#endif
+
 	if (!ViewFamily.EngineShowFlags.HitProxies)
 	{
 		TaskDatas.Decals = FDecalVisibilityTaskData::Launch(GraphBuilder, *Scene, Views);
