Post Processing: Continued separation of lens flare and bloom on the mobile side.

* Base CL 42137513 missed the mobile renderer side.
* We simply re-use the bUseBloom / MOBILE_USEBLOOM code paths, only with bloom & lens flares separately enabled.

diff --git a/Engine/Source/Runtime/Renderer/Private/PostProcess/PostProcessing.cpp b/Engine/Source/Runtime/Renderer/Private/PostProcess/PostProcessing.cpp
--- a/Engine/Source/Runtime/Renderer/Private/PostProcess/PostProcessing.cpp
+++ b/Engine/Source/Runtime/Renderer/Private/PostProcess/PostProcessing.cpp
@@ -2643,8 +2643,9 @@ void AddMobilePostProcessingPasses(FRDGBuilder& GraphBuilder, FScene* Scene, con
 	if (bIsPostProcessingEnabled)
 	{
 		bool bUseSun = View.MobileLightShaft.IsSet();
-			
-		bool bUseBloom = View.FinalPostProcessSettings.BloomIntensity > 0.0f;
+		
+		const bool bBloomEnabled = View.FinalPostProcessSettings.BloomIntensity > 0.0f;
+		const bool bUseBloom = bBloomEnabled || IsLensFlaresEnabled(View);
 
 		bool bUseBasicEyeAdaptation = bUseEyeAdaptation && (AutoExposureMethod == EAutoExposureMethod::AEM_Basic);
 		bool bUseHistogramEyeAdaptation = bUseEyeAdaptation && (AutoExposureMethod == EAutoExposureMethod::AEM_Histogram) &&
@@ -2891,66 +2892,72 @@ void AddMobilePostProcessingPasses(FRDGBuilder& GraphBuilder, FScene* Scene, con
 				return AddMobileBloomDownPass(GraphBuilder, View, BloomDownInputs);
 			};
 
-			const float BloomDownScale = 0.66f * 4.0f;
-			const int MaxPasses = 6;
+			constexpr float BloomDownScale = 0.66f * 4.0f;
+			constexpr int32 MaxPasses = 6;
 			const EBloomQuality BloomQuality = GetBloomQuality();
-			const int NumDownsamplePasses = BloomQuality == EBloomQuality::Q1 ? 4 : BloomQuality == EBloomQuality::Q2 ? 5 : 6;
+			
+			static_assert(static_cast<int32>(ELensFlareQuality::MAX) == 4);
+			const int NumDownsamplePasses = (BloomQuality == EBloomQuality::Q1 || !bBloomEnabled) ? 4 : BloomQuality == EBloomQuality::Q2 ? 5 : 6;
+
 			FScreenPassTexture PostProcessDownsample_Bloom[MaxPasses];
 
 			for (int32 i = 0; i < NumDownsamplePasses; ++i)
 			{
 				PostProcessDownsample_Bloom[i] = AddBloomDownPass(i == 0 ? BloomSetupOutputs.Bloom : PostProcessDownsample_Bloom[i - 1], BloomDownScale);
 			}
-
-			const FFinalPostProcessSettings& Settings = View.FinalPostProcessSettings;
-
-			auto AddBloomUpPass = [&GraphBuilder, &View](FScreenPassTexture& BloomUpSourceA, FScreenPassTexture& BloomUpSourceB, float BloomSourceScale, const FVector4f& TintA, const FVector4f& TintB)
+			
+			if (bBloomEnabled)
 			{
-				FMobileBloomUpInputs BloomUpInputs;
-				BloomUpInputs.BloomUpSourceA = BloomUpSourceA;
-				BloomUpInputs.BloomUpSourceB = BloomUpSourceB;
-				BloomUpInputs.ScaleAB = FVector2D(BloomSourceScale, BloomSourceScale);
-				BloomUpInputs.TintA = TintA;
-				BloomUpInputs.TintB = TintB;
-
-				return AddMobileBloomUpPass(GraphBuilder, View, BloomUpInputs);
-			};
+				const FFinalPostProcessSettings& Settings = View.FinalPostProcessSettings;
 
-			const float BloomUpScale = 0.66f * 2.0f;
+				auto AddBloomUpPass = [&GraphBuilder, &View](FScreenPassTexture& BloomUpSourceA, FScreenPassTexture& BloomUpSourceB, float BloomSourceScale, const FVector4f& TintA, const FVector4f& TintB)
+					{
+						FMobileBloomUpInputs BloomUpInputs;
+						BloomUpInputs.BloomUpSourceA = BloomUpSourceA;
+						BloomUpInputs.BloomUpSourceB = BloomUpSourceB;
+						BloomUpInputs.ScaleAB = FVector2D(BloomSourceScale, BloomSourceScale);
+						BloomUpInputs.TintA = TintA;
+						BloomUpInputs.TintB = TintB;
 
-			const FLinearColor Tint[]{ Settings.Bloom1Tint, Settings.Bloom2Tint, Settings.Bloom3Tint, Settings.Bloom4Tint, Settings.Bloom5Tint, Settings.Bloom6Tint };
+						return AddMobileBloomUpPass(GraphBuilder, View, BloomUpInputs);
+					};
 
-			// Upsample by 2
-			{
-				const int IndexA = NumDownsamplePasses - 2;
-				const int IndexB = NumDownsamplePasses - 1;
-				FVector4f TintA{ Tint[IndexA], 0.0f };
-				FVector4f TintB{ Tint[IndexB], 0.0f };
-				TintA *= Settings.BloomIntensity;
-				TintB *= Settings.BloomIntensity;
+				const float BloomUpScale = 0.66f * 2.0f;
 
-				BloomUpOutputs = AddBloomUpPass(PostProcessDownsample_Bloom[IndexA], PostProcessDownsample_Bloom[IndexB], BloomUpScale, TintA, TintB);
-			}
+				const FLinearColor Tint[]{ Settings.Bloom1Tint, Settings.Bloom2Tint, Settings.Bloom3Tint, Settings.Bloom4Tint, Settings.Bloom5Tint, Settings.Bloom6Tint };
 
-			for (int Index = NumDownsamplePasses - 3; Index > 0; --Index)
-			{
 				// Upsample by 2
-				FVector4f TintA{ Tint[Index], 0.0f };
-				TintA *= Settings.BloomIntensity;
-				FVector4f TintB = FVector4f(1.0f, 1.0f, 1.0f, 0.0f);
+				{
+					const int IndexA = NumDownsamplePasses - 2;
+					const int IndexB = NumDownsamplePasses - 1;
+					FVector4f TintA{ Tint[IndexA], 0.0f };
+					FVector4f TintB{ Tint[IndexB], 0.0f };
+					TintA *= Settings.BloomIntensity;
+					TintB *= Settings.BloomIntensity;
+
+					BloomUpOutputs = AddBloomUpPass(PostProcessDownsample_Bloom[IndexA], PostProcessDownsample_Bloom[IndexB], BloomUpScale, TintA, TintB);
+				}
 
-				BloomUpOutputs = AddBloomUpPass(PostProcessDownsample_Bloom[Index], BloomUpOutputs, BloomUpScale, TintA, TintB);
-			}
+				for (int Index = NumDownsamplePasses - 3; Index > 0; --Index)
+				{
+					// Upsample by 2
+					FVector4f TintA{ Tint[Index], 0.0f };
+					TintA *= Settings.BloomIntensity;
+					FVector4f TintB = FVector4f(1.0f, 1.0f, 1.0f, 0.0f);
 
-			// Upsample by 2
-			{
-				FVector4f TintA = FVector4f(Settings.Bloom2Tint.R, Settings.Bloom2Tint.G, Settings.Bloom2Tint.B, 0.0f);
-				TintA *= Settings.BloomIntensity;
-				// Scaling Bloom2 by extra factor to match filter area difference between PC default and mobile.
-				TintA *= 0.5;
-				FVector4f TintB = FVector4f(1.0f, 1.0f, 1.0f, 0.0f);
+					BloomUpOutputs = AddBloomUpPass(PostProcessDownsample_Bloom[Index], BloomUpOutputs, BloomUpScale, TintA, TintB);
+				}
+
+				// Upsample by 2
+				{
+					FVector4f TintA = FVector4f(Settings.Bloom2Tint.R, Settings.Bloom2Tint.G, Settings.Bloom2Tint.B, 0.0f);
+					TintA *= Settings.BloomIntensity;
+					// Scaling Bloom2 by extra factor to match filter area difference between PC default and mobile.
+					TintA *= 0.5;
+					FVector4f TintB = FVector4f(1.0f, 1.0f, 1.0f, 0.0f);
 
-				BloomUpOutputs = AddBloomUpPass(PostProcessDownsample_Bloom[0], BloomUpOutputs, BloomUpScale, TintA, TintB);
+					BloomUpOutputs = AddBloomUpPass(PostProcessDownsample_Bloom[0], BloomUpOutputs, BloomUpScale, TintA, TintB);
+				}
 			}
 
 			if (IsLensFlaresEnabled(View))
