Post Processing: Minor cleanup of lens flare enabled state.

diff --git a/Engine/Source/Runtime/Renderer/Private/PostProcess/PostProcessing.cpp b/Engine/Source/Runtime/Renderer/Private/PostProcess/PostProcessing.cpp
--- a/Engine/Source/Runtime/Renderer/Private/PostProcess/PostProcessing.cpp
+++ b/Engine/Source/Runtime/Renderer/Private/PostProcess/PostProcessing.cpp
@@ -2645,7 +2645,8 @@ void AddMobilePostProcessingPasses(FRDGBuilder& GraphBuilder, FScene* Scene, con
 		bool bUseSun = View.MobileLightShaft.IsSet();
 		
 		const bool bBloomEnabled = View.FinalPostProcessSettings.BloomIntensity > 0.0f;
-		const bool bUseBloom = bBloomEnabled || IsLensFlaresEnabled(View);
+		const bool bLensFlareEnabled = IsLensFlaresEnabled(View);
+		const bool bUseBloom = bBloomEnabled || bLensFlareEnabled;
 
 		bool bUseBasicEyeAdaptation = bUseEyeAdaptation && (AutoExposureMethod == EAutoExposureMethod::AEM_Basic);
 		bool bUseHistogramEyeAdaptation = bUseEyeAdaptation && (AutoExposureMethod == EAutoExposureMethod::AEM_Histogram) &&
@@ -2960,7 +2961,7 @@ void AddMobilePostProcessingPasses(FRDGBuilder& GraphBuilder, FScene* Scene, con
 				}
 			}
 
-			if (IsLensFlaresEnabled(View))
+			if (bLensFlareEnabled)
 			{
 				const ELensFlareQuality LensFlareQuality = GetLensFlareQuality();
 				const uint32 LensFlareDownsampleStageIndex = static_cast<uint32>(ELensFlareQuality::MAX) - static_cast<uint32>(LensFlareQuality) - 1;
