Add PSO precaching for mobile SkyPass

Fixes check failure with DX12+ES31 in PSOCollectorStats::CheckMinimalPipelineStateInCache when validation is enabled.

diff --git a/Engine/Source/Runtime/Renderer/Private/SkyPassRendering.cpp b/Engine/Source/Runtime/Renderer/Private/SkyPassRendering.cpp
--- a/Engine/Source/Runtime/Renderer/Private/SkyPassRendering.cpp
+++ b/Engine/Source/Runtime/Renderer/Private/SkyPassRendering.cpp
@@ -126,34 +126,7 @@ bool FSkyPassMeshProcessor::Process(
 			return false;
 		}
 
-		if (SkyPassType == SPT_Default)
-		{
-			// Mask sky pixels so we can skip them when rendering per-pixel fog (USE_VERTEX_FOG is applying for on sky pixels).
-			PassDrawRenderState.SetDepthStencilState(TStaticDepthStencilState<
-				false, CF_DepthNearOrEqual,
-				true, CF_Always, SO_Keep, SO_Keep, SO_Replace,
-				false, CF_Always, SO_Keep, SO_Keep, SO_Keep,
-				0x00, STENCIL_MOBILE_SKY_MASK>::GetRHI());
-
-			PassDrawRenderState.SetStencilRef(STENCIL_MOBILE_SKY_MASK);
-		}
-		else if (SkyPassType == SPT_RealTimeCapture_DepthWrite)
-		{
-			// Capturing real time sky light, writing depth only.
-			PassDrawRenderState.SetDepthStencilState(TStaticDepthStencilState<
-				true, CF_DepthNearOrEqual>::GetRHI());
-		}
-		else if (SkyPassType == SPT_RealTimeCapture_DepthNop)
-		{
-
-			// Capturing real time sky light, no depth writes.
-			PassDrawRenderState.SetDepthStencilState(TStaticDepthStencilState<
-				false, CF_Always>::GetRHI());
-		}
-		else
-		{
-			check(false);
-		}
+		SetStateForMobile();
 		
 		TMobileBasePassShaderElementData<LightMapPolicyType> ShaderElementData(nullptr, false);
 		ShaderElementData.InitializeMeshMaterialData(ViewIfDynamicMeshCommand, PrimitiveSceneProxy, MeshBatch, StaticMeshId, false);
@@ -178,6 +151,37 @@ bool FSkyPassMeshProcessor::Process(
 	return true;
 }
 
+void FSkyPassMeshProcessor::SetStateForMobile()
+{
+	if (SkyPassType == SPT_Default)
+	{
+		// Mask sky pixels so we can skip them when rendering per-pixel fog (USE_VERTEX_FOG is applying for on sky pixels).
+		PassDrawRenderState.SetDepthStencilState(TStaticDepthStencilState<
+			false, CF_DepthNearOrEqual,
+			true, CF_Always, SO_Keep, SO_Keep, SO_Replace,
+			false, CF_Always, SO_Keep, SO_Keep, SO_Keep,
+			0x00, STENCIL_MOBILE_SKY_MASK>::GetRHI());
+
+		PassDrawRenderState.SetStencilRef(STENCIL_MOBILE_SKY_MASK);
+	}
+	else if (SkyPassType == SPT_RealTimeCapture_DepthWrite)
+	{
+		// Capturing real time sky light, writing depth only.
+		PassDrawRenderState.SetDepthStencilState(TStaticDepthStencilState<
+			true, CF_DepthNearOrEqual>::GetRHI());
+	}
+	else if (SkyPassType == SPT_RealTimeCapture_DepthNop)
+	{
+		// Capturing real time sky light, no depth writes.
+		PassDrawRenderState.SetDepthStencilState(TStaticDepthStencilState<
+			false, CF_Always>::GetRHI());
+	}
+	else
+	{
+		check(false);
+	}
+}
+
 void FSkyPassMeshProcessor::CollectPSOInitializers(const FSceneTexturesConfig& SceneTexturesConfig, const FMaterial& Material, const FPSOPrecacheVertexFactoryData& VertexFactoryData, const FPSOPrecacheParams& PreCacheParams, TArray<FPSOPrecacheData>& PSOInitializers)
 {
 	// Early out if not sky
@@ -186,43 +190,63 @@ void FSkyPassMeshProcessor::CollectPSOInitializers(const FSceneTexturesConfig& S
 		return;
 	}
 
-	// Only do deferred path for now
-	if (GetFeatureLevelShadingPath(FeatureLevel) != EShadingPath::Deferred)
+	typedef FUniformLightMapPolicy LightMapPolicyType;
+
+	if (GetFeatureLevelShadingPath(FeatureLevel) == EShadingPath::Deferred)
 	{
-		return;
+		TMeshProcessorShaders<
+			TBasePassVertexShaderPolicyParamType<LightMapPolicyType>,
+			TBasePassPixelShaderPolicyParamType<LightMapPolicyType>> SkyPassShaders;
+
+		FUniformLightMapPolicy NoLightmapPolicy(LMP_NO_LIGHTMAP);
+		const bool bRenderSkylight = false;
+		if (GetBasePassShaders<LightMapPolicyType>(
+			Material,
+			VertexFactoryData.VertexFactoryType,
+			NoLightmapPolicy,
+			FeatureLevel,
+			bRenderSkylight,
+			false, // 128bit
+			false, // bIsDebug
+			GBL_Default,
+			&SkyPassShaders.VertexShader,
+			&SkyPassShaders.PixelShader
+		))
+		{
+			CollectPSOInitializersInternal(SkyPassShaders, SceneTexturesConfig, Material, VertexFactoryData, PreCacheParams, PSOInitializers);
+		}
+	}
+	else
+	{
+		TMeshProcessorShaders<
+			TMobileBasePassVSPolicyParamType<LightMapPolicyType>,
+			TMobileBasePassPSPolicyParamType<LightMapPolicyType>> SkyPassShaders;
+
+		if (MobileBasePass::GetShaders(
+			LMP_NO_LIGHTMAP,
+			EMobileLocalLightSetting::LOCAL_LIGHTS_DISABLED,
+			Material,
+			VertexFactoryData.VertexFactoryType,
+			SkyPassShaders.VertexShader,
+			SkyPassShaders.PixelShader
+		))
+		{
+			SetStateForMobile();
+			CollectPSOInitializersInternal(SkyPassShaders, SceneTexturesConfig, Material, VertexFactoryData, PreCacheParams, PSOInitializers);
+		}
 	}
+}
 
+template <typename T>
+void FSkyPassMeshProcessor::CollectPSOInitializersInternal(T& SkyPassShaders, const FSceneTexturesConfig& SceneTexturesConfig, const FMaterial& Material, const FPSOPrecacheVertexFactoryData& VertexFactoryData, const FPSOPrecacheParams& PreCacheParams, TArray<FPSOPrecacheData>& PSOInitializers)
+{
 	const FMeshDrawingPolicyOverrideSettings OverrideSettings = ComputeMeshOverrideSettings(PreCacheParams);
 	const ERasterizerFillMode MeshFillMode = ComputeMeshFillMode(Material, OverrideSettings);
 	const ERasterizerCullMode MeshCullMode = ComputeMeshCullMode(Material, OverrideSettings);
 
-	typedef FUniformLightMapPolicy LightMapPolicyType;
-	FUniformLightMapPolicy NoLightmapPolicy(LMP_NO_LIGHTMAP);
-
-	TMeshProcessorShaders<
-		TBasePassVertexShaderPolicyParamType<LightMapPolicyType>,
-		TBasePassPixelShaderPolicyParamType<LightMapPolicyType>> SkyPassShaders;
-
-	const bool bRenderSkylight = false;
-	if (!GetBasePassShaders<LightMapPolicyType>(
-		Material,
-		VertexFactoryData.VertexFactoryType,
-		NoLightmapPolicy,
-		FeatureLevel,
-		bRenderSkylight,
-		false, // 128bit
-		false, // bIsDebug
-		GBL_Default,
-		&SkyPassShaders.VertexShader,
-		&SkyPassShaders.PixelShader
-		))
-	{
-		return;
-	}
-	
 	FGraphicsPipelineRenderTargetsInfo RenderTargetsInfo;
 	SetupGBufferRenderTargetInfo(SceneTexturesConfig, RenderTargetsInfo, true /*bSetupDepthStencil*/);
-	
+
 	FBasePassMeshProcessor::AddBasePassGraphicsPipelineStateInitializer(
 		FeatureLevel,
 		VertexFactoryData,
@@ -287,4 +311,4 @@ FMeshPassProcessor* CreateSkyPassProcessor(ERHIFeatureLevel::Type FeatureLevel,
 }
 
 REGISTER_MESHPASSPROCESSOR_AND_PSOCOLLECTOR(SkyPass, CreateSkyPassProcessor, EShadingPath::Deferred, EMeshPass::SkyPass, EMeshPassFlags::MainView);
-FRegisterPassProcessorCreateFunction RegisterMobileSkyPass(&CreateSkyPassProcessor, EShadingPath::Mobile, EMeshPass::SkyPass, EMeshPassFlags::MainView);
+REGISTER_MESHPASSPROCESSOR_AND_PSOCOLLECTOR(MobileSkyPass, CreateSkyPassProcessor, EShadingPath::Mobile, EMeshPass::SkyPass, EMeshPassFlags::MainView);
diff --git a/Engine/Source/Runtime/Renderer/Private/SkyPassRendering.h b/Engine/Source/Runtime/Renderer/Private/SkyPassRendering.h
--- a/Engine/Source/Runtime/Renderer/Private/SkyPassRendering.h
+++ b/Engine/Source/Runtime/Renderer/Private/SkyPassRendering.h
@@ -52,6 +52,11 @@ private:
 		const FMaterial& RESTRICT MaterialResource,
 		ERasterizerFillMode MeshFillMode,
 		ERasterizerCullMode MeshCullMode);
+
+	void SetStateForMobile();
+
+	template <typename T>
+	void CollectPSOInitializersInternal(T& SkyPassShaders, const FSceneTexturesConfig& SceneTexturesConfig, const FMaterial& Material, const FPSOPrecacheVertexFactoryData& VertexFactoryData, const FPSOPrecacheParams& PreCacheParams, TArray<FPSOPrecacheData>& PSOInitializers);
 };
 
 
