[Lit Translucency] Expose cvar to control history weight when using temporal filtering.

diff --git a/Engine/Shaders/Private/TranslucentLightingShaders.usf b/Engine/Shaders/Private/TranslucentLightingShaders.usf
--- a/Engine/Shaders/Private/TranslucentLightingShaders.usf
+++ b/Engine/Shaders/Private/TranslucentLightingShaders.usf
@@ -167,6 +167,8 @@ float4 PrevTranslucencyLightingVolumeInvSize;
 float3 HistoryTextureBilinearUVMin;
 float3 HistoryTextureBilinearUVMax;
 
+float HistoryWeight;
+
 float GetBorderLerpFactor(float3 UVW)
 {
 	// Larger values result in a shorter transition distance
@@ -213,7 +215,7 @@ void FilterTranslucentVolumeCS(uint3 DispatchThreadId : SV_DispatchThreadID)
 	float4 HistoryTextureValue0 = HistoryAmbient[HistoryCoord];
 	float4 HistoryTextureValue1 = HistoryDirectional[HistoryCoord];
 
-	float HistoryAlpha = 0.9f;
+	float HistoryAlpha = HistoryWeight;
 
 	if (any(HistoryCoord < 0) || any(HistoryCoord >= VolumeSize))
 	{
diff --git a/Engine/Source/Runtime/Renderer/Private/TranslucentLighting.cpp b/Engine/Source/Runtime/Renderer/Private/TranslucentLighting.cpp
--- a/Engine/Source/Runtime/Renderer/Private/TranslucentLighting.cpp
+++ b/Engine/Source/Runtime/Renderer/Private/TranslucentLighting.cpp
@@ -135,6 +135,13 @@ static TAutoConsoleVariable<bool> CVarTranslucencyLightingVolumeTemporal(
 	ECVF_Scalability | ECVF_RenderThreadSafe
 );
 
+static TAutoConsoleVariable<float> CVarTranslucencyLightingVolumeHistoryWeight(
+	TEXT("r.TranslucencyLightingVolume.Temporal.HistoryWeight"),
+	0.9,
+	TEXT("How much the history value should be weighted each frame.  This is a tradeoff between visible jittering and responsiveness."),
+	ECVF_Scalability | ECVF_RenderThreadSafe
+);
+
 static TAutoConsoleVariable<bool> CVarTranslucencyLightingVolumeMarkVoxelsSupported(
 	TEXT("r.TranslucencyLightingVolume.MarkVoxels.Supported"),
 	true,
@@ -835,6 +842,7 @@ public:
 		SHADER_PARAMETER(FVector3f, HistoryTextureBilinearUVMin)
 		SHADER_PARAMETER(FVector3f, HistoryTextureBilinearUVMax)
 
+		SHADER_PARAMETER(float, HistoryWeight)
 	END_SHADER_PARAMETER_STRUCT()
 
 	class FUseTemporalReprojection : SHADER_PERMUTATION_BOOL("USE_TEMPORAL_REPROJECTION");
@@ -2250,6 +2258,8 @@ void FilterTranslucencyLightingVolume(
 
 				PassParameters->HistoryMark = HistoryMark;
 
+				PassParameters->HistoryWeight = CVarTranslucencyLightingVolumeHistoryWeight.GetValueOnRenderThread();;
+
 				FFilterTranslucentVolumeCS::FPermutationDomain PermutationVector;
 				PermutationVector.Set<FFilterTranslucentVolumeCS::FUseTemporalReprojection>(bTemporal);
 				PermutationVector.Set<FFilterTranslucentVolumeCS::FCheckHistoryMark>(HistoryMark != nullptr);
