Fix Velocitypass not taking the right stencil value

diff --git a/Engine/Source/Runtime/Renderer/Private/VelocityRendering.cpp b/Engine/Source/Runtime/Renderer/Private/VelocityRendering.cpp
--- a/Engine/Source/Runtime/Renderer/Private/VelocityRendering.cpp
+++ b/Engine/Source/Runtime/Renderer/Private/VelocityRendering.cpp
@@ -806,7 +806,8 @@ bool FVelocityMeshProcessor::Process(
 	{
 		extern void SetMobileBasePassDepthState(FMeshPassProcessorRenderState& DrawRenderState, const FPrimitiveSceneProxy* PrimitiveSceneProxy, const FMaterial& Material, FMaterialShadingModelField ShadingModels, bool bUsesDeferredShading);
 
-		FMaterialShadingModelField ShadingModels = MaterialResource.GetShadingModels();
+		// *Don't* get shading models from MaterialResource since it's for a default material
+		FMaterialShadingModelField ShadingModels = MeshBatch.MaterialRenderProxy->GetIncompleteMaterialWithFallback(ERHIFeatureLevel::ES3_1).GetShadingModels();
 		bool bUsesDeferredShading = IsMobileDeferredShadingEnabled(GetFeatureLevelShaderPlatform(FeatureLevel));
 		SetMobileBasePassDepthState(PassDrawRenderState, PrimitiveSceneProxy, MaterialResource, ShadingModels, bUsesDeferredShading);
 	}
