Merging //UE5/Dev-ParallelRendering to Main (//UE5/Main)
RHI CreateTexture refactor
- Textures are now created using FRHITextureInitializer objects via RHICreateTextureInitializer using the FRHITextureCreateDesc struct.
- FRHITextureCreateDesc::InitAction specifies if the texture should be initialized with default data, filled with resource bulk data, or use data provided to FRHITextureInitializer.
- FRHITextureInitializer objects can carry platform-agnostic pointers for initializing the new texture. This allows many unified memory platforms to remove extra copies when creating textures.
- FRHITexture objects are created via FRHITextureInitializer::Finalize() which internally will queue uploads of pending data, depending on the platform.
- Converted a number of locations that did CreateTexture->LockTexture->UnlockTexture to use the Texture Initializers.

diff --git a/Engine/Plugins/Experimental/AVCodecs/AVCodecsCore/Source/AVCodecsCoreRHI/Private/Video/Resources/VideoResourceRHI.cpp b/Engine/Plugins/Experimental/AVCodecs/AVCodecsCore/Source/AVCodecsCoreRHI/Private/Video/Resources/VideoResourceRHI.cpp
--- a/Engine/Plugins/Experimental/AVCodecs/AVCodecsCore/Source/AVCodecsCoreRHI/Private/Video/Resources/VideoResourceRHI.cpp
+++ b/Engine/Plugins/Experimental/AVCodecs/AVCodecsCore/Source/AVCodecsCoreRHI/Private/Video/Resources/VideoResourceRHI.cpp
@@ -161,7 +161,7 @@ TSharedPtr<FVideoResourceRHI> FVideoResourceRHI::Create(TSharedPtr<FAVDevice> co
         
         if (Descriptor.BulkData)
         {
-            TextureDesc.SetBulkData(Descriptor.BulkData);
+            TextureDesc.SetInitActionBulkData(Descriptor.BulkData);
         }
 
 		// Unreals support for NV12 asnd P010 is not universally supported so we actually use R8 or G16 under the hood
@@ -213,7 +213,7 @@ TSharedPtr<FVideoResourceRHI> FVideoResourceRHI::Create(TSharedPtr<FAVDevice> co
 		FTextureRHIRef Texture;
 		ENQUEUE_RENDER_COMMAND(CreateTexture)
 			([&Texture, TextureDesc, Fence](FRHICommandListImmediate& RHICmdList) {
-				Texture = GDynamicRHI->RHICreateTexture(RHICmdList, TextureDesc);
+				Texture = RHICmdList.CreateTexture(TextureDesc);
 				RHICmdList.WriteGPUFence(Fence);
 				});
 
diff --git a/Engine/Plugins/Media/AvfMedia/Source/AvfMedia/Private/Player/AvfMediaVideoSampler.cpp b/Engine/Plugins/Media/AvfMedia/Source/AvfMedia/Private/Player/AvfMediaVideoSampler.cpp
--- a/Engine/Plugins/Media/AvfMedia/Source/AvfMedia/Private/Player/AvfMediaVideoSampler.cpp
+++ b/Engine/Plugins/Media/AvfMedia/Source/AvfMedia/Private/Player/AvfMediaVideoSampler.cpp
@@ -227,12 +227,12 @@ void FAvfMediaVideoSampler::ProcessFrame(CVPixelBufferRef Frame, FTimespan Sampl
 			const FRHITextureCreateDesc YDesc =
 				FRHITextureCreateDesc::Create2D(TEXT("YTex"), YWidth, YHeight, PF_G8)
 				.SetFlags(ETextureCreateFlags::NoTiling | ETextureCreateFlags::ShaderResource)
-				.SetBulkData(new FAvfTexture2DResourceWrapper(YTextureRef));
+				.SetInitActionBulkData(new FAvfTexture2DResourceWrapper(YTextureRef));
 
 			const FRHITextureCreateDesc UVDesc =
 				FRHITextureCreateDesc::Create2D(TEXT("UVTex"), UVWidth, UVHeight, PF_R8G8)
 				.SetFlags(ETextureCreateFlags::NoTiling | ETextureCreateFlags::ShaderResource)
-				.SetBulkData(new FAvfTexture2DResourceWrapper(UVTextureRef));
+				.SetInitActionBulkData(new FAvfTexture2DResourceWrapper(UVTextureRef));
 			
 			TRefCountPtr<FRHITexture> YTex = RHICmdList.CreateTexture(YDesc);
 			TRefCountPtr<FRHITexture> UVTex = RHICmdList.CreateTexture(UVDesc);
@@ -295,7 +295,7 @@ void FAvfMediaVideoSampler::ProcessFrame(CVPixelBufferRef Frame, FTimespan Sampl
 			const FRHITextureCreateDesc Desc =
 				FRHITextureCreateDesc::Create2D(TEXT("FAvfMediaVideoSampler"), Width, Height, PF_B8G8R8A8)
 				.SetFlags(ETextureCreateFlags::SRGB | ETextureCreateFlags::NoTiling | ETextureCreateFlags::ShaderResource)
-				.SetBulkData(new FAvfTexture2DResourceWrapper(TextureRef));
+				.SetInitActionBulkData(new FAvfTexture2DResourceWrapper(TextureRef));
 
 			ShaderResource = RHICmdList.CreateTexture(Desc);
 			
diff --git a/Engine/Plugins/Media/AvfMedia/Source/AvfMediaCapture/Private/Player/AvfMediaCapturePlayer.cpp b/Engine/Plugins/Media/AvfMedia/Source/AvfMediaCapture/Private/Player/AvfMediaCapturePlayer.cpp
--- a/Engine/Plugins/Media/AvfMedia/Source/AvfMediaCapture/Private/Player/AvfMediaCapturePlayer.cpp
+++ b/Engine/Plugins/Media/AvfMedia/Source/AvfMediaCapture/Private/Player/AvfMediaCapturePlayer.cpp
@@ -740,7 +740,7 @@ void FAvfMediaCapturePlayer::ProcessSampleBufferVideo(CMSampleBufferRef SampleBu
 		FRHITextureCreateDesc Desc =
 			FRHITextureCreateDesc::Create2D(TEXT("FAvfMediaCapturePlayer"), Width, Height, PF_B8G8R8A8)
 			.SetFlags(ETextureCreateFlags::SRGB | ETextureCreateFlags::NoTiling | ETextureCreateFlags::ShaderResource)
-			.SetBulkData(new FAvfTexture2DResourceWrapper(TextureRef));
+			.SetInitActionBulkData(new FAvfTexture2DResourceWrapper(TextureRef));
 
 		CFRelease(TextureRef);
 		
diff --git a/Engine/Plugins/Media/ElectraUtil/Source/ElectraSamples/Private/ElectraTextureSampleAVF.cpp b/Engine/Plugins/Media/ElectraUtil/Source/ElectraSamples/Private/ElectraTextureSampleAVF.cpp
--- a/Engine/Plugins/Media/ElectraUtil/Source/ElectraSamples/Private/ElectraTextureSampleAVF.cpp
+++ b/Engine/Plugins/Media/ElectraUtil/Source/ElectraSamples/Private/ElectraTextureSampleAVF.cpp
@@ -209,13 +209,13 @@ void FElectraMediaTexConvApple::ConvertTexture(FTextureRHIRef & InDstTexture, CV
 			const FRHITextureCreateDesc YDesc =
 				FRHITextureCreateDesc::Create2D(TEXT("YTex"), YWidth, YHeight, bIs8Bit ? PF_G8 : PF_G16)
 				.SetFlags(ETextureCreateFlags::NoTiling | ETextureCreateFlags::ShaderResource)
-				.SetBulkData(new FTexConvTexResourceWrapper(YTextureRef))
+				.SetInitActionBulkData(new FTexConvTexResourceWrapper(YTextureRef))
 				.SetInitialState(ERHIAccess::SRVMask);
 
 			const FRHITextureCreateDesc UVDesc =
 				FRHITextureCreateDesc::Create2D(TEXT("UVTex"), UVWidth, UVHeight, bIs8Bit ? PF_R8G8 : PF_G16R16)
 				.SetFlags(ETextureCreateFlags::NoTiling | ETextureCreateFlags::ShaderResource)
-				.SetBulkData(new FTexConvTexResourceWrapper(UVTextureRef))
+				.SetInitActionBulkData(new FTexConvTexResourceWrapper(UVTextureRef))
 				.SetInitialState(ERHIAccess::SRVMask);
 
 			TRefCountPtr<FRHITexture> YTex = RHICreateTexture(YDesc);
@@ -335,7 +335,7 @@ void FElectraMediaTexConvApple::ConvertTexture(FTextureRHIRef & InDstTexture, CV
 				const FRHITextureCreateDesc Desc =
 					FRHITextureCreateDesc::Create2D(TEXT("DstTexture"), Width, Height, PF_R16G16B16A16_UNORM)
 					.SetFlags(ETextureCreateFlags::NoTiling | ETextureCreateFlags::ShaderResource)
-					.SetBulkData(new FTexConvTexResourceWrapper(TextureRef));
+					.SetInitActionBulkData(new FTexConvTexResourceWrapper(TextureRef));
 
 				InDstTexture = RHICreateTexture(Desc);
 
@@ -363,7 +363,7 @@ void FElectraMediaTexConvApple::ConvertTexture(FTextureRHIRef & InDstTexture, CV
 				const FRHITextureCreateDesc Desc =
 					FRHITextureCreateDesc::Create2D(TEXT("DstTexture"), Width, Height, PF_B8G8R8A8)
 					.SetFlags(ETextureCreateFlags::SRGB | ETextureCreateFlags::NoTiling | ETextureCreateFlags::ShaderResource)
-					.SetBulkData(new FTexConvTexResourceWrapper(TextureRef));
+					.SetInitActionBulkData(new FTexConvTexResourceWrapper(TextureRef));
 
 				InDstTexture = RHICreateTexture(Desc);
 
diff --git a/Engine/Plugins/Runtime/AR/AppleAR/AppleARKit/Source/AppleARKit/Private/AppleARKitTextures.cpp b/Engine/Plugins/Runtime/AR/AppleAR/AppleARKit/Source/AppleARKit/Private/AppleARKitTextures.cpp
--- a/Engine/Plugins/Runtime/AR/AppleAR/AppleARKit/Source/AppleARKit/Private/AppleARKitTextures.cpp
+++ b/Engine/Plugins/Runtime/AR/AppleAR/AppleARKit/Source/AppleARKit/Private/AppleARKitTextures.cpp
@@ -956,12 +956,12 @@ public:
 			const FRHITextureCreateDesc YDesc =
 				FRHITextureCreateDesc::Create2D(TEXT("VideoTextureY"), CapturedYImageSize, PF_B8G8R8A8)
 				.SetFlags(ETextureCreateFlags::NoTiling | ETextureCreateFlags::ShaderResource)
-				.SetBulkData(new FAppleARKitCameraTextureResourceWrapper(CapturedYImage));
+				.SetInitActionBulkData(new FAppleARKitCameraTextureResourceWrapper(CapturedYImage));
 
 			const FRHITextureCreateDesc cBcRDesc =
 				FRHITextureCreateDesc::Create2D(TEXT("VideoTextureCbCr"), CapturedCbCrImageSize, PF_B8G8R8A8)
 				.SetFlags(ETextureCreateFlags::NoTiling | ETextureCreateFlags::ShaderResource)
-				.SetBulkData(new FAppleARKitCameraTextureResourceWrapper(CapturedCbCrImage));
+				.SetInitActionBulkData(new FAppleARKitCameraTextureResourceWrapper(CapturedCbCrImage));
 
 			// pull the Y and CbCr textures out of the captured image planes (format is fake here, it will get the format from the FAppleARKitCameraTextureResourceWrapper)
 			VideoTextureY = RHICreateTexture(YDesc);
diff --git a/Engine/Plugins/Runtime/AppleMoviePlayer/Source/AppleMoviePlayer/Private/AppleMovieStreamer.cpp b/Engine/Plugins/Runtime/AppleMoviePlayer/Source/AppleMoviePlayer/Private/AppleMovieStreamer.cpp
--- a/Engine/Plugins/Runtime/AppleMoviePlayer/Source/AppleMoviePlayer/Private/AppleMovieStreamer.cpp
+++ b/Engine/Plugins/Runtime/AppleMoviePlayer/Source/AppleMoviePlayer/Private/AppleMovieStreamer.cpp
@@ -596,7 +596,7 @@ bool FAVPlayerMovieStreamer::CheckForNextFrameAndCopy()
 		const FRHITextureCreateDesc Desc =
 			FRHITextureCreateDesc::Create2D(TEXT("FAVPlayerMovieStreamer"), Width, Height, PF_B8G8R8A8)
 			.SetFlags(ETextureCreateFlags::NoTiling)
-			.SetBulkData(new FAppleMovieStreamerTexture2DResourceWrapper(TextureCacheRef));
+			.SetInitActionBulkData(new FAppleMovieStreamerTexture2DResourceWrapper(TextureCacheRef));
 
 		FTextureRHIRef MediaWrappedTexture = RHICreateTexture(Desc);
 
diff --git a/Engine/Plugins/Runtime/nDisplay/Source/DisplayCluster/Private/Render/Containers/DisplayClusterRender_TextureResource.cpp b/Engine/Plugins/Runtime/nDisplay/Source/DisplayCluster/Private/Render/Containers/DisplayClusterRender_TextureResource.cpp
--- a/Engine/Plugins/Runtime/nDisplay/Source/DisplayCluster/Private/Render/Containers/DisplayClusterRender_TextureResource.cpp
+++ b/Engine/Plugins/Runtime/nDisplay/Source/DisplayCluster/Private/Render/Containers/DisplayClusterRender_TextureResource.cpp
@@ -125,7 +125,7 @@ void FDisplayClusterRender_TextureResource::InitRHI(FRHICommandListBase&)
 	const FRHITextureCreateDesc Desc =
 		FRHITextureCreateDesc::Create2D(TEXT("DisplayClusterRender_TextureResource"), Width, Height, PixelFormat)
 		.SetFlags(ETextureCreateFlags::ShaderResource)
-		.SetBulkData(&BulkDataInterface);
+		.SetInitActionBulkData(&BulkDataInterface);
 
 	TextureRHI = RHICreateTexture(Desc);
 
diff --git a/Engine/Source/Runtime/Apple/MetalRHI/Private/MetalDynamicRHI.h b/Engine/Source/Runtime/Apple/MetalRHI/Private/MetalDynamicRHI.h
--- a/Engine/Source/Runtime/Apple/MetalRHI/Private/MetalDynamicRHI.h
+++ b/Engine/Source/Runtime/Apple/MetalRHI/Private/MetalDynamicRHI.h
@@ -144,7 +144,7 @@ public:
 	virtual uint64 RHIGetMinimumAlignmentForBufferBackedSRV(EPixelFormat Format) final override;
 	virtual void RHIGetTextureMemoryStats(FTextureMemoryStats& OutStats) final override;
 	virtual bool RHIGetTextureMemoryVisualizeData(FColor* TextureData,int32 SizeX,int32 SizeY,int32 Pitch,int32 PixelSize) final override;
-	virtual FTextureRHIRef RHICreateTexture(FRHICommandListBase& RHICmdList, const FRHITextureCreateDesc& CreateDesc) final override;
+	[[nodiscard]] virtual FRHITextureInitializer RHICreateTextureInitializer(FRHICommandListBase& RHICmdList, const FRHITextureCreateDesc& CreateDesc) final override;
 	virtual FTextureRHIRef RHIAsyncCreateTexture2D(uint32 SizeX, uint32 SizeY, uint8 Format, uint32 NumMips, ETextureCreateFlags Flags, ERHIAccess InResourceState, void** InitialMipData, uint32 NumInitialMips, const TCHAR* DebugName, FGraphEventRef& OutCompletionEvent) final override;
 	virtual uint32 RHIComputeMemorySize(FRHITexture* TextureRHI) final override;
 	virtual FTextureRHIRef RHIAsyncReallocateTexture2D(FRHITexture* Texture2D, int32 NewMipCount, int32 NewSizeX, int32 NewSizeY, FThreadSafeCounter* RequestStatus) final override;
diff --git a/Engine/Source/Runtime/Apple/MetalRHI/Private/MetalTexture.cpp b/Engine/Source/Runtime/Apple/MetalRHI/Private/MetalTexture.cpp
--- a/Engine/Source/Runtime/Apple/MetalRHI/Private/MetalTexture.cpp
+++ b/Engine/Source/Runtime/Apple/MetalRHI/Private/MetalTexture.cpp
@@ -14,14 +14,13 @@
 #include "MetalLLM.h"
 #include "RHILockTracker.h"
 #include "RHICoreStats.h"
+#include "RHICoreTextureInitializer.h"
 #include "MetalBindlessDescriptors.h"
 #include "HAL/LowLevelMemStats.h"
 #include "ProfilingDebugging/AssetMetadataTrace.h"
 #include "RHICoreTexture.h"
 #include <CoreVideo/CVMetalTexture.h>
 
-volatile int64 FMetalSurface::ActiveUploads = 0;
-
 int32 GMetalMaxOutstandingAsyncTexUploads = 100 * 1024 * 1024;
 FAutoConsoleVariableRef CVarMetalMaxOutstandingAsyncTexUploads(
 															   TEXT("rhi.Metal.MaxOutstandingAsyncTexUploads"),
@@ -539,7 +538,7 @@ FMetalTextureCreateDesc& FMetalTextureCreateDesc::operator =(FMetalTextureCreate
     return *this;
 }
 
-FMetalSurface::FMetalSurface(FMetalDevice& MetalDevice, FRHICommandListBase* RHICmdList, FMetalTextureCreateDesc const& CreateDesc)
+FMetalSurface::FMetalSurface(FMetalDevice& MetalDevice, FMetalTextureCreateDesc const& CreateDesc)
 	: FRHITexture       (CreateDesc)
 	, Device			(MetalDevice)
 	, FormatKey         (CreateDesc.FormatKey)
@@ -694,28 +693,6 @@ FMetalSurface::FMetalSurface(FMetalDevice& MetalDevice, FRHICommandListBase* RHI
 		METAL_FATAL_ASSERT(Texture, TEXT("Failed to create texture, desc %s"), *NSStringToFString(CreateDesc.Desc->description()));
 	}
 
-	if (BulkData)
-	{
-		// Regular texture has some bulk data to handle
-		UE_LOG(LogMetal, Display, TEXT("Got a bulk data texture, with %d mips"), CreateDesc.NumMips);
-		checkf(CreateDesc.NumMips == 1, TEXT("Only handling bulk data with 1 mip and 1 array length"));
-		check(RHICmdList);
-
-		FRHICommandListImmediate& RHICmdListImmediate = RHICmdList->GetAsImmediate();
-
-		// lock, copy, unlock
-		const FRHILockTextureArgs LockArgs = FRHILockTextureArgs::Lock2D(this, 0, RLM_WriteOnly, false);
-
-		FRHILockTextureResult LockResult = RHICmdListImmediate.LockTexture(LockArgs);
-		check(LockResult.Data);
-		FMemory::Memcpy(LockResult.Data, BulkData->GetResourceBulkData(), BulkData->GetResourceBulkDataSize());
-		RHICmdListImmediate.UnlockTexture(LockArgs);
-
-		// bulk data can be unloaded now
-		BulkData->Discard();
-		BulkData = nullptr;
-	}
-
 	// calculate size of the texture
 	TotalTextureSize = GetMemorySize();
 
@@ -813,18 +790,25 @@ FMetalSurface::FMetalSurface(FMetalDevice& MetalDevice, FRHICommandListBase* RHI
 	}
 
 #if PLATFORM_SUPPORTS_BINDLESS_RENDERING
-    FMetalBindlessDescriptorManager* BindlessDescriptorManager = Device.GetBindlessDescriptorManager();
-    check(BindlessDescriptorManager);
-	
 	if(IsMetalBindlessEnabled())
 	{
+		FMetalBindlessDescriptorManager* BindlessDescriptorManager = Device.GetBindlessDescriptorManager();
+		check(BindlessDescriptorManager);
+
 		BindlessHandle = BindlessDescriptorManager->ReserveDescriptor(ERHIDescriptorHeapType::Standard);
-		
-		// NOTE: Might be updated later (using RHIUpdateTextureReference).
-		if (Texture)
-		{
-			BindlessDescriptorManager->BindTexture(*RHICmdList, BindlessHandle, Texture.get(), EDescriptorUpdateType_Immediate);
-		}
+	}
+#endif
+}
+
+void FMetalSurface::Initialize(FRHICommandListBase& RHICmdList)
+{
+#if PLATFORM_SUPPORTS_BINDLESS_RENDERING
+	if (BindlessHandle.IsValid())
+	{
+		FMetalBindlessDescriptorManager* BindlessDescriptorManager = Device.GetBindlessDescriptorManager();
+		check(BindlessDescriptorManager);
+
+		BindlessDescriptorManager->BindTexture(RHICmdList, BindlessHandle, Texture.get(), EDescriptorUpdateType_Immediate);
 	}
 #endif
 }
@@ -1021,10 +1005,6 @@ void FMetalSurface::UpdateSurfaceAndDestroySourceBuffer(FMetalRHICommandContext*
 		
 		int64 Size = BytesPerImage * Region.size.depth * FMath::Max(1u, ArrayIndex);
 		
-		int64 Count = FPlatformAtomics::InterlockedAdd(&ActiveUploads, Size);
-		
-		bool const bWait = GMetalMaxOutstandingAsyncTexUploads > 0 && Count >= GMetalMaxOutstandingAsyncTexUploads;
-		
 		MTL::BlitOption Options = MTL::BlitOptionNone;
 #if !PLATFORM_MAC
 		if (Texture->pixelFormat() >= MTL::PixelFormatPVRTC_RGB_2BPP && Texture->pixelFormat() <= MTL::PixelFormatPVRTC_RGBA_4BPP_sRGB)
@@ -1034,11 +1014,6 @@ void FMetalSurface::UpdateSurfaceAndDestroySourceBuffer(FMetalRHICommandContext*
 #endif
 		Context->CopyFromBufferToTexture(Source, 0, Stride, BytesPerImage, Region.size, Texture.get(), ArrayIndex, MipIndex, Region.origin, Options);
 		
-		Context->GetContextSyncPoint()->OnCompletionCallback([Size]()
-		{
-			FPlatformAtomics::InterlockedAdd(&ActiveUploads, -Size);
-		});
-		
 		INC_DWORD_STAT_BY(STAT_MetalTextureMemUpdate, Size);
 	}
 	else
@@ -1376,13 +1351,185 @@ uint32 FMetalDynamicRHI::RHIComputeMemorySize(FRHITexture* TextureRHI)
  2D texture support.
  -----------------------------------------------------------------------------*/
 
-FTextureRHIRef FMetalDynamicRHI::RHICreateTexture(FRHICommandListBase& RHICmdList, const FRHITextureCreateDesc& CreateDesc)
+struct FMetalInitializerMemory
+{
+	FMetalInitializerMemory(MTL::Buffer* InMemory)
+		: Memory(InMemory)
+	{
+	}
+	FMetalInitializerMemory(FMetalInitializerMemory&& Other)
+		: Memory(Other.Memory)
+	{
+		Other.Memory = nullptr;
+	}
+
+	~FMetalInitializerMemory()
+	{
+		if (Memory)
+		{
+			Memory->release();
+		}
+	}
+
+	MTL::Buffer* Relinquish()
+	{
+		if (Memory)
+		{
+			MTL::Buffer* Result = Memory;
+			Memory = nullptr;
+			return Result;
+		}
+
+		return nullptr;
+	}
+
+	MTL::Buffer* Memory{};
+};
+
+static bool IsMetalTextureInitializingPVRTC(MTL::Texture* Texture)
+{
+#if !PLATFORM_MAC
+	return (Texture->pixelFormat() >= MTL::PixelFormatPVRTC_RGB_2BPP && Texture->pixelFormat() <= MTL::PixelFormatPVRTC_RGBA_4BPP_sRGB);
+#else
+	return false;
+#endif
+}
+
+static UE::RHICore::FDefaultLayoutTextureInitializer CreateWritableTextureInitializer(FRHICommandListBase& RHICmdList, FMetalSurface* Texture, const FRHITextureCreateDesc& CreateDesc)
+{
+	const uint64 WritableMemorySize = UE::RHITextureUtils::CalculateTextureSize(CreateDesc);
+	
+	MTL::Buffer* WritableBuffer = Texture->Device.GetDevice()->newBuffer(WritableMemorySize, MTL::ResourceStorageModeShared);
+	WritableBuffer->setLabel(NS::String::string("Texture Initializer Scratch", NS::UTF8StringEncoding));
+
+	return UE::RHICore::FDefaultLayoutTextureInitializer(RHICmdList, Texture, WritableBuffer->contents(), WritableMemorySize,
+		[Texture = TRefCountPtr<FMetalSurface>(Texture), WritableBuffer = FMetalInitializerMemory(WritableBuffer), WritableMemorySize](FRHICommandListBase& RHICmdList) mutable
+		{
+			Texture->Initialize(RHICmdList);
+
+			RHICmdList.EnqueueLambda(
+				[Texture, WritableBuffer = MoveTemp(WritableBuffer), WritableMemorySize](FRHICommandListBase& RHICmdList) mutable
+				{
+					FMetalRHICommandContext& Context = FMetalRHICommandContext::Get(RHICmdList);
+
+					MTL::Texture* const MetalTexture = Texture->Texture.get();
+
+					const bool bPrivateStorage = MetalTexture->storageMode() == MTL::StorageModePrivate;
+					const bool bPVRTC = IsMetalTextureInitializingPVRTC(MetalTexture);
+					const MTL::BlitOption CopyOptions =  bPVRTC ? MTL::BlitOptionRowLinearPVRTC : MTL::BlitOptionNone;
+
+					FMetalBufferPtr SourceBuffer = FMetalBufferPtr(new FMetalBuffer(WritableBuffer.Relinquish(), FMetalBuffer::FreePolicy::Owner));
+
+					const FRHITextureDesc& TextureDesc = Texture->GetDesc();
+					const uint32 FaceCount = TextureDesc.IsTextureCube() ? 6 : 1;
+
+					for (uint32 FaceIndex = 0; FaceIndex < FaceCount; FaceIndex++)
+					{
+						for (uint32 ArrayIndex = 0; ArrayIndex < TextureDesc.ArraySize; ArrayIndex++)
+						{
+							for (uint32 MipIndex = 0; MipIndex < TextureDesc.NumMips; MipIndex++)
+							{
+								uint64 SubresourceStride = 0;
+								uint64 SubresourceSize = 0;
+								const uint64 Offset = UE::RHITextureUtils::CalculateSubresourceOffset(TextureDesc, FaceIndex, ArrayIndex, MipIndex, SubresourceStride, SubresourceSize);
+
+								// The number of bytes between each 2D image of a 3D texture. This should be 0 for 2D textures.
+								uint64 SourceBytesPerImage = 0;
+								if (TextureDesc.IsTexture3D())
+								{
+									SourceBytesPerImage = Texture->GetMipSize(MipIndex, nullptr, true);
+								}
+
+								check((Offset + SubresourceSize) <= WritableMemorySize);
+
+								const uint32 FullArrayIndex = FaceIndex + ArrayIndex * FaceCount;
+
+								const FUintVector3 Extents = UE::RHITextureUtils::CalculateMipExtents(TextureDesc, MipIndex);
+								const MTL::Region Region = MTL::Region(0, 0, 0, Extents.X, Extents.Y, Extents.Z);
+
+								if (bPrivateStorage)
+								{
+									Context.CopyFromBufferToTexture(
+										SourceBuffer,
+										Offset,
+										SubresourceStride,
+										SourceBytesPerImage,
+										Region.size,
+										MetalTexture,
+										FullArrayIndex,
+										MipIndex,
+										Region.origin,
+										CopyOptions
+									);
+								}
+								else
+								{
+									// @todo Calculate correct strides and byte-counts
+									const uint64 UsedStride = bPVRTC ? 0 : SubresourceStride;
+									const uint64 UsedSize = bPVRTC ? 0 : SourceBytesPerImage;
+
+									MetalTexture->replaceRegion(
+										Region,
+										MipIndex,
+										FullArrayIndex,
+										SourceBuffer->Contents(),
+										UsedStride,
+										UsedSize
+									);
+								}
+							}
+						}
+					}
+
+					FMetalDynamicRHI::Get().DeferredDelete(SourceBuffer);
+				}
+			);
+
+			return TRefCountPtr<FRHITexture>(MoveTemp(Texture));
+		}
+	);
+}
+
+FRHITextureInitializer FMetalDynamicRHI::RHICreateTextureInitializer(FRHICommandListBase& RHICmdList, const FRHITextureCreateDesc& CreateDesc)
 {
     MTL_SCOPED_AUTORELEASE_POOL;
 	LLM_SCOPE_DYNAMIC_STAT_OBJECTPATH_FNAME(CreateDesc.OwnerName, ELLMTagSet::Assets);
 	LLM_SCOPE_DYNAMIC_STAT_OBJECTPATH_FNAME(CreateDesc.GetTraceClassName(), ELLMTagSet::AssetClasses);
 	UE_TRACE_METADATA_SCOPE_ASSET_FNAME(CreateDesc.DebugName, CreateDesc.GetTraceClassName(), CreateDesc.OwnerName);
-    return new FMetalSurface(*Device, &RHICmdList, FMetalTextureCreateDesc(*Device, CreateDesc));
+
+	FMetalSurface* Texture = new FMetalSurface(*Device, FMetalTextureCreateDesc(*Device, CreateDesc));
+
+	if (CreateDesc.InitAction == ERHITextureInitAction::Default)
+	{
+		return UE::RHICore::FBaseTextureInitializerImplementation(RHICmdList, Texture,
+			[Texture = TRefCountPtr<FMetalSurface>(Texture)](FRHICommandListBase& RHICmdList) mutable
+			{
+				Texture->Initialize(RHICmdList);
+				return TRefCountPtr<FRHITexture>(MoveTemp(Texture));
+			}
+		);
+	}
+
+	if (CreateDesc.InitAction == ERHITextureInitAction::BulkData)
+	{
+		check(CreateDesc.BulkData);
+
+		UE::RHICore::FDefaultLayoutTextureInitializer Initializer = CreateWritableTextureInitializer(RHICmdList, Texture, CreateDesc);
+		check(CreateDesc.BulkData->GetResourceBulkDataSize() <= Initializer.GetWritableSize());
+
+		FMemory::Memcpy(Initializer.GetWritableData(), CreateDesc.BulkData->GetResourceBulkData(), CreateDesc.BulkData->GetResourceBulkDataSize());
+
+		CreateDesc.BulkData->Discard();
+
+		return MoveTemp(Initializer);
+	}
+
+	if (CreateDesc.InitAction == ERHITextureInitAction::Initializer)
+	{
+		return CreateWritableTextureInitializer(RHICmdList, Texture, CreateDesc);
+	}
+
+	return UE::RHICore::HandleUnknownTextureInitializerInitAction(RHICmdList, CreateDesc);
 }
 
 FTextureRHIRef FMetalDynamicRHI::RHIAsyncCreateTexture2D(uint32 SizeX, uint32 SizeY, uint8 Format, uint32 NumMips, ETextureCreateFlags Flags, ERHIAccess InResourceState, void** InitialMipData, uint32 NumInitialMips, const TCHAR* DebugName, FGraphEventRef& OutCompletionEvent)
@@ -1419,7 +1566,8 @@ FTextureRHIRef FMetalDynamicRHI::RHIAsyncReallocateTexture2D(FRHITexture* OldTex
 	LLM_SCOPE_DYNAMIC_STAT_OBJECTPATH_FNAME(CreateDesc.GetTraceClassName(), ELLMTagSet::AssetClasses);
 	UE_TRACE_METADATA_SCOPE_ASSET_FNAME(CreateDesc.DebugName, CreateDesc.GetTraceClassName(), CreateDesc.OwnerName);
 
-    FMetalSurface* NewTexture = new FMetalSurface(*Device, &RHICmdList, FMetalTextureCreateDesc(*Device, CreateDesc));
+	FMetalSurface* NewTexture = new FMetalSurface(*Device, FMetalTextureCreateDesc(*Device, CreateDesc));
+	NewTexture->Initialize(RHICmdList);
 
     // Copy shared mips
     RHICmdList.EnqueueLambda([this, OldTexture, NewSizeX, NewSizeY, NewTexture, RequestStatus](FRHICommandListImmediate& RHICmdList)
diff --git a/Engine/Source/Runtime/Apple/MetalRHI/Private/MetalViewport.cpp b/Engine/Source/Runtime/Apple/MetalRHI/Private/MetalViewport.cpp
--- a/Engine/Source/Runtime/Apple/MetalRHI/Private/MetalViewport.cpp
+++ b/Engine/Source/Runtime/Apple/MetalRHI/Private/MetalViewport.cpp
@@ -279,12 +279,12 @@ void FMetalViewport::Resize(uint32 InSizeX, uint32 InSizeY, bool bInIsFullscreen
 		
 		CreateDesc.SetInitialState(RHIGetDefaultResourceState(CreateDesc.Flags, false));
 
-		NewBackBuffer = new FMetalSurface(Device, nullptr, FMetalTextureCreateDesc(Device, CreateDesc));
+		NewBackBuffer = new FMetalSurface(Device, FMetalTextureCreateDesc(Device, CreateDesc));
 		NewBackBuffer->Viewport = this;
 
         if (GMetalSupportsIntermediateBackBuffer && GMetalSeparatePresentThread)
         {
-            DoubleBuffer = new FMetalSurface(Device, nullptr, FMetalTextureCreateDesc(Device, CreateDesc));
+            DoubleBuffer = new FMetalSurface(Device, FMetalTextureCreateDesc(Device, CreateDesc));
             DoubleBuffer->Viewport = this;
         }
 
diff --git a/Engine/Source/Runtime/Apple/MetalRHI/Public/MetalResources.h b/Engine/Source/Runtime/Apple/MetalRHI/Public/MetalResources.h
--- a/Engine/Source/Runtime/Apple/MetalRHI/Public/MetalResources.h
+++ b/Engine/Source/Runtime/Apple/MetalRHI/Public/MetalResources.h
@@ -166,13 +166,15 @@ public:
 	/** 
 	 * Constructor that will create Texture and Color/DepthBuffers as needed
 	 */
-	FMetalSurface(FMetalDevice& Device, FRHICommandListBase* RHICmdList, FMetalTextureCreateDesc const& CreateDesc);
+	FMetalSurface(FMetalDevice& Device, FMetalTextureCreateDesc const& CreateDesc);
 	
 	/**
 	 * Destructor
 	 */
 	virtual ~FMetalSurface();
 
+	void Initialize(FRHICommandListBase& RHICmdList);
+
 	/** @returns A newly allocated buffer object large enough for the surface within the texture specified. */
 	MTL::Buffer* AllocSurface(const FRHILockTextureArgs& Arguments, uint32 MipBytes, uint32 DestStride);
 
@@ -273,9 +275,6 @@ private:
 	
 	// The movie playback IOSurface/CVTexture wrapper to avoid page-off
 	CFTypeRef ImageSurfaceRef;
-
-	// Count of outstanding async. texture uploads
-	static volatile int64 ActiveUploads;
 };
 
 class FMetalBufferData
diff --git a/Engine/Source/Runtime/D3D12RHI/Private/D3D12RHIPrivate.h b/Engine/Source/Runtime/D3D12RHI/Private/D3D12RHIPrivate.h
--- a/Engine/Source/Runtime/D3D12RHI/Private/D3D12RHIPrivate.h
+++ b/Engine/Source/Runtime/D3D12RHI/Private/D3D12RHIPrivate.h
@@ -354,7 +354,9 @@ public:
 	virtual void RHIGetTextureMemoryStats(FTextureMemoryStats& OutStats) final override;
 	virtual bool RHIGetTextureMemoryVisualizeData(FColor* TextureData, int32 SizeX, int32 SizeY, int32 Pitch, int32 PixelSize) final override;
 	virtual FTextureRHIRef RHIAsyncCreateTexture2D(uint32 SizeX, uint32 SizeY, uint8 Format, uint32 NumMips, ETextureCreateFlags Flags, ERHIAccess InResourceState, void** InitialMipData, uint32 NumInitialMips, const TCHAR* DebugName, FGraphEventRef& OutCompletionEvent) final override;
-	virtual FTextureRHIRef RHICreateTexture(FRHICommandListBase& RHICmdList, const FRHITextureCreateDesc& CreateDesc) override;
+	
+	[[nodiscard]] virtual FRHITextureInitializer RHICreateTextureInitializer(FRHICommandListBase& RHICmdList, const FRHITextureCreateDesc& CreateDesc) override;
+
 	virtual uint32 RHIComputeMemorySize(FRHITexture* TextureRHI) final override;
 	virtual FTextureRHIRef RHIAsyncReallocateTexture2D(FRHITexture* Texture2D, int32 NewMipCount, int32 NewSizeX, int32 NewSizeY, FThreadSafeCounter* RequestStatus) override;
 
@@ -568,9 +570,6 @@ public:
 
 	virtual FD3D12ResourceDesc GetResourceDesc(const FRHITextureDesc& CreateInfo) const;
 
-	// Fills in the ClearValue argument based on the ResourceDesc and CreateDesc flags.
-	D3D12_CLEAR_VALUE* FillClearValue(D3D12_CLEAR_VALUE& ClearValue, const FD3D12ResourceDesc& ResourceDesc, const FRHITextureCreateDesc& CreateDesc) const;
-
 	struct FCreateTextureInternalResult
 	{
 		FD3D12Texture* Texture{};
@@ -579,32 +578,18 @@ public:
 	};
 
 	// Per-platform D3D12Texture creation.
-	virtual FCreateTextureInternalResult CreateTextureInternal(const FRHITextureCreateDesc& CreateDesc, bool bHasInitialData, ID3D12ResourceAllocator* ResourceAllocator = nullptr);
-
-	// Per-platform D3D12Texture post-creation adjustments, like changing the resource state or uploading data.
-	virtual FD3D12Texture* FinalizeCreateTexture(FRHICommandListBase& RHICmdList, const FRHITextureCreateDesc& CreateDesc, const FCreateTextureInternalResult& CreateResult) const;
+	virtual FCreateTextureInternalResult CreateTextureInternal(const FRHITextureCreateDesc& CreateDesc, ID3D12ResourceAllocator* ResourceAllocator = nullptr);
 
 	// Create D3D12 texture with a resource allocator. Does not currently support any post-creation state or uploading of data.
 	FD3D12Texture* CreateD3D12Texture(const FRHITextureCreateDesc& CreateDesc, ID3D12ResourceAllocator* ResourceAllocator)
 	{
 		check(CreateDesc.BulkData == nullptr);
-		constexpr bool bHasInitialData = false;
 
-		FCreateTextureInternalResult CreateResult = CreateTextureInternal(CreateDesc, bHasInitialData, ResourceAllocator);
+		FCreateTextureInternalResult CreateResult = CreateTextureInternal(CreateDesc, ResourceAllocator);
 		checkSlow(CreateResult.CreateState == CreateResult.DesiredState);
 		return CreateResult.Texture;
 	}
 
-	// Create D3D12 texture with a command list, meaning we can upload data and adjust resource states
-	FD3D12Texture* CreateD3D12Texture(FRHICommandListBase& RHICmdList, const FRHITextureCreateDesc& CreateDesc)
-	{
-		const bool bHasInitialData = (CreateDesc.BulkData != nullptr);
-		constexpr ID3D12ResourceAllocator* ResourceAllocator = nullptr;
-
-		const FCreateTextureInternalResult CreateResult = CreateTextureInternal(CreateDesc, bHasInitialData, ResourceAllocator);
-		return FinalizeCreateTexture(RHICmdList, CreateDesc, CreateResult);
-	}
-
 	virtual FD3D12Texture* CreateNewD3D12Texture(const FRHITextureCreateDesc& CreateDesc, class FD3D12Device* Device);
 
 	struct FCreateBufferInternalResult
diff --git a/Engine/Source/Runtime/D3D12RHI/Private/D3D12Texture.cpp b/Engine/Source/Runtime/D3D12RHI/Private/D3D12Texture.cpp
--- a/Engine/Source/Runtime/D3D12RHI/Private/D3D12Texture.cpp
+++ b/Engine/Source/Runtime/D3D12RHI/Private/D3D12Texture.cpp
@@ -13,6 +13,7 @@
 #include "ProfilingDebugging/AssetMetadataTrace.h"
 #include "RHICoreStats.h"
 #include "RHICoreTexture.h"
+#include "RHICoreTextureInitializer.h"
 #include "RHITextureUtils.h"
 
 int64 FD3D12GlobalStats::GDedicatedVideoMemory = 0;
@@ -591,34 +592,34 @@ FD3D12Texture* FD3D12DynamicRHI::CreateNewD3D12Texture(const FRHITextureCreateDe
 	return new FD3D12Texture(CreateDesc, Device);
 }
 
-D3D12_CLEAR_VALUE* FD3D12DynamicRHI::FillClearValue(D3D12_CLEAR_VALUE& ClearValue, const FD3D12ResourceDesc& ResourceDesc, const FRHITextureCreateDesc& CreateDesc) const
+D3D12_CLEAR_VALUE* UE::D3D12RHI::TextureUtils::FillClearValue(D3D12_CLEAR_VALUE& ClearValue, const FD3D12ResourceDesc& ResourceDesc, const FRHITextureDesc& TextureDesc)
 {
 	bool bCreateRTV = EnumHasAnyFlags(ResourceDesc.Flags, D3D12_RESOURCE_FLAG_ALLOW_RENDER_TARGET);
 	bool bCreateDSV = EnumHasAnyFlags(ResourceDesc.Flags, D3D12_RESOURCE_FLAG_ALLOW_DEPTH_STENCIL);
 
-	const DXGI_FORMAT PlatformResourceFormat = UE::DXGIUtilities::GetPlatformTextureResourceFormat((DXGI_FORMAT)GPixelFormats[CreateDesc.Format].PlatformFormat, CreateDesc.Flags);
+	const DXGI_FORMAT PlatformResourceFormat = UE::DXGIUtilities::GetPlatformTextureResourceFormat((DXGI_FORMAT)GPixelFormats[TextureDesc.Format].PlatformFormat, TextureDesc.Flags);
 
 	D3D12_CLEAR_VALUE* ClearValuePtr = nullptr;
-	if (bCreateDSV && CreateDesc.ClearValue.ColorBinding == EClearBinding::EDepthStencilBound)
+	if (bCreateDSV && TextureDesc.ClearValue.ColorBinding == EClearBinding::EDepthStencilBound)
 	{
 		const DXGI_FORMAT PlatformDepthStencilFormat = UE::DXGIUtilities::FindDepthStencilFormat(PlatformResourceFormat);
 
-		ClearValue = CD3DX12_CLEAR_VALUE(PlatformDepthStencilFormat, CreateDesc.ClearValue.Value.DSValue.Depth, (uint8)CreateDesc.ClearValue.Value.DSValue.Stencil);
+		ClearValue = CD3DX12_CLEAR_VALUE(PlatformDepthStencilFormat, TextureDesc.ClearValue.Value.DSValue.Depth, (uint8)TextureDesc.ClearValue.Value.DSValue.Stencil);
 		ClearValuePtr = &ClearValue;
 	}
-	else if (bCreateRTV && CreateDesc.ClearValue.ColorBinding == EClearBinding::EColorBound)
+	else if (bCreateRTV && TextureDesc.ClearValue.ColorBinding == EClearBinding::EColorBound)
 	{
-		const bool bSRGB = EnumHasAnyFlags(CreateDesc.Flags, TexCreate_SRGB);
+		const bool bSRGB = EnumHasAnyFlags(TextureDesc.Flags, TexCreate_SRGB);
 		const DXGI_FORMAT PlatformRenderTargetFormat = UE::DXGIUtilities::FindShaderResourceFormat(PlatformResourceFormat, bSRGB);
 
-		ClearValue = CD3DX12_CLEAR_VALUE(PlatformRenderTargetFormat, CreateDesc.ClearValue.Value.Color);
+		ClearValue = CD3DX12_CLEAR_VALUE(PlatformRenderTargetFormat, TextureDesc.ClearValue.Value.Color);
 		ClearValuePtr = &ClearValue;
 	}
 
 	return ClearValuePtr;
 }
 
-FD3D12DynamicRHI::FCreateTextureInternalResult FD3D12DynamicRHI::CreateTextureInternal(const FRHITextureCreateDesc& InCreateDesc, bool bHasInitialData, ID3D12ResourceAllocator* ResourceAllocator)
+FD3D12DynamicRHI::FCreateTextureInternalResult FD3D12DynamicRHI::CreateTextureInternal(const FRHITextureCreateDesc& InCreateDesc, ID3D12ResourceAllocator* ResourceAllocator)
 {
 #if PLATFORM_WINDOWS
 	TRACE_CPUPROFILER_EVENT_SCOPE(D3D12RHI::CreateTextureInternal);
@@ -632,11 +633,13 @@ FD3D12DynamicRHI::FCreateTextureInternalResult FD3D12DynamicRHI::CreateTextureIn
 	// Virtual textures currently not supported in default D3D12
 	check(!EnumHasAnyFlags(CreateDesc.Flags, TexCreate_Virtual));
 
+	const bool bCreateAsCopyDest = (ResourceAllocator == nullptr && CreateDesc.InitAction != ERHITextureInitAction::Default);
+
 	// Get the resource desc
 	FD3D12ResourceDesc ResourceDesc = GetResourceDesc(CreateDesc);
 
 	D3D12_CLEAR_VALUE ClearValue;
-	D3D12_CLEAR_VALUE* ClearValuePtr = FillClearValue(ClearValue, ResourceDesc, CreateDesc);
+	D3D12_CLEAR_VALUE* ClearValuePtr = UE::D3D12RHI::TextureUtils::FillClearValue(ClearValue, ResourceDesc, CreateDesc);
 
 	const FD3D12Resource::FD3D12ResourceTypeHelper Type(ResourceDesc, D3D12_HEAP_TYPE_DEFAULT);
 	const D3D12_RESOURCE_STATES DesiredState = Type.GetOptimalInitialState(CreateDesc.InitialState, true);
@@ -645,7 +648,7 @@ FD3D12DynamicRHI::FCreateTextureInternalResult FD3D12DynamicRHI::CreateTextureIn
 		// When creating a resource with castable formats, d3d12 uses enhanced barriers behind the scenes which means we have to start it in the "common" state to use "legacy" barriers
 		ResourceDesc.SupportsUncompressedUAV() ? D3D12_RESOURCE_STATE_COMMON
 		// If we have initial data, we want the resource created in COPY_DEST so we can copy the data immediately
-		: bHasInitialData ? D3D12_RESOURCE_STATE_COPY_DEST
+		: bCreateAsCopyDest ? D3D12_RESOURCE_STATE_COPY_DEST
 		// Otherwise, create the resource in the caller's specified state
 		: DesiredState;
 
@@ -722,7 +725,7 @@ FD3D12DynamicRHI::FCreateTextureInternalResult FD3D12DynamicRHI::CreateTextureIn
 		}
 
 		// Unlock immediately if no initial data
-		if (!bHasInitialData)
+		if (!bCreateAsCopyDest)
 		{
 			Location.UnlockPoolData();
 		}
@@ -763,7 +766,7 @@ FD3D12DynamicRHI::FCreateTextureInternalResult FD3D12DynamicRHI::CreateTextureIn
 #endif // PLATFORM_WINDOWS
 }
 
-static void ReconcileInitialState(FRHICommandListBase& RHICmdList, FD3D12Texture* Texture, D3D12_RESOURCE_STATES CurrentState, D3D12_RESOURCE_STATES NeededState)
+void UE::D3D12RHI::TextureUtils::ReconcileInitialState(FRHICommandListBase& RHICmdList, FD3D12Texture* Texture, D3D12_RESOURCE_STATES CurrentState, D3D12_RESOURCE_STATES NeededState)
 {
 	if (CurrentState != NeededState)
 	{
@@ -780,29 +783,228 @@ static void ReconcileInitialState(FRHICommandListBase& RHICmdList, FD3D12Texture
 	}
 }
 
-FD3D12Texture* FD3D12DynamicRHI::FinalizeCreateTexture(FRHICommandListBase& RHICmdList, const FRHITextureCreateDesc& CreateDesc, const FCreateTextureInternalResult& CreateResult) const
+static uint32 CalculateSubresourceCount(const FRHITextureDesc& TextureDesc)
 {
-	if (CreateDesc.BulkData)
+	// each mip of each array slice counts as a subresource
+	uint16 ArraySize = TextureDesc.IsTextureArray() ? TextureDesc.ArraySize : 1;
+	if (TextureDesc.IsTextureCube())
 	{
-		// Make sure the resource is in the copy dest state before enqueuing the copy
-		ReconcileInitialState(RHICmdList, CreateResult.Texture, CreateResult.CreateState, D3D12_RESOURCE_STATE_COPY_DEST);
-
-		// After the copy, make sure the resource is in the requested state
-		CreateResult.Texture->InitializeTextureData(RHICmdList, CreateDesc, CreateResult.DesiredState);
-		CreateDesc.BulkData->Discard();
+		ArraySize *= 6;
 	}
-	else
+	return TextureDesc.NumMips * ArraySize;
+}
+
+uint64 UE::D3D12RHI::TextureUtils::CalculateResourceSize(FD3D12Texture* Texture)
+{
+	const uint32 NumSubresources = CalculateSubresourceCount(Texture->GetDesc());
+
+	const D3D12_RESOURCE_DESC& ResourceDesc = Texture->GetResource()->GetDesc();
+	uint64 AllocationSize = 0;
+	Texture->GetParentDevice()->GetDevice()->GetCopyableFootprints(
+		&ResourceDesc,
+		0,
+		NumSubresources,
+		0,
+		nullptr,
+		nullptr,
+		nullptr,
+		&AllocationSize
+	);
+
+	return AllocationSize;
+}
+
+void UE::D3D12RHI::TextureUtils::CopyBulkData(void* UploadMemory, uint64 UploadMemorySize, FD3D12Texture* Texture, FResourceBulkDataInterface* BulkData)
+{
+	const uint32 NumSubresources = CalculateSubresourceCount(Texture->GetDesc());
+	CA_ASSUME(NumSubresources > 0);
+
+	const size_t FootprintsMemorySize = NumSubresources * (sizeof(D3D12_PLACED_SUBRESOURCE_FOOTPRINT) + sizeof(UINT) + sizeof(UINT64));
+	const bool bAllocateOnStack = (FootprintsMemorySize < 4096);
+	void* const FootprintsMemory = bAllocateOnStack ? FMemory_Alloca(FootprintsMemorySize) : FMemory::Malloc(FootprintsMemorySize);
+	ON_SCOPE_EXIT
 	{
-		// Make sure the resource is in the requested state
-		ReconcileInitialState(RHICmdList, CreateResult.Texture, CreateResult.CreateState, CreateResult.DesiredState);
+		if (!bAllocateOnStack)
+		{
+			FMemory::Free(FootprintsMemory);
+		}
+	};
+
+	CA_ASSUME(FootprintsMemory != nullptr);
+
+	D3D12_PLACED_SUBRESOURCE_FOOTPRINT* const Footprints = reinterpret_cast<D3D12_PLACED_SUBRESOURCE_FOOTPRINT*>(FootprintsMemory);
+	UINT* const Rows = reinterpret_cast<UINT*>(Footprints + NumSubresources);
+	UINT64* const RowSizeInBytes = reinterpret_cast<UINT64*>(Rows + NumSubresources);
+
+	FD3D12Device* Device = Texture->GetParentDevice();
+
+	const D3D12_RESOURCE_DESC& ResourceDesc = Texture->GetResource()->GetDesc();
+	Device->GetDevice()->GetCopyableFootprints(&ResourceDesc, 0, NumSubresources, 0, Footprints, Rows, RowSizeInBytes, nullptr);
+
+	check(UploadMemorySize >= BulkData->GetResourceBulkDataSize());
+
+	uint8* const DestinationDataStart = reinterpret_cast<uint8*>(UploadMemory);
+	const uint8* const SourceDataStart = reinterpret_cast<const uint8*>(BulkData->GetResourceBulkData());
+
+	const uint8* SrcData = SourceDataStart;
+
+	for (uint32 Subresource = 0; Subresource < NumSubresources; Subresource++)
+	{
+		uint8* DstData = DestinationDataStart + Footprints[Subresource].Offset;
+
+		const uint32 NumRows = Rows[Subresource] * Footprints[Subresource].Footprint.Depth;
+		const uint32 SrcRowPitch = RowSizeInBytes[Subresource];
+		const uint32 DstRowPitch = Footprints[Subresource].Footprint.RowPitch;
+
+		// If src and dst pitch are aligned, which is typically the case for the bulk of the data (most large mips, POT textures), we can use a single large memcpy()
+		if (SrcRowPitch == DstRowPitch)
+		{
+			FMemory::Memcpy(DstData, SrcData, SrcRowPitch * NumRows);
+			SrcData += SrcRowPitch * NumRows;
+		}
+		else
+		{
+			for (uint32 Row = 0; Row < NumRows; ++Row)
+			{
+				FMemory::Memcpy(DstData, SrcData, SrcRowPitch);
+
+				SrcData += SrcRowPitch;
+				DstData += DstRowPitch;
+			}
+		}
 	}
 
-	return CreateResult.Texture;
+	check(SrcData == SourceDataStart + BulkData->GetResourceBulkDataSize());
+}
+
+// GetCopyableFootprints gives us offsets in reference to the start of the subresource we ask for, so we have to ask for ALL subresources before that one so that we can get the actual offset to the subresource
+D3D12_PLACED_SUBRESOURCE_FOOTPRINT UE::D3D12RHI::TextureUtils::GetInitializerSubresourceFootprint(FD3D12Texture* Texture, uint32 Subresource, uint32& NumRows)
+{
+	FD3D12Device* Device = Texture->GetParentDevice();
+	const FD3D12ResourceDesc& ResourceDesc = Texture->GetResource()->GetDesc();
+
+	const uint32 MaxSubresources = CalculateSubresourceCount(Texture->GetDesc());
+	const uint32 SubresourceIndex = Subresource < MaxSubresources ? Subresource : (MaxSubresources - 1);
+
+	const uint32 NumSubresources = SubresourceIndex + 1;
+
+	const size_t FootprintsMemorySize = NumSubresources * (sizeof(D3D12_PLACED_SUBRESOURCE_FOOTPRINT) + sizeof(UINT));
+	const bool bAllocateOnStack = (FootprintsMemorySize < 4096);
+	void* const FootprintsMemory = bAllocateOnStack ? FMemory_Alloca(FootprintsMemorySize) : FMemory::Malloc(FootprintsMemorySize);
+	ON_SCOPE_EXIT
+	{
+		if (!bAllocateOnStack)
+		{
+			FMemory::Free(FootprintsMemory);
+		}
+	};
+
+	CA_ASSUME(FootprintsMemory != nullptr);
+
+	D3D12_PLACED_SUBRESOURCE_FOOTPRINT* const Footprints = reinterpret_cast<D3D12_PLACED_SUBRESOURCE_FOOTPRINT*>(FootprintsMemory);
+	UINT* const Rows = reinterpret_cast<UINT*>(Footprints + NumSubresources);
+
+	Device->GetDevice()->GetCopyableFootprints(&ResourceDesc, 0, NumSubresources, 0, Footprints, Rows, nullptr, nullptr);
+
+	NumRows = Rows[SubresourceIndex];
+	return Footprints[SubresourceIndex];
+}
+
+static UE::RHICore::FBaseTextureInitializerImplementation CreateTextureInitializerForWriting(FRHICommandListBase& RHICmdList, const FD3D12DynamicRHI::FCreateTextureInternalResult& CreateResult)
+{
+	FD3D12Texture* Texture = CreateResult.Texture;
+
+	const uint64 TextureSize = UE::D3D12RHI::TextureUtils::CalculateResourceSize(Texture);
+
+	FD3D12ResourceLocation UploadLocation(Texture->GetParentDevice());
+	Texture->GetParentDevice()->GetDefaultFastAllocator().Allocate(TextureSize, D3D12_TEXTURE_DATA_PLACEMENT_ALIGNMENT, &UploadLocation);
+
+	void* WritableData = UploadLocation.GetMappedBaseAddress();
+
+	return UE::RHICore::FBaseTextureInitializerImplementation(RHICmdList, Texture, WritableData, TextureSize,
+		[Texture = TRefCountPtr<FD3D12Texture>(Texture), CreateState = CreateResult.CreateState, DesiredState = CreateResult.DesiredState, UploadLocation = MoveTemp(UploadLocation)](FRHICommandListBase& RHICmdList) mutable
+		{
+			UE::D3D12RHI::TextureUtils::ReconcileInitialState(RHICmdList, Texture, CreateState, D3D12_RESOURCE_STATE_COPY_DEST);
+
+			Texture->UploadInitialData(RHICmdList, MoveTemp(UploadLocation), DesiredState);
+			return TRefCountPtr<FRHITexture>(MoveTemp(Texture));
+		},
+		[Texture = TRefCountPtr<FD3D12Texture>(Texture), WritableData](FRHITextureInitializer::FSubresourceIndex SubresourceIndex)
+		{
+			const FD3D12ResourceDesc& ResourceDesc = Texture->GetResource()->GetDesc();
+
+			const uint32 ArrayOffset = UE::RHICore::GetCombinedArrayIndex(Texture->GetDesc(), SubresourceIndex.FaceIndex, SubresourceIndex.ArrayIndex);
+			const uint32 Subresource = CalcSubresource(SubresourceIndex.MipIndex, ArrayOffset, ResourceDesc.MipLevels);
+
+			uint32 NumRows = 0;
+			const D3D12_PLACED_SUBRESOURCE_FOOTPRINT PlacedFootprint = UE::D3D12RHI::TextureUtils::GetInitializerSubresourceFootprint(Texture, Subresource, NumRows);
+
+			FRHITextureSubresourceInitializer Result{};
+
+			Result.Data = reinterpret_cast<uint8*>(WritableData) + PlacedFootprint.Offset;
+			Result.Stride = PlacedFootprint.Footprint.RowPitch;
+			Result.Size = PlacedFootprint.Footprint.RowPitch * NumRows * PlacedFootprint.Footprint.Depth;
+
+			return Result;
+		}
+	);
 }
 
-FTextureRHIRef FD3D12DynamicRHI::RHICreateTexture(FRHICommandListBase& RHICmdList, const FRHITextureCreateDesc& CreateDesc)
+uint64 CalculateTextureSize(const FRHITextureDesc& Desc)
 {
-	return CreateD3D12Texture(RHICmdList, CreateDesc);
+	const FPixelFormatInfo& PixelFormat = GPixelFormats[Desc.Format];
+
+	const uint64 BlockSizeX = FMath::DivideAndRoundUp<uint64>(Desc.Extent.X, PixelFormat.BlockSizeX);
+	const uint64 BlockSizeY = FMath::DivideAndRoundUp<uint64>(Desc.Extent.Y, PixelFormat.BlockSizeY);
+	const uint64 BlockSizeZ = FMath::DivideAndRoundUp<uint64>(Desc.Depth,    PixelFormat.BlockSizeZ);
+	const uint64 BlockCount = BlockSizeX * BlockSizeY * BlockSizeZ;
+	const uint64 ResourceSize = BlockCount * PixelFormat.BlockBytes;
+	const uint64 ResourceCount = Desc.GetSubresourceCount();
+
+	return ResourceSize * ResourceCount;
+}
+
+FRHITextureInitializer FD3D12DynamicRHI::RHICreateTextureInitializer(FRHICommandListBase& RHICmdList, const FRHITextureCreateDesc& CreateDesc)
+{
+	TRACE_CPUPROFILER_EVENT_SCOPE(D3D12RHI::RHICreateTextureInitializer);
+	LLM_SCOPE_DYNAMIC_STAT_OBJECTPATH_FNAME(CreateDesc.OwnerName, ELLMTagSet::Assets);
+	LLM_SCOPE_DYNAMIC_STAT_OBJECTPATH_FNAME(CreateDesc.GetTraceClassName(), ELLMTagSet::AssetClasses);
+	UE_TRACE_METADATA_SCOPE_ASSET_FNAME(CreateDesc.DebugName, CreateDesc.GetTraceClassName(), CreateDesc.OwnerName);
+
+	FCreateTextureInternalResult CreateResult = CreateTextureInternal(CreateDesc, nullptr);
+
+	if (CreateDesc.InitAction == ERHITextureInitAction::Default)
+	{
+		// Just return the texture with its default contents
+		return UE::RHICore::FBaseTextureInitializerImplementation(RHICmdList, CreateResult.Texture,
+			[Texture = TRefCountPtr<FD3D12Texture>(CreateResult.Texture), CreateState = CreateResult.CreateState, DesiredState = CreateResult.DesiredState](FRHICommandListBase& RHICmdList) mutable
+			{
+				UE::D3D12RHI::TextureUtils::ReconcileInitialState(RHICmdList, Texture, CreateState, DesiredState);
+				return TRefCountPtr<FRHITexture>(MoveTemp(Texture));
+			}
+		);
+	}
+
+	if (CreateDesc.InitAction == ERHITextureInitAction::BulkData)
+	{
+		check(CreateDesc.BulkData);
+
+		UE::RHICore::FBaseTextureInitializerImplementation Initializer(CreateTextureInitializerForWriting(RHICmdList, CreateResult));
+
+		UE::D3D12RHI::TextureUtils::CopyBulkData(Initializer.GetWritableData(), Initializer.GetWritableSize(), CreateResult.Texture, CreateDesc.BulkData);
+
+		// Discard the bulk data's contents.
+		CreateDesc.BulkData->Discard();
+
+		return MoveTemp(Initializer);
+	}
+
+	if (CreateDesc.InitAction == ERHITextureInitAction::Initializer)
+	{
+		return CreateTextureInitializerForWriting(RHICmdList, CreateResult);
+	}
+
+	return UE::RHICore::HandleUnknownTextureInitializerInitAction(RHICmdList, CreateDesc);
 }
 
 class FWaitInitialMipDataUploadTask
@@ -1302,7 +1504,10 @@ FTextureRHIRef FD3D12DynamicRHI::AsyncReallocateTexture2D_RenderThread(FRHIComma
 	CreateDesc.SetOwnerName(OldTexture->GetOwnerName());
 
 	// Allocate a new texture.
-	FD3D12Texture* NewTexture = CreateD3D12Texture(RHICmdList, CreateDesc);
+	const FCreateTextureInternalResult CreateResult = CreateTextureInternal(CreateDesc, nullptr);
+	UE::D3D12RHI::TextureUtils::ReconcileInitialState(RHICmdList, CreateResult.Texture, CreateResult.CreateState, CreateResult.DesiredState);
+
+	FD3D12Texture* NewTexture = CreateResult.Texture;
 
 	RHICmdList.EnqueueLambda([
 		RootOldTexture = OldTexture,
@@ -1917,74 +2122,12 @@ void FD3D12Texture::CopyTextureRegion(FD3D12CommandContext& Context, uint32 Dest
 	Context.UpdateResidency(GetResource());
 }
 
-void FD3D12Texture::InitializeTextureData(FRHICommandListBase& RHICmdList, const FRHITextureCreateDesc& CreateDesc, D3D12_RESOURCE_STATES DestinationState)
+void FD3D12Texture::UploadInitialData(FRHICommandListBase& RHICmdList, FD3D12ResourceLocation&& SourceLocation, D3D12_RESOURCE_STATES DestinationState)
 {
-	// each mip of each array slice counts as a subresource
-	uint16 ArraySize = CreateDesc.IsTextureArray() ? CreateDesc.ArraySize : 1;
-	if (CreateDesc.IsTextureCube())
+	RHICmdList.EnqueueLambda([Texture = this, SourceLocation = MoveTemp(SourceLocation), DestinationState](FRHICommandListBase& ExecutingCmdList)
 	{
-		ArraySize *= 6;
-	}
-	uint32 NumSubresources = CreateDesc.NumMips * ArraySize;
-
-	FD3D12Device* Device = GetParentDevice();
-
-	size_t MemSize = NumSubresources * (sizeof(D3D12_PLACED_SUBRESOURCE_FOOTPRINT) + sizeof(UINT) + sizeof(UINT64));
-	const bool bAllocateOnStack = (MemSize < 4096);
-	void* Mem = bAllocateOnStack ? FMemory_Alloca(MemSize) : FMemory::Malloc(MemSize);
-	ON_SCOPE_EXIT
-	{
-		if (!bAllocateOnStack)
-		{
-			FMemory::Free(Mem);
-		}
-	};
-
-	D3D12_PLACED_SUBRESOURCE_FOOTPRINT* Footprints = (D3D12_PLACED_SUBRESOURCE_FOOTPRINT*) Mem;
-	check(Footprints);
-	UINT* Rows = (UINT*) (Footprints + NumSubresources);
-	check(Rows);
-	UINT64* RowSizeInBytes = (UINT64*) (Rows + NumSubresources);
-	check(RowSizeInBytes);
-
-	uint64 Size = 0;
-	const D3D12_RESOURCE_DESC& Desc = GetResource()->GetDesc();
-	Device->GetDevice()->GetCopyableFootprints(&Desc, 0, NumSubresources, 0, Footprints, Rows, RowSizeInBytes, &Size);
-
-	FD3D12ResourceLocation SrcResourceLoc(Device);
-	uint8* DstDataBase = (uint8*) Device->GetDefaultFastAllocator().Allocate(Size, D3D12_TEXTURE_DATA_PLACEMENT_ALIGNMENT, &SrcResourceLoc);
-
-	const uint8* SrcData = (const uint8*)CreateDesc.BulkData->GetResourceBulkData();
-	for (uint32 Subresource = 0; Subresource < NumSubresources; Subresource++)
-	{
-		uint8* DstData = DstDataBase + Footprints[Subresource].Offset;
-
-		const uint32 NumRows = Rows[Subresource] * Footprints[Subresource].Footprint.Depth;
-		const uint32 SrcRowPitch = RowSizeInBytes[Subresource];
-		const uint32 DstRowPitch = Footprints[Subresource].Footprint.RowPitch;
+		const uint32 NumSubresources = CalculateSubresourceCount(Texture->GetDesc());
 
-		// If src and dst pitch are aligned, which is typically the case for the bulk of the data (most large mips, POT textures), we can use a single large memcpy()
-		if (SrcRowPitch == DstRowPitch)
-		{
-			memcpy(DstData, SrcData, SrcRowPitch * NumRows);
-			SrcData += SrcRowPitch * NumRows;
-		}
-		else
-		{
-			for (uint32 Row = 0; Row < NumRows; ++Row)
-			{
-				memcpy(DstData, SrcData, SrcRowPitch);
-
-				SrcData += SrcRowPitch;
-				DstData += DstRowPitch;
-			}
-		}
-	}
-
-	check(SrcData == (uint8*)CreateDesc.BulkData->GetResourceBulkData() + CreateDesc.BulkData->GetResourceBulkDataSize());
-
-	RHICmdList.EnqueueLambda([Texture = this, SrcResourceLoc = MoveTemp(SrcResourceLoc), NumSubresources, DestinationState](FRHICommandListBase& ExecutingCmdList)
-	{
 		size_t MemSize = NumSubresources * (sizeof(D3D12_PLACED_SUBRESOURCE_FOOTPRINT) + sizeof(UINT) + sizeof(UINT64));
 		const bool bAllocateOnStack = (MemSize < 4096);
 		void* Mem = bAllocateOnStack ? FMemory_Alloca(MemSize) : FMemory::Malloc(MemSize);
@@ -1996,17 +2139,17 @@ void FD3D12Texture::InitializeTextureData(FRHICommandListBase& RHICmdList, const
 			}
 		};
 
-		D3D12_PLACED_SUBRESOURCE_FOOTPRINT* Footprints = (D3D12_PLACED_SUBRESOURCE_FOOTPRINT*)Mem;
+		D3D12_PLACED_SUBRESOURCE_FOOTPRINT * Footprints = (D3D12_PLACED_SUBRESOURCE_FOOTPRINT*)Mem;
 		check(Footprints);
-		UINT* Rows = (UINT*)(Footprints + NumSubresources);
-		UINT64* RowSizeInBytes = (UINT64*)(Rows + NumSubresources);
+		UINT * Rows = (UINT*)(Footprints + NumSubresources);
+		UINT64 * RowSizeInBytes = (UINT64*)(Rows + NumSubresources);
 
 		uint64 Size = 0;
-		const D3D12_RESOURCE_DESC& Desc = Texture->GetResource()->GetDesc();
-		Texture->GetParentDevice()->GetDevice()->GetCopyableFootprints(&Desc, 0, NumSubresources, SrcResourceLoc.GetOffsetFromBaseOfResource(), Footprints, Rows, RowSizeInBytes, &Size);
+		const D3D12_RESOURCE_DESC & Desc = Texture->GetResource()->GetDesc();
+		Texture->GetParentDevice()->GetDevice()->GetCopyableFootprints(&Desc, 0, NumSubresources, SourceLocation.GetOffsetFromBaseOfResource(), Footprints, Rows, RowSizeInBytes, &Size);
 
 		D3D12_TEXTURE_COPY_LOCATION Src;
-		Src.pResource = SrcResourceLoc.GetResource()->GetResource();
+		Src.pResource = SourceLocation.GetResource()->GetResource();
 		Src.Type = D3D12_TEXTURE_COPY_TYPE_PLACED_FOOTPRINT;
 
 		// Initialize all the textures in the chain
diff --git a/Engine/Source/Runtime/D3D12RHI/Private/D3D12Texture.h b/Engine/Source/Runtime/D3D12RHI/Private/D3D12Texture.h
--- a/Engine/Source/Runtime/D3D12RHI/Private/D3D12Texture.h
+++ b/Engine/Source/Runtime/D3D12RHI/Private/D3D12Texture.h
@@ -69,7 +69,14 @@ public:
 
 	void GetReadBackHeapDesc(D3D12_PLACED_SUBRESOURCE_FOOTPRINT& OutFootprint, uint32 Subresource) const;
 
+	bool HasCreatedRTVsPerSlice() const
+	{
+		return bCreatedRTVsPerSlice;
+	}
+
 	bool HasRenderTargetViews() const							{ return (RenderTargetViews.Num() > 0); }
+	TConstArrayView<TSharedPtr<FD3D12RenderTargetView>> GetRenderTargetViews() const { return TConstArrayView<TSharedPtr<FD3D12RenderTargetView>>(RenderTargetViews); }
+
 	FD3D12ShaderResourceView* GetShaderResourceView() const		{ return ShaderResourceView.Get(); }
 	FD3D12RenderTargetView* GetRenderTargetView(int32 MipIndex, int32 ArraySliceIndex) const;
 	FD3D12DepthStencilView* GetDepthStencilView(FExclusiveDepthStencil AccessType) const { return DepthStencilViews[AccessType.GetIndex()].Get(); }
@@ -79,7 +86,7 @@ public:
 #endif // #if PLATFORM_REQUIRES_TYPELESS_RESOURCE_DISCARD_WORKAROUND
 		
 	// Setup functionality
-	void InitializeTextureData(FRHICommandListBase& RHICmdList, const FRHITextureCreateDesc& CreateDesc, D3D12_RESOURCE_STATES DestinationState);
+	void UploadInitialData(FRHICommandListBase& RHICmdList, FD3D12ResourceLocation&& SourceLocation, D3D12_RESOURCE_STATES DestinationState);
 	void CreateViews(FD3D12Texture* FirstLinkedObject);
 	void SetCreatedRTVsPerSlice(bool Value, int32 InRTVArraySize)
 	{
@@ -262,3 +269,15 @@ namespace FD3D12TextureStats
 	void D3D12TextureAllocated(FD3D12Texture& Texture);
 	void D3D12TextureDeleted(FD3D12Texture& Texture);
 };
+
+namespace UE::D3D12RHI::TextureUtils
+{
+	// Fills in the ClearValue argument based on the ResourceDesc and CreateDesc flags.
+	D3D12_CLEAR_VALUE* FillClearValue(D3D12_CLEAR_VALUE& ClearValue, const FD3D12ResourceDesc& ResourceDesc, const FRHITextureDesc& TextureDesc);
+
+	void ReconcileInitialState(FRHICommandListBase& RHICmdList, FD3D12Texture* Texture, D3D12_RESOURCE_STATES CurrentState, D3D12_RESOURCE_STATES NeededState);
+
+	uint64 CalculateResourceSize(FD3D12Texture* Texture);
+	void CopyBulkData(void* UploadMemory, uint64 UploadMemorySize, FD3D12Texture* Texture, FResourceBulkDataInterface* BulkData);
+	D3D12_PLACED_SUBRESOURCE_FOOTPRINT GetInitializerSubresourceFootprint(FD3D12Texture* Texture, uint32 Subresource, uint32& NumRows);
+}
diff --git a/Engine/Source/Runtime/Engine/Private/PrecomputedVolumetricLightmap.cpp b/Engine/Source/Runtime/Engine/Private/PrecomputedVolumetricLightmap.cpp
--- a/Engine/Source/Runtime/Engine/Private/PrecomputedVolumetricLightmap.cpp
+++ b/Engine/Source/Runtime/Engine/Private/PrecomputedVolumetricLightmap.cpp
@@ -27,7 +27,7 @@ void FVolumetricLightmapDataLayer::CreateTexture(FIntVector Dimensions)
 	const FRHITextureCreateDesc Desc =
 		FRHITextureCreateDesc::Create3D(TEXT("VolumetricLightmap"), Dimensions, Format)
 		.SetFlags(ETextureCreateFlags::ShaderResource | ETextureCreateFlags::UAV)
-		.SetBulkData(this)
+		.SetInitActionBulkData(this)
 		.SetClassName(VolumetricLightmapDataLayerName);
 
 	Texture = RHICreateTexture(Desc);
diff --git a/Engine/Source/Runtime/Engine/Private/Rendering/Texture2DResource.cpp b/Engine/Source/Runtime/Engine/Private/Rendering/Texture2DResource.cpp
--- a/Engine/Source/Runtime/Engine/Private/Rendering/Texture2DResource.cpp
+++ b/Engine/Source/Runtime/Engine/Private/Rendering/Texture2DResource.cpp
@@ -117,15 +117,19 @@ void FTexture2DResource::CreateTexture()
 	const FTexture2DMipMap* RequestedMip = GetPlatformMip(RequestedMipIdx);
 
 	// create texture with ResourceMem data when available
-	const FRHITextureCreateDesc Desc =
+	FRHITextureCreateDesc Desc =
 		FRHITextureCreateDesc::Create2D(TEXT("FTexture2DResource"), RequestedMip->SizeX, RequestedMip->SizeY, PixelFormat)
 		.SetNumMips(State.NumRequestedLODs)
 		.SetFlags(CreationFlags)
 		.SetExtData(PlatformData->GetExtData())
-		.SetBulkData(ResourceMem)
 		.SetClassName(Texture2DResourceName)
 		.SetOwnerName(GetOwnerName());
 
+	if (ResourceMem)
+	{
+		Desc.SetInitActionBulkData(ResourceMem);
+	}
+
 	TextureRHI = RHICreateTexture(Desc);
 
 	// if( ResourceMem && !State.bReadyForStreaming) // To investigate!
@@ -171,7 +175,6 @@ void FTexture2DResource::CreatePartiallyResidentTexture()
 		.SetNumMips(State.MaxNumLODs)
 		.SetFlags(CreationFlags | ETextureCreateFlags::Virtual)
 		.SetExtData(PlatformData->GetExtData())
-		.SetBulkData(ResourceMem)
 		.SetClassName(Texture2DResourceName)
 		.SetOwnerName(GetOwnerName());
 
@@ -237,22 +240,22 @@ uint64 FTexture2DResource::GetPlatformMipsSize(uint32 NumMips) const
 	return 0;
 }
 
-uint32 FTexture2DResource::CalculateTightPackedMipSize(int32 MipSizeX,int32 MipSizeY,EPixelFormat PixelFormat,
-	uint32 & OutPitch)
+uint32 FTexture2DResource::CalculateTightPackedMipSize(int32 MipSizeX, int32 MipSizeY, EPixelFormat PixelFormat, uint32& OutPitch)
 {
-	const uint32 BlockSizeX = GPixelFormats[PixelFormat].BlockSizeX;		// Block width in pixels
-	const uint32 BlockSizeY = GPixelFormats[PixelFormat].BlockSizeY;		// Block height in pixels
-	const uint32 BlockBytes = GPixelFormats[PixelFormat].BlockBytes;
-	uint32 NumColumns		= (MipSizeX + BlockSizeX - 1) / BlockSizeX;	// Num-of columns in the source data (in blocks)
-	uint32 NumRows			= (MipSizeY + BlockSizeY - 1) / BlockSizeY;	// Num-of rows in the source data (in blocks)
-	if ( PixelFormat == PF_PVRTC2 || PixelFormat == PF_PVRTC4 )
+	const FPixelFormatInfo& Format = GPixelFormats[PixelFormat];
+
+	uint32 NumColumns = FMath::DivideAndRoundUp(MipSizeX, Format.BlockSizeX); // Num-of columns in the source data (in blocks)
+	uint32 NumRows    = FMath::DivideAndRoundUp(MipSizeY, Format.BlockSizeY); // Num-of rows in the source data (in blocks)
+
+	if (PixelFormat == PF_PVRTC2 || PixelFormat == PF_PVRTC4)
 	{
 		// PVRTC has minimum 2 blocks width and height
 		NumColumns = FMath::Max<uint32>(NumColumns, 2);
 		NumRows = FMath::Max<uint32>(NumRows, 2);
 	}
-	const uint32 SrcPitch   = NumColumns * BlockBytes;						// Num-of bytes per row in the source data
-	const uint32 EffectiveSize = BlockBytes*NumColumns*NumRows;
+
+	const uint32 SrcPitch   = NumColumns * Format.BlockBytes; // Num-of bytes per row in the source data
+	const uint32 EffectiveSize = Format.BlockBytes * NumColumns * NumRows;
 
 	OutPitch = SrcPitch;
 	return EffectiveSize;
diff --git a/Engine/Source/Runtime/Engine/Private/Rendering/Texture3DResource.cpp b/Engine/Source/Runtime/Engine/Private/Rendering/Texture3DResource.cpp
--- a/Engine/Source/Runtime/Engine/Private/Rendering/Texture3DResource.cpp
+++ b/Engine/Source/Runtime/Engine/Private/Rendering/Texture3DResource.cpp
@@ -136,9 +136,11 @@ void FTexture3DResource::CreateTexture()
 		if (GUseTexture3DBulkDataRHI)
 		{
 			InitialData.MergeMips(State.MaxNumLODs);
-			Desc.SetBulkData(&InitialData);
+			Desc.SetInitActionBulkData(&InitialData);
 		}
 
+		Desc.DetermineInititialState();
+
 		TextureRHI = RHICreateTexture(Desc);
 	}
 
diff --git a/Engine/Source/Runtime/Engine/Private/Slate/SlateTextures.cpp b/Engine/Source/Runtime/Engine/Private/Slate/SlateTextures.cpp
--- a/Engine/Source/Runtime/Engine/Private/Slate/SlateTextures.cpp
+++ b/Engine/Source/Runtime/Engine/Private/Slate/SlateTextures.cpp
@@ -40,7 +40,9 @@ void FSlateTexture2DRHIRef::Cleanup()
 	BeginCleanup(this);
 }
 
-void FSlateTexture2DRHIRef::InitRHI(FRHICommandListBase&)
+const FName NAME_FSlateTexture2DRHIRef(TEXT("FSlateTexture2DRHIRef"));
+
+void FSlateTexture2DRHIRef::InitRHI(FRHICommandListBase& RHICmdList)
 {
 	SCOPED_LOADTIMER(FSlateTexture2DRHIRef_InitDynamicRHI);
 
@@ -48,52 +50,63 @@ void FSlateTexture2DRHIRef::InitRHI(FRHICommandListBase&)
 
 	if( Width > 0 && Height > 0 )
 	{
-		if( TextureData.IsValid() || bCreateEmptyTexture )
+		const bool bValid = TextureData.IsValid();
+		const bool bHasUploadData = TextureData.IsValid() && TextureData->GetRawBytes().Num() > 0;
+
+		if (bValid || bHasUploadData || bCreateEmptyTexture)
 		{
-			check( !IsValidRef( ShaderResource) );
+			check(!IsValidRef(ShaderResource));
 
-			const FRHITextureCreateDesc Desc =
+			FRHITextureCreateDesc CreateDesc =
 				FRHITextureCreateDesc::Create2D(TEXT("FSlateTexture2DRHIRef"), Width, Height, PixelFormat)
 				.SetFlags(TexCreateFlags)
-				.SetClassName(TEXT("FSlateTexture2DRHIRef"));
+				.SetClassName(NAME_FSlateTexture2DRHIRef);
 
-			ShaderResource = RHICreateTexture(Desc);
-			check( IsValidRef( ShaderResource ) );
+			if (bHasUploadData)
+			{
+				CreateDesc.SetInitActionInitializer();
+			}
 
-			INC_MEMORY_STAT_BY(STAT_SlateTextureGPUMemory, Width*Height*GPixelFormats[PixelFormat].BlockBytes);
-		}
+			FRHITextureInitializer Initializer = RHICmdList.CreateTextureInitializer(CreateDesc);
 
-		if( TextureData.IsValid() && TextureData->GetRawBytes().Num() > 0 )
-		{
-			check(Width == TextureData->GetWidth());
-			check(Height == TextureData->GetHeight());
+			if (bHasUploadData)
+			{
+				check(Width == TextureData->GetWidth());
+				check(Height == TextureData->GetHeight());
 
-			uint32 Stride;
-			uint8* DestTextureData = (uint8*)RHILockTexture2D(ShaderResource, 0, RLM_WriteOnly, Stride, false);
-			const uint8* SourceTextureData = TextureData->GetRawBytes().GetData();
+				const FRHITextureSubresourceInitializer Subresource = Initializer.GetTexture2DSubresource(0);
 
-			const uint32 BlocksX = GPixelFormats[PixelFormat].GetBlockCountForWidth(Width);
-			const uint32 BlocksY = GPixelFormats[PixelFormat].GetBlockCountForHeight(Height);
-			const uint32 DataStride = BlocksX * GPixelFormats[PixelFormat].BlockBytes;
+				const uint8* SourceTextureData = TextureData->GetRawBytes().GetData();
 
-			checkf((uint32)TextureData->GetRawBytes().Num() >= DataStride * BlocksY, TEXT("Not enough bytes in source TextureData to complete copy operation"));
+				const uint32 BlocksX = GPixelFormats[PixelFormat].GetBlockCountForWidth(Width);
+				const uint32 BlocksY = GPixelFormats[PixelFormat].GetBlockCountForHeight(Height);
+				const uint32 SourceDataStride = BlocksX * GPixelFormats[PixelFormat].BlockBytes;
 
-			if (Stride == DataStride)
-			{
-				FMemory::Memcpy(DestTextureData, SourceTextureData, Stride * BlocksY);
-			}
-			else
-			{
-				checkf(DataStride < Stride, TEXT("Texture stride of %u is smaller than source data stride of %u, PixelFormat=%s (%d)"), Stride, DataStride, GPixelFormats[PixelFormat].Name, (int32)PixelFormat);
-				for (uint32 i = 0; i < BlocksY; i++)
+				checkf((uint32)TextureData->GetRawBytes().Num() >= SourceDataStride * BlocksY, TEXT("Not enough bytes in source TextureData to complete copy operation"));
+
+				if (Subresource.Stride == SourceDataStride)
 				{
-					FMemory::Memcpy(DestTextureData, SourceTextureData, DataStride);
-					DestTextureData += Stride;
-					SourceTextureData += DataStride;
+					FMemory::Memcpy(Subresource.Data, SourceTextureData, Subresource.Stride * BlocksY);
 				}
+				else
+				{
+					uint8* WritePointer = reinterpret_cast<uint8*>(Subresource.Data);
+
+					checkf(SourceDataStride < Subresource.Stride, TEXT("Texture stride of %u is smaller than source data stride of %u, PixelFormat=%s (%d)"), Subresource.Stride, SourceDataStride, GPixelFormats[PixelFormat].Name, (int32)PixelFormat);
+					for (uint32 i = 0; i < BlocksY; i++)
+					{
+						FMemory::Memcpy(WritePointer, SourceTextureData, SourceDataStride);
+						WritePointer += Subresource.Stride;
+						SourceTextureData += SourceDataStride;
+					}
+				}
+				TextureData->Empty();
 			}
-			RHIUnlockTexture2D(ShaderResource, 0, false);
-			TextureData->Empty();
+
+			ShaderResource = Initializer.Finalize();
+			check(IsValidRef(ShaderResource));
+
+			INC_MEMORY_STAT_BY(STAT_SlateTextureGPUMemory, Width * Height * GPixelFormats[PixelFormat].BlockBytes);
 		}
 	}
 }
diff --git a/Engine/Source/Runtime/Engine/Private/Streaming/Texture2DMipAllocator_Virtual.cpp b/Engine/Source/Runtime/Engine/Private/Streaming/Texture2DMipAllocator_Virtual.cpp
--- a/Engine/Source/Runtime/Engine/Private/Streaming/Texture2DMipAllocator_Virtual.cpp
+++ b/Engine/Source/Runtime/Engine/Private/Streaming/Texture2DMipAllocator_Virtual.cpp
@@ -148,13 +148,16 @@ void FTexture2DMipAllocator_Virtual::DoConvertToVirtualWithNewMips(const FTextur
 			FTexture2DResource* Resource2D = Context.Resource->GetTexture2DResource();
 
 			// Create a copy of the texture that is a virtual texture.
-			const FRHITextureCreateDesc Desc =
+			FRHITextureCreateDesc CreateDesc =
 				FRHITextureCreateDesc::Create2D(TEXT("FTexture2DUpdate"), MipMap0.SizeX, MipMap0.SizeY, Context.Resource->GetPixelFormat())
 				.SetNumMips(ResourceState.MaxNumLODs)
-				.SetFlags(Context.Resource->GetCreationFlags() | ETextureCreateFlags::Virtual)
-				.SetBulkData(Resource2D->ResourceMem);
+				.SetFlags(Context.Resource->GetCreationFlags() | ETextureCreateFlags::Virtual);
+			if (Resource2D->ResourceMem)
+			{
+				CreateDesc.SetInitActionBulkData(Resource2D->ResourceMem);
+			}
 
-			IntermediateTextureRHI = RHICreateTexture(Desc);
+			IntermediateTextureRHI = RHICreateTexture(CreateDesc);
 
 			RHIVirtualTextureSetFirstMipInMemory(IntermediateTextureRHI, CurrentFirstLODIdx);
 			RHIVirtualTextureSetFirstMipVisible(IntermediateTextureRHI, CurrentFirstLODIdx);
diff --git a/Engine/Source/Runtime/Engine/Private/Streaming/Texture2DUpdate.cpp b/Engine/Source/Runtime/Engine/Private/Streaming/Texture2DUpdate.cpp
--- a/Engine/Source/Runtime/Engine/Private/Streaming/Texture2DUpdate.cpp
+++ b/Engine/Source/Runtime/Engine/Private/Streaming/Texture2DUpdate.cpp
@@ -94,11 +94,17 @@ void FTexture2DUpdate::DoConvertToVirtualWithNewMips(const FContext& Context)
 			ensure(!IntermediateTextureRHI);
 
 			// Create a copy of the texture that is a virtual texture.
-			const FRHITextureCreateDesc Desc =
+			FRHITextureCreateDesc Desc =
 				FRHITextureCreateDesc::Create2D(TEXT("FTexture2DUpdate"), MipMap0.SizeX, MipMap0.SizeY, Context.Resource->GetPixelFormat())
 				.SetNumMips(ResourceState.MaxNumLODs)
-				.SetFlags(Context.Resource->GetCreationFlags() | ETextureCreateFlags::Virtual)
-				.SetBulkData(Context.Resource->ResourceMem);
+				.SetFlags(Context.Resource->GetCreationFlags() | ETextureCreateFlags::Virtual);
+
+			if (Context.Resource->ResourceMem)
+			{
+				Desc.SetInitActionBulkData(Context.Resource->ResourceMem);
+			}
+
+			Desc.DetermineInititialState();
 
 			IntermediateTextureRHI = RHICreateTexture(Desc);
 
@@ -132,11 +138,17 @@ bool FTexture2DUpdate::DoConvertToNonVirtual(const FContext& Context)
 
 			ensure(!IntermediateTextureRHI);
 
-			const FRHITextureCreateDesc Desc =
+			FRHITextureCreateDesc Desc =
 				FRHITextureCreateDesc::Create2D(TEXT("FTexture2DUpdate"), PendingFirstMipMap.SizeX, PendingFirstMipMap.SizeY, Context.Resource->GetPixelFormat())
 				.SetNumMips(ResourceState.NumRequestedLODs)
-				.SetFlags(Context.Resource->GetCreationFlags())
-				.SetBulkData(Context.Resource->ResourceMem);
+				.SetFlags(Context.Resource->GetCreationFlags());
+
+			if (Context.Resource->ResourceMem)
+			{
+				Desc.SetInitActionBulkData(Context.Resource->ResourceMem);
+			}
+
+			Desc.DetermineInititialState();
 
 			IntermediateTextureRHI = RHICreateTexture(Desc);
 
diff --git a/Engine/Source/Runtime/Engine/Private/Streaming/VolumeTextureStreaming.cpp b/Engine/Source/Runtime/Engine/Private/Streaming/VolumeTextureStreaming.cpp
--- a/Engine/Source/Runtime/Engine/Private/Streaming/VolumeTextureStreaming.cpp
+++ b/Engine/Source/Runtime/Engine/Private/Streaming/VolumeTextureStreaming.cpp
@@ -88,9 +88,11 @@ bool FVolumeTextureMipAllocator_Reallocate::UploadMips(const FTextureUpdateConte
 
 		if (GUseTexture3DBulkDataRHI)
 		{
-			Desc.SetBulkData(&StreamedInMipData);
+			Desc.SetInitActionBulkData(&StreamedInMipData);
 		}
 
+		Desc.DetermineInititialState();
+
 		IntermediateTextureRHI = RHICreateTexture(Desc);
 	}
 
diff --git a/Engine/Source/Runtime/Engine/Private/Texture2D.cpp b/Engine/Source/Runtime/Engine/Private/Texture2D.cpp
--- a/Engine/Source/Runtime/Engine/Private/Texture2D.cpp
+++ b/Engine/Source/Runtime/Engine/Private/Texture2D.cpp
@@ -1677,14 +1677,19 @@ void FVirtualTexture2DResource::InitializeEditorResources(IVirtualTexture* InVir
 		if (bUseFallbackTexture)
 		{
 			// Fallback to simple 1x1 texture of the fallback color.
-			FRHITextureCreateDesc TextureDesc = FRHITextureCreateDesc::Create2D(*Name, 1, 1, EPixelFormat::PF_B8G8R8A8);
-			TextureDesc.SetFlags(TexCreateFlags);
-			
-			const FColor Color = VTData->LayerFallbackColors[0].ToFColor(true);
-			FResourceBulkDataArrayView BulkDataInterface(&Color, sizeof(FColor));
-			TextureDesc.SetBulkData(&BulkDataInterface);
+			const FRHITextureCreateDesc TextureDesc =
+				FRHITextureCreateDesc::Create2D(*Name, 1, 1, EPixelFormat::PF_B8G8R8A8)
+				.AddFlags(TexCreateFlags);
+
+			FRHICommandListImmediate& RHICmdList = FRHICommandListExecutor::GetImmediateCommandList();
+			FRHITextureInitializer Initializer = RHICmdList.CreateTextureInitializer(TextureDesc);
+
+			{
+				FRHITextureSubresourceInitializer Subresource = Initializer.GetTexture2DSubresource(0);
+				Subresource.WriteColor(VTData->LayerFallbackColors[0].ToFColor(true));
+			}
 
-			Texture2DRHI = RHICreateTexture(TextureDesc);
+			Texture2DRHI = Initializer.Finalize();
 		}
 		else
 		{
diff --git a/Engine/Source/Runtime/Engine/Private/VectorField.cpp b/Engine/Source/Runtime/Engine/Private/VectorField.cpp
--- a/Engine/Source/Runtime/Engine/Private/VectorField.cpp
+++ b/Engine/Source/Runtime/Engine/Private/VectorField.cpp
@@ -174,7 +174,7 @@ public:
 			const FRHITextureCreateDesc Desc =
 				FRHITextureCreateDesc::Create3D(TEXT("FVectorFieldStaticResource"), SizeX, SizeY, SizeZ, PF_FloatRGBA)
 				.SetFlags(ETextureCreateFlags::ShaderResource)
-				.SetBulkData(&BulkDataInterface);
+				.SetInitActionBulkData(&BulkDataInterface);
 
 			VolumeTextureRHI = RHICreateTexture(Desc);
 
diff --git a/Engine/Source/Runtime/NullDrv/Public/NullRHI.h b/Engine/Source/Runtime/NullDrv/Public/NullRHI.h
--- a/Engine/Source/Runtime/NullDrv/Public/NullRHI.h
+++ b/Engine/Source/Runtime/NullDrv/Public/NullRHI.h
@@ -7,6 +7,7 @@
 #include "RHI.h"
 #include "RHITypes.h"
 #include "Async/TaskGraphInterfaces.h"
+#include "RHITextureUtils.h"
 
 struct Rect;
 
@@ -206,9 +207,54 @@ public:
 		{}
 	};
 
-	virtual FTextureRHIRef RHICreateTexture(FRHICommandListBase&, const FRHITextureCreateDesc& CreateDesc) final override
+	struct FNullTextureInitializer : public FRHITextureInitializer
 	{
-		return new FNullTexture(CreateDesc);
+		FNullTextureInitializer(FRHICommandListBase& RHICmdList, FRHITexture* Texture, void* InWritableData, uint64 InWritableDataSize)
+			: FRHITextureInitializer(RHICmdList, Texture, InWritableData, InWritableDataSize,
+				[Texture = TRefCountPtr<FRHITexture>(Texture)](FRHICommandListBase&) mutable
+				{
+					return MoveTemp(Texture);
+				},
+				[Texture = TRefCountPtr<FRHITexture>(Texture), WritableData = InWritableData](FRHITextureInitializer::FSubresourceIndex SubresourceIndex)
+				{
+					const FRHITextureDesc& TextureDesc = Texture->GetDesc();
+
+					uint64 SubresourceStride = 0;
+					uint64 SubresourceSize = 0;
+					const uint64 Offset = UE::RHITextureUtils::CalculateSubresourceOffset(TextureDesc, SubresourceIndex.FaceIndex, SubresourceIndex.ArrayIndex, SubresourceIndex.MipIndex, SubresourceStride, SubresourceSize);
+
+					FRHITextureSubresourceInitializer Result{};
+					Result.Data = reinterpret_cast<uint8*>(WritableData) + Offset;
+					Result.Stride = SubresourceStride;
+					Result.Size = SubresourceSize;
+
+					return Result;
+				}
+			)
+		{
+		}
+	};
+
+	[[nodiscard]]
+	virtual FRHITextureInitializer RHICreateTextureInitializer(FRHICommandListBase& RHICmdList, const FRHITextureCreateDesc& CreateDesc)
+	{
+		if (CreateDesc.InitAction == ERHITextureInitAction::BulkData && CreateDesc.BulkData)
+		{
+			CreateDesc.BulkData->Discard();
+		}
+
+		void* WritableData = nullptr;
+		uint64 WritableDataSize = 0;
+
+		if (CreateDesc.InitAction == ERHITextureInitAction::Initializer)
+		{
+			WritableDataSize = UE::RHITextureUtils::CalculateTextureSize(CreateDesc);
+			WritableData = GetStaticBuffer(WritableDataSize);
+		}
+
+		FRHITexture* Texture = new FNullTexture(CreateDesc);
+
+		return FNullTextureInitializer(RHICmdList, Texture, WritableData, WritableDataSize);
 	}
 
 	virtual FTextureRHIRef RHIAsyncCreateTexture2D(uint32 SizeX, uint32 SizeY, uint8 Format, uint32 NumMips, ETextureCreateFlags Flags, ERHIAccess InResourceState, void** InitialMipData, uint32 NumInitialMips, const TCHAR* DebugName, FGraphEventRef& OutCompletionEvent) final override
diff --git a/Engine/Source/Runtime/OpenGLDrv/Private/Android/AndroidOpenGL.cpp b/Engine/Source/Runtime/OpenGLDrv/Private/Android/AndroidOpenGL.cpp
--- a/Engine/Source/Runtime/OpenGLDrv/Private/Android/AndroidOpenGL.cpp
+++ b/Engine/Source/Runtime/OpenGLDrv/Private/Android/AndroidOpenGL.cpp
@@ -538,7 +538,10 @@ FOpenGLTexture* PlatformCreateBuiltinBackBuffer(FOpenGLDynamicRHI* OpenGLRHI, ui
 			.SetFlags(ETextureCreateFlags::RenderTargetable | ETextureCreateFlags::Presentable | ETextureCreateFlags::ResolveTargetable)
 			.DetermineInititialState();
 
-		return new FOpenGLTexture(FRHICommandListImmediate::Get(), CreateDesc);
+		FOpenGLTexture* Texture = new FOpenGLTexture(CreateDesc);
+		Texture->Initialize(FRHICommandListImmediate::Get());
+
+		return Texture;
 	}
 	else
 	{
diff --git a/Engine/Source/Runtime/OpenGLDrv/Private/OpenGLDrv.h b/Engine/Source/Runtime/OpenGLDrv/Private/OpenGLDrv.h
--- a/Engine/Source/Runtime/OpenGLDrv/Private/OpenGLDrv.h
+++ b/Engine/Source/Runtime/OpenGLDrv/Private/OpenGLDrv.h
@@ -348,7 +348,7 @@ public:
 	virtual FRHICalcTextureSizeResult RHICalcTexturePlatformSize(FRHITextureDesc const& Desc, uint32 FirstMipIndex) final override;
 	virtual void RHIGetTextureMemoryStats(FTextureMemoryStats& OutStats) final override;
 	virtual bool RHIGetTextureMemoryVisualizeData(FColor* TextureData,int32 SizeX,int32 SizeY,int32 Pitch,int32 PixelSize) final override;
-	virtual FTextureRHIRef RHICreateTexture(FRHICommandListBase& RHICmdList, const FRHITextureCreateDesc& CreateDesc) final override;
+	[[nodiscard]] virtual FRHITextureInitializer RHICreateTextureInitializer(FRHICommandListBase& RHICmdList, const FRHITextureCreateDesc& CreateDesc) final override;
 	virtual FTextureRHIRef RHIAsyncCreateTexture2D(uint32 SizeX, uint32 SizeY, uint8 Format, uint32 NumMips, ETextureCreateFlags Flags, ERHIAccess InResourceState, void** InitialMipData, uint32 NumInitialMips, const TCHAR* DebugName, FGraphEventRef& OutCompletionEvent) final override;
 	virtual void RHIGenerateMips(FRHITexture* Texture) final override;
 	virtual uint32 RHIComputeMemorySize(FRHITexture* TextureRHI) final override;
diff --git a/Engine/Source/Runtime/OpenGLDrv/Private/OpenGLResources.h b/Engine/Source/Runtime/OpenGLDrv/Private/OpenGLResources.h
--- a/Engine/Source/Runtime/OpenGLDrv/Private/OpenGLResources.h
+++ b/Engine/Source/Runtime/OpenGLDrv/Private/OpenGLResources.h
@@ -945,12 +945,11 @@ class FOpenGLTexture : public FRHITexture, public FOpenGLViewableResource
 
 public:
 	// Standard constructor.
-	UE_API explicit FOpenGLTexture(FRHICommandListBase& RHICmdList, FOpenGLTextureCreateDesc const& CreateDesc);
+	UE_API explicit FOpenGLTexture(FOpenGLTextureCreateDesc const& CreateDesc);
 #if PLATFORM_ANDROID
 	// Constructor for external hardware buffer.
-	UE_API explicit FOpenGLTexture(FRHICommandListBase& RHICmdList, FOpenGLTextureCreateDesc const& CreateDesc, 
-		const AHardwareBuffer_Desc& HardwareBufferDesc, AHardwareBuffer* HardwareBuffer);
-#endif //PLATFORM_ANDROID
+	UE_API explicit FOpenGLTexture(FRHICommandListBase& RHICmdList, FOpenGLTextureCreateDesc const& CreateDesc, const AHardwareBuffer_Desc& HardwareBufferDesc, AHardwareBuffer* HardwareBuffer);
+#endif
 
 	// Constructor for external resources (RHICreateTexture2DFromResource etc).
 	UE_API explicit FOpenGLTexture(FOpenGLTextureCreateDesc const& CreateDesc, GLuint Resource);
@@ -958,10 +957,12 @@ public:
 	// Constructor for RHICreateAliasedTexture
 	enum EAliasConstructorParam { AliasResource };
 	UE_API explicit FOpenGLTexture(FOpenGLTexture& Texture, const FString& Name, EAliasConstructorParam);
-	UE_API void AliasResources(FOpenGLTexture& Texture);
 
 	UE_API virtual ~FOpenGLTexture();
 
+	UE_API void Initialize(FRHICommandListBase& RHICmdList);
+	UE_API void AliasResources(FOpenGLTexture& Texture);
+
 	GLuint GetResource()
 	{
 		TryRestoreGLResource();
@@ -1011,7 +1012,7 @@ public:
 	/**
 	* Returns the size of the memory block that is returned from Lock, threadsafe
 	*/
-	UE_API uint32 GetLockSize(uint32 MipIndex, uint32 ArrayIndex, EResourceLockMode LockMode, uint32& DestStride);
+	UE_API uint32 GetLockSize(uint32 MipIndex, EResourceLockMode LockMode, uint32& DestStride);
 
 	/** Unlocks a previously locked mip-map. */
 	UE_API void Unlock(uint32 MipIndex, uint32 ArrayIndex);
@@ -1082,7 +1083,7 @@ private:
 
 	UE_API void DeleteGLResource();
 
-	UE_API void Fill2DGLTextureImage(const FOpenGLTextureFormat& GLFormat, const bool bSRGB, uint32 MipIndex, const void* LockedBuffer, uint32 LockedSize, uint32 ArrayIndex);
+	UE_API void FillGLTextureImage(const FOpenGLTextureFormat& GLFormat, bool bSRGB, uint32 MipIndex, uint32 ArrayIndex, const FUintVector3& MipExtents, const void* BufferOrPBOOffset, size_t ImageSize);
 
 public:
 	uint32 GetEffectiveSizeZ() const
diff --git a/Engine/Source/Runtime/OpenGLDrv/Private/OpenGLTexture.cpp b/Engine/Source/Runtime/OpenGLDrv/Private/OpenGLTexture.cpp
--- a/Engine/Source/Runtime/OpenGLDrv/Private/OpenGLTexture.cpp
+++ b/Engine/Source/Runtime/OpenGLDrv/Private/OpenGLTexture.cpp
@@ -16,7 +16,9 @@
 #include "ProfilingDebugging/AssetMetadataTrace.h"
 #include "Engine/Texture.h"
 #include "RHICoreStats.h"
+#include "RHITextureUtils.h"
 #include "RHICoreTexture.h"
+#include "RHICoreTextureInitializer.h"
 
 #if PLATFORM_ANDROID
 #include "ThirdParty/Android/detex/AndroidETC.h"
@@ -83,6 +85,18 @@ static FAutoConsoleVariableRef CVarTextureEvictionMinLRUCapacity(
 	ECVF_RenderThreadSafe
 );
 
+struct FScopedPackAlignment
+{
+	FScopedPackAlignment()
+	{
+		glPixelStorei(GL_UNPACK_ALIGNMENT, 1);
+	}
+	~FScopedPackAlignment()
+	{
+		glPixelStorei(GL_UNPACK_ALIGNMENT, 4);
+	}
+};
+
 /*-----------------------------------------------------------------------------
 	Texture allocator support.
 -----------------------------------------------------------------------------*/
@@ -329,12 +343,12 @@ FOpenGLTexture::FOpenGLTexture(FOpenGLTextureCreateDesc const& CreateDesc, GLuin
 	, bDepthStencil      (CreateDesc.bDepthStencil)
 	, bAlias             (true)
 	, bMultisampleRenderbuffer(CreateDesc.bMultisampleRenderbuffer)
-{}
+{
+}
 
 #if PLATFORM_ANDROID
 // Constructor for external hardware buffer.
-FOpenGLTexture::FOpenGLTexture(FRHICommandListBase& RHICmdList, FOpenGLTextureCreateDesc const& CreateDesc, 
-	const AHardwareBuffer_Desc& HardwareBufferDesc, AHardwareBuffer* HardwareBuffer)
+FOpenGLTexture::FOpenGLTexture(FRHICommandListBase& RHICmdList, FOpenGLTextureCreateDesc const& CreateDesc, const AHardwareBuffer_Desc& HardwareBufferDesc, AHardwareBuffer* HardwareBuffer)
 	: FRHITexture(CreateDesc)
 	, Target(CreateDesc.Target)
 	, Attachment(CreateDesc.Attachment)
@@ -392,10 +406,10 @@ FOpenGLTexture::FOpenGLTexture(FRHICommandListBase& RHICmdList, FOpenGLTextureCr
 
 	UpdateTextureStats(this, true);
 }
-#endif//PLATFORM_ANDROID
+#endif
 
 // Standard constructor.
-FOpenGLTexture::FOpenGLTexture(FRHICommandListBase& RHICmdList, FOpenGLTextureCreateDesc const& CreateDesc)
+FOpenGLTexture::FOpenGLTexture(const FOpenGLTextureCreateDesc& CreateDesc)
 	: FRHITexture        (CreateDesc)
 	, Target             (CreateDesc.Target)
 	, Attachment         (CreateDesc.Attachment)
@@ -416,48 +430,9 @@ FOpenGLTexture::FOpenGLTexture(FRHICommandListBase& RHICmdList, FOpenGLTextureCr
 		EvictionParamsPtr = MakeUnique<FTextureEvictionParams>(CreateDesc.NumMips);
 	}
 
-	void* BulkDataPtr = nullptr;
-	uint64 BulkDataSize = 0;
-	bool bFreeBulkData = false;
-
-	if (CreateDesc.BulkData)
-	{
- 		if (RHICmdList.IsTopOfPipe())
-		{
-			// If bulk data is provided, and texture initialization is done by the RHI thread, it needs to be copied out of the FResourceBulkDataInterface.
-			// It is not safe to pass this pointer to the RHI thread, as the interface may have been stack allocated in the renderer.
-			// @todo: remove this memcpy when the new resource creation API is available.
-			BulkDataSize = CreateDesc.BulkData->GetResourceBulkDataSize();
-			BulkDataPtr = FMemory::Malloc(BulkDataSize);
-
-			FMemory::Memcpy(BulkDataPtr, CreateDesc.BulkData->GetResourceBulkData(), BulkDataSize);
-			bFreeBulkData = true;
-		}
-		else
-		{
-			// Otherwise, initialization will be done on this thread.
-			// Just use the raw pointer / size as-is.
-			BulkDataSize = CreateDesc.BulkData->GetResourceBulkDataSize();
-			BulkDataPtr = const_cast<void*>(CreateDesc.BulkData->GetResourceBulkData());
-		}
-	}
- 	RHICmdList.EnqueueLambda([this, BulkDataPtr, BulkDataSize, bFreeBulkData](FRHICommandListBase&)
-	{
-		FOpenGLDynamicRHI::Get().InitializeGLTexture(this, BulkDataPtr, BulkDataSize);
-		if (bFreeBulkData)
-		{
-			FMemory::Free(BulkDataPtr);
-		}
-	});
-
 	UpdateTextureStats(this, true);
 
 	PixelBuffers.AddZeroed(CreateDesc.NumMips * (bCubemap ? 6 : 1) * GetEffectiveSizeZ());
-
-	if (CreateDesc.BulkData)
-	{
-		CreateDesc.BulkData->Discard();
-	}
 }
 
 FOpenGLTexture::~FOpenGLTexture()
@@ -475,6 +450,14 @@ FOpenGLTexture::~FOpenGLTexture()
 	UpdateTextureStats(this, false);
 }
 
+void FOpenGLTexture::Initialize(FRHICommandListBase& RHICmdList)
+{
+	RHICmdList.EnqueueLambda([this](FRHICommandListBase&)
+	{
+		FOpenGLDynamicRHI::Get().InitializeGLTexture(this, nullptr, 0);
+	});
+}
+
 void FOpenGLTexture::DeleteGLResource()
 {
 	VERIFY_GL_SCOPE();
@@ -677,53 +660,35 @@ void FOpenGLDynamicRHI::InitializeGLTextureInternal(FOpenGLTexture* Texture, voi
 
 		if (Target != GL_TEXTURE_EXTERNAL_OES)
 		{
-			auto EnumerateSubresources = [&](void const* Data, TFunctionRef<bool(GLenum Target, uint32 SizeX, uint32 SizeY, uint32 SizeZ, uint32 MipIndex, uint32 ArraySlice, void const* MipSliceData, uint32 MipSliceSize)> Callback)
+			using TEnumerationCallback = TFunctionRef<bool(GLenum Target, FUintVector3 Extent, uint32 MipIndex, uint32 ArraySlice, const void* MipSliceData, uint64 MipSliceSize)>;
+
+			auto EnumerateSubresources = [&](void const* Data, TEnumerationCallback&& Callback)
 			{
-				struct FScopedPackAlignment
-				{
-					 FScopedPackAlignment() { glPixelStorei(GL_UNPACK_ALIGNMENT, 1); }
-					~FScopedPackAlignment() { glPixelStorei(GL_UNPACK_ALIGNMENT, 4); }
-				} PackAlignment;
+				FScopedPackAlignment PackAlignment;
 
-				uint64 DataOffset = 0;
+				const FPixelFormatInfo& PixelFormat = GPixelFormats[Desc.Format];
 
-				for (uint32 MipIndex = 0; MipIndex < Desc.NumMips; ++MipIndex)
+				for (uint32 ArraySlice = 0; ArraySlice < Desc.ArraySize; ++ArraySlice)
 				{
-					for (uint32 ArraySlice = 0; ArraySlice < Desc.ArraySize; ++ArraySlice)
+					for (uint32 FaceIndex = 0; FaceIndex < (Desc.IsTextureCube() ? 6u : 1u); ++FaceIndex)
 					{
-						for (uint32 FaceIndex = 0; FaceIndex < (Desc.IsTextureCube() ? 6u : 1u); ++FaceIndex)
+						for (uint32 MipIndex = 0; MipIndex < Desc.NumMips; ++MipIndex)
 						{
-							const uint32 MipPixelSizeX = FMath::Max<uint32>(1u, Desc.Extent.X >> MipIndex);
-							const uint32 MipPixelSizeY = FMath::Max<uint32>(1u, Desc.Extent.Y >> MipIndex);
-							const uint32 MipPixelSizeZ = FMath::Max<uint32>(1u, Desc.Depth >> MipIndex);
+							uint64 SubresourceStride{};
+							uint64 SubresourceSize{};
+							uint64 ResourceOffset = UE::RHITextureUtils::CalculateSubresourceOffset(Desc, FaceIndex, ArraySlice, MipIndex, SubresourceStride, SubresourceSize);
+							const uint64 MipSize = SubresourceSize;
+
+							const FUintVector3 BlockCounts = UE::RHITextureUtils::CalculateMipBlockCounts(Desc, MipIndex);
+							const FUintVector3 MipExtents = FUintVector3(BlockCounts.X * PixelFormat.BlockSizeX, BlockCounts.Y * PixelFormat.BlockSizeY, BlockCounts.Z * PixelFormat.BlockSizeZ);
 
 							const GLenum CurrentTarget = Target == GL_TEXTURE_CUBE_MAP
 								? GL_TEXTURE_CUBE_MAP_POSITIVE_X + FaceIndex
 								: Target;
 
-							if (Data)
+							if (!Callback(CurrentTarget, MipExtents, MipIndex, ArraySlice, reinterpret_cast<const uint8*>(Data) + ResourceOffset, MipSize))
 							{
-								const uint32 MipBlockSizeX = FMath::DivideAndRoundUp<uint32>(MipPixelSizeX, GPixelFormats[Desc.Format].BlockSizeX);
-								const uint32 MipBlockSizeY = FMath::DivideAndRoundUp<uint32>(MipPixelSizeY, GPixelFormats[Desc.Format].BlockSizeY);
-								const uint32 MipBlockSizeZ = FMath::DivideAndRoundUp<uint32>(MipPixelSizeZ, GPixelFormats[Desc.Format].BlockSizeZ);
-
-								const uint32 MipNumBlocks = MipBlockSizeX * MipBlockSizeY * MipBlockSizeZ;
-								const uint32 MipSize = MipNumBlocks * GPixelFormats[Desc.Format].BlockBytes * Desc.ArraySize;
-
-								if (!Callback(CurrentTarget, MipPixelSizeX, MipPixelSizeY, MipPixelSizeZ, MipIndex, ArraySlice, Data, MipSize))
-									return;
-
-								DataOffset += MipSize;
-								if (DataOffset >= BulkDataSize)
-								{
-									// Reach the end of bulk data. Only pass nullptr to the callback for any subsequent mips / slices
-									Data = nullptr;
-								}
-							}
-							else
-							{
-								if (!Callback(CurrentTarget, MipPixelSizeX, MipPixelSizeY, MipPixelSizeZ, MipIndex, ArraySlice, nullptr, 0))
-									return;
+								return;
 							}
 						}
 					}
@@ -746,22 +711,20 @@ void FOpenGLDynamicRHI::InitializeGLTextureInternal(FOpenGLTexture* Texture, voi
 
 					if (BulkDataPtr)
 					{
-						EnumerateSubresources(BulkDataPtr, [&](GLenum CurrentTarget, uint32 MipSizeX, uint32 MipSizeY, uint32 MipSizeZ, uint32 MipIndex, uint32 ArraySlice, void const* MipSliceData, uint32 MipSliceSize)
+						EnumerateSubresources(BulkDataPtr, [&](GLenum CurrentTarget, FUintVector3 MipExtent, uint32 MipIndex, uint32 ArraySlice, void const* MipSliceData, uint32 MipSliceSize)
 						{
-							// Stop when there's no more bulk data
-							if (MipSliceData == nullptr)
-								return false;
-
 							if (GLFormat.bCompressed)
 							{
 								glCompressedTexSubImage2D(
 									CurrentTarget,
 									MipIndex,
 									0, 0, // X/Y offset
-									MipSizeX, MipSizeY,
+									MipExtent.X, MipExtent.Y,
 									GLFormat.Format,
 									MipSliceSize,
-									MipSliceData);
+									MipSliceData
+								);
+								GL_CHECK();
 							}
 							else
 							{
@@ -769,10 +732,12 @@ void FOpenGLDynamicRHI::InitializeGLTextureInternal(FOpenGLTexture* Texture, voi
 									CurrentTarget,
 									MipIndex,
 									0, 0, // X/Y offset
-									MipSizeX, MipSizeY,
+									MipExtent.X, MipExtent.Y,
 									GLFormat.Format,
 									GLFormat.Type,
-									MipSliceData);
+									MipSliceData
+								);
+								GL_CHECK();
 							}
 
 							return true;
@@ -797,22 +762,20 @@ void FOpenGLDynamicRHI::InitializeGLTextureInternal(FOpenGLTexture* Texture, voi
 					// Texture created with immutable storage. Now fill in the bulk data.
 					if (BulkDataPtr)
 					{
-						EnumerateSubresources(BulkDataPtr, [&](GLenum CurrentTarget, uint32 MipSizeX, uint32 MipSizeY, uint32 MipSizeZ, uint32 MipIndex, uint32 ArraySlice, void const* MipSliceData, uint32 MipSliceSize)
+						EnumerateSubresources(BulkDataPtr, [&](GLenum CurrentTarget, FUintVector3 MipExtent, uint32 MipIndex, uint32 ArraySlice, void const* MipSliceData, uint32 MipSliceSize)
 						{
-							// Stop when there's no more bulk data
-							if (MipSliceData == nullptr)
-								return false;
-
 							if (GLFormat.bCompressed)
 							{
 								glCompressedTexSubImage3D(
 									CurrentTarget,
 									MipIndex,
 									0, 0, ArraySlice, // X/Y/Z offset
-									MipSizeX, MipSizeY, MipSizeZ,
+									MipExtent.X, MipExtent.Y, MipExtent.Z,
 									GLFormat.Format,
 									MipSliceSize,
-									MipSliceData);
+									MipSliceData
+								);
+								GL_CHECK();
 							}
 							else
 							{
@@ -820,10 +783,12 @@ void FOpenGLDynamicRHI::InitializeGLTextureInternal(FOpenGLTexture* Texture, voi
 									CurrentTarget,
 									MipIndex,
 									0, 0, ArraySlice, // X/Y/Z offset
-									MipSizeX, MipSizeY, MipSizeZ,
+									MipExtent.X, MipExtent.Y, MipExtent.Z,
 									GLFormat.Format,
 									GLFormat.Type,
-									MipSliceData);
+									MipSliceData
+								);
+								GL_CHECK();
 							}
 
 							return true;
@@ -941,71 +906,103 @@ void FOpenGLDynamicRHI::ResolveTexture(FOpenGLTexture* Texture, uint32 MipIndex,
 	// next operation that needs the stage will switch something else in on it.
 }
 
-uint32 FOpenGLTexture::GetLockSize(uint32 InMipIndex, uint32 ArrayIndex, EResourceLockMode LockMode, uint32& DestStride)
+uint32 FOpenGLTexture::GetLockSize(uint32 InMipIndex, EResourceLockMode LockMode, uint32& DestStride)
 {
 	check(LockMode != EResourceLockMode::RLM_WriteOnly_NoOverwrite);
-	// Calculate the dimensions of the mip-map.
-	EPixelFormat PixelFormat = this->GetFormat();
-	const uint32 BlockSizeX = GPixelFormats[PixelFormat].BlockSizeX;
-	const uint32 BlockSizeY = GPixelFormats[PixelFormat].BlockSizeY;
-	const uint32 BlockBytes = GPixelFormats[PixelFormat].BlockBytes;
-	const uint32 MipSizeX = FMath::Max(this->GetSizeX() >> InMipIndex, BlockSizeX);
-	const uint32 MipSizeY = FMath::Max(this->GetSizeY() >> InMipIndex, BlockSizeY);
-	uint32 NumBlocksX = (MipSizeX + BlockSizeX - 1) / BlockSizeX;
-	uint32 NumBlocksY = (MipSizeY + BlockSizeY - 1) / BlockSizeY;
-	const uint32 MipBytes = NumBlocksX * NumBlocksY * BlockBytes;
-	DestStride = NumBlocksX * BlockBytes;
-	return MipBytes;
+	uint64 LargerStride{};
+	const uint64 MipSize = UE::RHITextureUtils::CalculateTextureMipSize(GetDesc(), InMipIndex, LargerStride);
+
+	DestStride = static_cast<uint32>(LargerStride);
+	return MipSize;
 }
 
-void FOpenGLTexture::Fill2DGLTextureImage(const FOpenGLTextureFormat& GLFormat, const bool bSRGB, uint32 MipIndex, const void* BufferOrPBOOffset, uint32 ImageSize, uint32 ArrayIndex)
+void FOpenGLTexture::FillGLTextureImage(const FOpenGLTextureFormat& GLFormat, bool bSRGB, uint32 MipIndex, uint32 ArrayIndex, const FUintVector3& MipExtents, const void* BufferOrPBOOffset, size_t ImageSize)
 {
-	if (GLFormat.bCompressed)
+	FScopedPackAlignment PackAlignment;
+
+	if (GetDesc().IsTextureArray() || GetDesc().IsTexture3D())
 	{
-		if (GetAllocatedStorageForMip(MipIndex,ArrayIndex))
+		const GLsizei Depth = 1;
+
+		if (GLFormat.bCompressed)
 		{
-			glCompressedTexSubImage2D(
-				bCubemap ? GL_TEXTURE_CUBE_MAP_POSITIVE_X + ArrayIndex : Target,
+			FOpenGL::CompressedTexSubImage3D(
+				Target,
 				MipIndex,
-				0,
-				0,
-				FMath::Max<uint32>(1,(this->GetSizeX() >> MipIndex)),
-				FMath::Max<uint32>(1,(this->GetSizeY() >> MipIndex)),
+				0, 0, ArrayIndex,
+				MipExtents.X, MipExtents.Y, Depth,
 				GLFormat.InternalFormat[bSRGB],
 				ImageSize,
-				BufferOrPBOOffset);	// offset into PBO
+				BufferOrPBOOffset
+			);
+			GL_CHECK();
 		}
 		else
 		{
-			glCompressedTexImage2D(
-				bCubemap ? GL_TEXTURE_CUBE_MAP_POSITIVE_X + ArrayIndex : Target,
+			check(FOpenGL::SupportsTexture3D());
+			FOpenGL::TexSubImage3D(
+				Target,
 				MipIndex,
-				GLFormat.InternalFormat[bSRGB],
-				FMath::Max<uint32>(1,(this->GetSizeX() >> MipIndex)),
-				FMath::Max<uint32>(1,(this->GetSizeY() >> MipIndex)),
-				0,
-				ImageSize,
-				BufferOrPBOOffset);	// offset into PBO
-			SetAllocatedStorageForMip(MipIndex,ArrayIndex);
+				0, 0, ArrayIndex,
+				MipExtents.X, MipExtents.Y, Depth,
+				GLFormat.Format,
+				GLFormat.Type,
+				BufferOrPBOOffset
+			);
+			GL_CHECK();
 		}
 	}
 	else
 	{
-		// All construction paths should have called TexStorage2D or TexImage2D. So we will
-		// always call TexSubImage2D.
-		check(GetAllocatedStorageForMip(MipIndex,ArrayIndex) == true);
-		glPixelStorei(GL_UNPACK_ALIGNMENT, 1);
-		glTexSubImage2D(
-			bCubemap ? GL_TEXTURE_CUBE_MAP_POSITIVE_X + ArrayIndex : Target,
-			MipIndex,
-			0,
-			0,
-			FMath::Max<uint32>(1,(this->GetSizeX() >> MipIndex)),
-			FMath::Max<uint32>(1,(this->GetSizeY() >> MipIndex)),
-			GLFormat.Format,
-			GLFormat.Type,
-			BufferOrPBOOffset);	// offset into PBO
-		glPixelStorei(GL_UNPACK_ALIGNMENT, 4);
+		const GLenum ActualTarget = bCubemap ? GL_TEXTURE_CUBE_MAP_POSITIVE_X + ArrayIndex : Target;
+
+		if (GLFormat.bCompressed)
+		{
+			if (GetAllocatedStorageForMip(MipIndex, ArrayIndex))
+			{
+				glCompressedTexSubImage2D(
+					ActualTarget,
+					MipIndex,
+					0, 0,
+					MipExtents.X, MipExtents.Y,
+					GLFormat.InternalFormat[bSRGB],
+					ImageSize,
+					BufferOrPBOOffset
+				);
+				GL_CHECK();
+			}
+			else
+			{
+				glCompressedTexImage2D(
+					ActualTarget,
+					MipIndex,
+					GLFormat.InternalFormat[bSRGB],
+					MipExtents.X, MipExtents.Y,
+					0,
+					ImageSize,
+					BufferOrPBOOffset
+				);
+				GL_CHECK();
+
+				SetAllocatedStorageForMip(MipIndex, ArrayIndex);
+			}
+		}
+		else
+		{
+			// All construction paths should have called TexStorage2D or TexImage2D. So we will always call TexSubImage2D.
+			check(GetAllocatedStorageForMip(MipIndex, ArrayIndex) == true);
+
+			glTexSubImage2D(
+				ActualTarget,
+				MipIndex,
+				0, 0,
+				MipExtents.X, MipExtents.Y,
+				GLFormat.Format,
+				GLFormat.Type,
+				BufferOrPBOOffset
+			);
+			GL_CHECK();
+		}
 	}
 }
 
@@ -1016,7 +1013,7 @@ void* FOpenGLTexture::Lock(uint32 InMipIndex,uint32 ArrayIndex,EResourceLockMode
 
 	SCOPE_CYCLE_COUNTER(STAT_OpenGLLockTextureTime);
 	
-	const uint32 MipBytes = GetLockSize(InMipIndex, ArrayIndex, LockMode, DestStride);
+	const uint32 MipBytes = GetLockSize(InMipIndex, LockMode, DestStride);
 
 	check(!IsEvicted() || ArrayIndex==0);
 	void* result = NULL;
@@ -1136,6 +1133,8 @@ void FOpenGLTexture::Unlock(uint32 MipIndex,uint32 ArrayIndex)
 	const int32 BufferIndex = MipIndex * (bCubemap ? 6 : 1) * this->GetEffectiveSizeZ() + ArrayIndex;
 	const FOpenGLTextureFormat& GLFormat = GOpenGLTextureFormats[this->GetFormat()];
 	const bool bSRGB = EnumHasAnyFlags(this->GetFlags(), TexCreate_SRGB);
+
+	check(BufferIndex < PixelBuffers.Num());
 	TRefCountPtr<FOpenGLPixelBuffer> PixelBuffer = PixelBuffers[BufferIndex];
 
 	check(IsValidRef(PixelBuffer));
@@ -1213,46 +1212,9 @@ void FOpenGLTexture::Unlock(uint32 MipIndex,uint32 ArrayIndex)
 		// Use a texture stage that's not likely to be used for draws, to avoid waiting
 		FOpenGLDynamicRHI::Get().CachedSetupTextureStage(FOpenGL::GetMaxCombinedTextureImageUnits() - 1, Target, GetResource(), -1, this->GetNumMips());
 
-		if (GetDesc().IsTextureArray() || GetDesc().IsTexture3D())
-		{
-			if (GLFormat.bCompressed)
-			{
-				FOpenGL::CompressedTexSubImage3D(
-					Target,
-					MipIndex,
-					0,
-					0,
-					ArrayIndex,
-					FMath::Max<uint32>(1,(this->GetSizeX() >> MipIndex)),
-					FMath::Max<uint32>(1,(this->GetSizeY() >> MipIndex)),
-					1,
-					GLFormat.InternalFormat[bSRGB],
-					PixelBuffer->GetSize(),
-					0);
-			}
-			else
-			{
-				glPixelStorei(GL_UNPACK_ALIGNMENT, 1);
-				check( FOpenGL::SupportsTexture3D() );
-				FOpenGL::TexSubImage3D(
-					Target,
-					MipIndex,
-					0,
-					0,
-					ArrayIndex,
-					FMath::Max<uint32>(1,(this->GetSizeX() >> MipIndex)),
-					FMath::Max<uint32>(1,(this->GetSizeY() >> MipIndex)),
-					1,
-					GLFormat.Format,
-					GLFormat.Type,
-					0);	// offset into PBO
-				glPixelStorei(GL_UNPACK_ALIGNMENT, 4);
-			}
-		}
-		else
-		{
-			Fill2DGLTextureImage(GLFormat, bSRGB, MipIndex, 0, PixelBuffer->GetSize(), ArrayIndex);
-		}
+		const FUintVector3 MipExtents = UE::RHITextureUtils::CalculateMipExtents(GetDesc(), MipIndex);
+
+		FillGLTextureImage(GLFormat, bSRGB, MipIndex, ArrayIndex, MipExtents, nullptr, PixelBuffer->GetSize());
 	}
 
 	//need to free PBO if we aren't keeping shadow copies
@@ -1297,12 +1259,14 @@ void FOpenGLTexture::RestoreEvictedGLResource(bool bAttemptToRetainMips)
 
 	FOpenGLDynamicRHI::Get().CachedBindBuffer(GL_PIXEL_UNPACK_BUFFER, 0);
 
-	for (int i = EvictionParamsPtr->MipImageData.Num() - 1; i >= 0; i--)
+	for (int32 MipIndex = EvictionParamsPtr->MipImageData.Num() - 1; MipIndex >= 0; MipIndex--)
 	{
-		auto& MipMem = EvictionParamsPtr->MipImageData[i];
-		if(MipMem.Num())
+		TConstArrayView<uint8> MipMem = EvictionParamsPtr->MipImageData[MipIndex];
+		if (MipMem.Num())
 		{
-			Fill2DGLTextureImage(GLFormat, bSRGB, i, MipMem.GetData(), MipMem.Num(), 0);
+			const FUintVector3 MipExtents = UE::RHITextureUtils::CalculateMipExtents(GetDesc(), MipIndex);
+
+			FillGLTextureImage(GLFormat, bSRGB, MipIndex, 0, MipExtents, MipMem.GetData(), MipMem.Num());
 		}
 	}
 
@@ -1471,14 +1435,96 @@ void FOpenGLTexture::CloneViaCopyImage(FOpenGLTexture* Src, uint32 InNumMips, in
 	2D texture support.
 -----------------------------------------------------------------------------*/
 
-FTextureRHIRef FOpenGLDynamicRHI::RHICreateTexture(FRHICommandListBase& RHICmdList, const FRHITextureCreateDesc& CreateDesc)
+static UE::RHICore::FDefaultLayoutTextureInitializer CreateWritableTextureInitializer(FRHICommandListBase& RHICmdList, FOpenGLTexture* Texture, const FRHITextureCreateDesc& CreateDesc)
+{
+	const uint64 WritableMemorySize = UE::RHITextureUtils::CalculateTextureSize(CreateDesc);
+	void* WritableMemory = FMemory::Malloc(WritableMemorySize, 16);
+
+	// TODO: when not evicted, we can use PixelBuffer memory for the Initializer. This may get too complicated to justify.
+
+	return UE::RHICore::FDefaultLayoutTextureInitializer(RHICmdList, Texture, WritableMemory, WritableMemorySize,
+		[Texture = TRefCountPtr<FOpenGLTexture>(Texture), WritableMemory = UE::RHICore::FInitializerScopedMemory(WritableMemory), WritableMemorySize](FRHICommandListBase& RHICmdList) mutable
+		{
+			RHICmdList.EnqueueLambda(
+				[Texture, WritableMemory = MoveTemp(WritableMemory), WritableMemorySize](FRHICommandListBase& RHICmdList)
+				{
+					FOpenGLDynamicRHI::Get().InitializeGLTexture(Texture, nullptr, 0);
+
+					const FRHITextureDesc& TextureDesc = Texture->GetDesc();
+					const uint32 FaceCount = TextureDesc.IsTextureCube() ? 6 : 1;
+
+					for (uint32 FaceIndex = 0; FaceIndex < FaceCount; FaceIndex++)
+					{
+						for (uint32 ArrayIndex = 0; ArrayIndex < TextureDesc.ArraySize; ArrayIndex++)
+						{
+							for (uint32 MipIndex = 0; MipIndex < TextureDesc.NumMips; MipIndex++)
+							{
+								uint64 SubresourceStride = 0;
+								uint64 SubresourceSize = 0;
+								const uint64 Offset = UE::RHITextureUtils::CalculateSubresourceOffset(TextureDesc, FaceIndex, ArrayIndex, MipIndex, SubresourceStride, SubresourceSize);
+
+								check((Offset + SubresourceSize) <= WritableMemorySize);
+
+								const uint32 GLArrayIndex = FaceIndex + ArrayIndex * FaceCount;
+
+								uint32 DestStride = 0;
+								void* DestinationMemory = Texture->Lock(MipIndex, GLArrayIndex, RLM_WriteOnly, DestStride);
+
+								FMemory::Memcpy(DestinationMemory, reinterpret_cast<const uint8*>(WritableMemory.Pointer) + Offset, SubresourceSize);
+								Texture->Unlock(MipIndex, GLArrayIndex);
+							}
+						}
+					}
+				}
+			);
+
+			return TRefCountPtr<FRHITexture>(MoveTemp(Texture));
+		}
+	);
+}
+
+FRHITextureInitializer FOpenGLDynamicRHI::RHICreateTextureInitializer(FRHICommandListBase& RHICmdList, const FRHITextureCreateDesc& CreateDesc)
 {
 	LLM_SCOPE_DYNAMIC_STAT_OBJECTPATH_FNAME(CreateDesc.OwnerName, ELLMTagSet::Assets);
 	LLM_SCOPE_DYNAMIC_STAT_OBJECTPATH_FNAME(CreateDesc.GetTraceClassName(), ELLMTagSet::AssetClasses);
 	UE_TRACE_METADATA_SCOPE_ASSET_FNAME(CreateDesc.DebugName, CreateDesc.GetTraceClassName(), CreateDesc.OwnerName);
-	return new FOpenGLTexture(RHICmdList, CreateDesc);
+
+	FOpenGLTexture* Texture = new FOpenGLTexture(CreateDesc);
+
+	if (CreateDesc.InitAction == ERHITextureInitAction::Default)
+	{
+		return UE::RHICore::FBaseTextureInitializerImplementation(RHICmdList, Texture,
+			[Texture = TRefCountPtr<FOpenGLTexture>(Texture)](FRHICommandListBase& RHICmdList) mutable
+			{
+				Texture->Initialize(RHICmdList);
+				return TRefCountPtr<FRHITexture>(MoveTemp(Texture));
+			}
+		);
+	}
+
+	if (CreateDesc.InitAction == ERHITextureInitAction::BulkData)
+	{
+		check(CreateDesc.BulkData);
+
+		UE::RHICore::FDefaultLayoutTextureInitializer Initializer = CreateWritableTextureInitializer(RHICmdList, Texture, CreateDesc);
+		check(CreateDesc.BulkData->GetResourceBulkDataSize() <= Initializer.GetWritableSize());
+
+		FMemory::Memcpy(Initializer.GetWritableData(), CreateDesc.BulkData->GetResourceBulkData(), CreateDesc.BulkData->GetResourceBulkDataSize());
+
+		CreateDesc.BulkData->Discard();
+
+		return MoveTemp(Initializer);
+	}
+
+	if (CreateDesc.InitAction == ERHITextureInitAction::Initializer)
+	{
+		return CreateWritableTextureInitializer(RHICmdList, Texture, CreateDesc);
+	}
+
+	return UE::RHICore::HandleUnknownTextureInitializerInitAction(RHICmdList, CreateDesc);
 }
 
+
 FTextureRHIRef FOpenGLDynamicRHI::RHIAsyncCreateTexture2D(uint32 SizeX, uint32 SizeY, uint8 Format, uint32 NumMips, ETextureCreateFlags Flags, ERHIAccess InResourceState, void** InitialMipData, uint32 NumInitialMips, const TCHAR* DebugName, FGraphEventRef& OutCompletionEvent)
 {
 	check(0);
@@ -1547,12 +1593,14 @@ FTextureRHIRef FOpenGLDynamicRHI::RHIAsyncReallocateTexture2D(FRHITexture* Textu
 	LLM_SCOPE_DYNAMIC_STAT_OBJECTPATH_FNAME(CreateDesc.GetTraceClassName(), ELLMTagSet::AssetClasses);
 	UE_TRACE_METADATA_SCOPE_ASSET_FNAME(CreateDesc.DebugName, CreateDesc.GetTraceClassName(), CreateDesc.OwnerName);
 
-	FOpenGLTexture* NewTexture = new FOpenGLTexture(RHICmdList, CreateDesc);
+	FOpenGLTexture* NewTexture = new FOpenGLTexture(CreateDesc);
 
 	RHICmdList.EnqueueLambda([OldTexture, NewTexture, SourceMipCount, NewMipCount, RequestStatus](FRHICommandListImmediate& RHICmdList)
 	{
 		VERIFY_GL_SCOPE();
 
+		NewTexture->Initialize(RHICmdList);
+
 		// Use the GPU to asynchronously copy the old mip-maps into the new texture.
 		const uint32 NumSharedMips = FMath::Min(SourceMipCount, NewMipCount);
 		const uint32 SourceMipOffset = SourceMipCount - NumSharedMips;
@@ -1634,6 +1682,7 @@ void FOpenGLDynamicRHI::RHIUpdateTexture2D(FRHICommandListBase& RHICmdList, FRHI
 				GLFormat.Format,
 				UpdatePitch * FMath::DivideAndRoundUp<uint32>(UpdateRegion.Height, FormatInfo.BlockSizeY),
 				UpdateMemory);
+			GL_CHECK();
 		}
 		else
 		{
@@ -1645,6 +1694,7 @@ void FOpenGLDynamicRHI::RHIUpdateTexture2D(FRHICommandListBase& RHICmdList, FRHI
 				GLFormat.Format,
 				GLFormat.Type,
 				UpdateMemory);
+			GL_CHECK();
 		}
 
 		glPixelStorei(GL_UNPACK_ALIGNMENT, 4);
@@ -1999,11 +2049,11 @@ FRHILockTextureResult FOpenGLDynamicRHI::RHILockTexture(FRHICommandListImmediate
 
 		RHITHREAD_GLTRACE_BLOCKING;
 		RHICmdList.ImmediateFlush(EImmediateFlushType::FlushRHIThread);
-		Result.ByteCount = Texture->GetLockSize(Arguments.MipIndex, ArrayIndex, RLM_ReadOnly, Result.Stride);
+		Result.ByteCount = Texture->GetLockSize(Arguments.MipIndex, RLM_ReadOnly, Result.Stride);
 	}
 	else
 	{
-		Result.ByteCount = Texture->GetLockSize(Arguments.MipIndex, 0, RLM_WriteOnly, Result.Stride);
+		Result.ByteCount = Texture->GetLockSize(Arguments.MipIndex, RLM_WriteOnly, Result.Stride);
 		Result.Data = FMemory::Malloc(Result.ByteCount, 16);
 	}
 
diff --git a/Engine/Source/Runtime/OpenGLDrv/Private/OpenGLViewport.cpp b/Engine/Source/Runtime/OpenGLDrv/Private/OpenGLViewport.cpp
--- a/Engine/Source/Runtime/OpenGLDrv/Private/OpenGLViewport.cpp
+++ b/Engine/Source/Runtime/OpenGLDrv/Private/OpenGLViewport.cpp
@@ -312,7 +312,8 @@ void FOpenGLViewport::Resize(uint32 InSizeX,uint32 InSizeY,bool bInIsFullscreen)
 				.SetFlags(ETextureCreateFlags::RenderTargetable | ETextureCreateFlags::ResolveTargetable)
 				.DetermineInititialState();
 
-			BackBuffer = new FOpenGLTexture(RHICmdList, Desc);
+			BackBuffer = new FOpenGLTexture(Desc);
+			BackBuffer->Initialize(RHICmdList);
 		}
 
 		RHICmdList.EnqueueLambda([this, InSizeX, InSizeY, bInIsFullscreen, bWasFullscreen](FRHICommandListImmediate&)
diff --git a/Engine/Source/Runtime/OpenGLDrv/Private/Windows/WindowsOpenGLPlatform.cpp b/Engine/Source/Runtime/OpenGLDrv/Private/Windows/WindowsOpenGLPlatform.cpp
--- a/Engine/Source/Runtime/OpenGLDrv/Private/Windows/WindowsOpenGLPlatform.cpp
+++ b/Engine/Source/Runtime/OpenGLDrv/Private/Windows/WindowsOpenGLPlatform.cpp
@@ -830,7 +830,9 @@ FOpenGLTexture* PlatformCreateBuiltinBackBuffer(FOpenGLDynamicRHI* OpenGLRHI, ui
 			.SetFlags(ETextureCreateFlags::RenderTargetable | ETextureCreateFlags::Presentable | ETextureCreateFlags::ResolveTargetable)
 			.DetermineInititialState();
 
-		return new FOpenGLTexture(FRHICommandListImmediate::Get(), Desc);
+		FOpenGLTexture* Texture = new FOpenGLTexture(Desc);
+		Texture->Initialize(FRHICommandListImmediate::Get());
+		return Texture;
 	}
 
 	return nullptr;
diff --git a/Engine/Source/Runtime/RHI/Private/RHICommandList.cpp b/Engine/Source/Runtime/RHI/Private/RHICommandList.cpp
--- a/Engine/Source/Runtime/RHI/Private/RHICommandList.cpp
+++ b/Engine/Source/Runtime/RHI/Private/RHICommandList.cpp
@@ -266,6 +266,21 @@ void FRHICommandListBase::FinishRecording()
 
 		UE_LOG(LogRHI, Fatal, TEXT("Detected pending buffer uploads on RHICmdList submission: %s"), BufferList.ToString());
 	}
+	if (PendingTextureUploads.Num() > 0)
+	{
+		FStringBuilderBase TextureList;
+
+		for (const FRHITexture* Texture : PendingTextureUploads)
+		{
+			if (TextureList.Len() > 0)
+			{
+				TextureList.Append(TEXT(", "));
+			}
+			TextureList.Append(Texture->GetName().ToString());
+		}
+
+		UE_LOG(LogRHI, Fatal, TEXT("Detected pending texture uploads on RHICmdList submission: %s"), TextureList.ToString());
+	}
 
 	// "Complete" the dispatch event.
 	DispatchEvent->DispatchSubsequents();
diff --git a/Engine/Source/Runtime/RHI/Private/RHITextureInitializer.cpp b/Engine/Source/Runtime/RHI/Private/RHITextureInitializer.cpp
new file mode 100644
--- /dev/null
+++ b/Engine/Source/Runtime/RHI/Private/RHITextureInitializer.cpp
@@ -0,0 +1,47 @@
+// Copyright Epic Games, Inc. All Rights Reserved.
+
+#include "RHITextureInitializer.h"
+#include "RHICommandList.h"
+#include "HAL/LowLevelMemStats.h"
+#include "ProfilingDebugging/AssetMetadataTrace.h"
+
+FRHITextureInitializer::FRHITextureInitializer(FRHICommandListBase& RHICmdList, FRHITexture* InTexture, void* InWritableData, uint64 InWritableSize, FRHITextureInitializer::FFinalizeCallback&& InFinalizeCallback, FGetSubresourceCallback&& InGetSubresourceCallback)
+	: FinalizeCallback(Forward<FFinalizeCallback>(InFinalizeCallback))
+	, GetSubresourceCallback(Forward<FGetSubresourceCallback>(InGetSubresourceCallback))
+	, CommandList(&RHICmdList)
+	, Texture(InTexture)
+	, WritableData(InWritableData)
+	, WritableSize(InWritableSize)
+	, Desc(InTexture->GetDesc())
+{
+	check(InTexture != nullptr);
+	RHICmdList.AddPendingTextureUpload(InTexture);
+}
+
+FTextureRHIRef FRHITextureInitializer::Finalize()
+{
+	FTextureRHIRef Result;
+	if (FinalizeCallback)
+	{
+		check(Texture != nullptr);
+
+		LLM_SCOPE(EnumHasAnyFlags(Texture->GetFlags(), TexCreate_RenderTargetable | TexCreate_DepthStencilTargetable) ? ELLMTag::RenderTargets : ELLMTag::Textures);
+
+		LLM_SCOPE_DYNAMIC_STAT_OBJECTPATH_FNAME(Texture->GetOwnerName(), ELLMTagSet::Assets);
+		LLM_SCOPE_DYNAMIC_STAT_OBJECTPATH_FNAME(FRHITextureCreateDesc().GetTraceClassName(), ELLMTagSet::AssetClasses);
+		UE_TRACE_METADATA_SCOPE_ASSET_FNAME(Texture->GetName(), FRHITextureCreateDesc().GetTraceClassName(), Texture->GetOwnerName());
+
+		FRHICommandListScopedPipelineGuard ScopedPipeline(*CommandList);
+		Result = FinalizeCallback(*CommandList);
+
+		RemovePendingTextureUpload();
+
+		Reset();
+	}
+	return Result;
+}
+
+void FRHITextureInitializer::RemovePendingTextureUpload()
+{
+	CommandList->RemovePendingTextureUpload(Texture);
+}
diff --git a/Engine/Source/Runtime/RHI/Public/DynamicRHI.h b/Engine/Source/Runtime/RHI/Public/DynamicRHI.h
--- a/Engine/Source/Runtime/RHI/Public/DynamicRHI.h
+++ b/Engine/Source/Runtime/RHI/Public/DynamicRHI.h
@@ -11,6 +11,7 @@ DynamicRHI.h: Dynamically bound Render Hardware Interface definitions.
 #include "RHIContext.h"
 #include "RHIFeatureLevel.h"
 #include "RHIShaderLibrary.h"
+#include "RHITextureInitializer.h"
 #include "MultiGPU.h"
 #include "Serialization/MemoryLayout.h"
 #include "Containers/ArrayView.h"
@@ -505,10 +506,9 @@ public:
 	// FlushType: Flush Immediate
 	virtual bool RHIGetTextureMemoryVisualizeData(FColor* TextureData, int32 SizeX, int32 SizeY, int32 Pitch, int32 PixelSize) = 0;
 
-	/**
-	* Creates an RHI texture resource.
-	*/
-	virtual FTextureRHIRef RHICreateTexture(FRHICommandListBase& RHICmdList, const FRHITextureCreateDesc& CreateDesc) = 0;
+	// Create a texture initializer, used for creating a RHITexture with data.
+	[[nodiscard]]
+	virtual FRHITextureInitializer RHICreateTextureInitializer(FRHICommandListBase& RHICmdList, const FRHITextureCreateDesc& CreateDesc) = 0;
 
 	/**
 	* Thread-safe function that can be used to create a texture outside of the
diff --git a/Engine/Source/Runtime/RHI/Public/RHIBufferInitializer.h b/Engine/Source/Runtime/RHI/Public/RHIBufferInitializer.h
--- a/Engine/Source/Runtime/RHI/Public/RHIBufferInitializer.h
+++ b/Engine/Source/Runtime/RHI/Public/RHIBufferInitializer.h
@@ -6,7 +6,7 @@
 #include "Templates/Function.h"
 
 // Structure used to allow optimal buffer initialization at creation time.
-// Should only ever be obtained by calling RHIBufferBeginCreate() and should be either moved for a RHIBufferEndCreate call or destroyed normally.
+// Should only ever be obtained by calling RHICreateBufferInitializer().
 // NO COPIES ALLOWED
 struct FRHIBufferInitializer
 {
diff --git a/Engine/Source/Runtime/RHI/Public/RHICommandList.h b/Engine/Source/Runtime/RHI/Public/RHICommandList.h
--- a/Engine/Source/Runtime/RHI/Public/RHICommandList.h
+++ b/Engine/Source/Runtime/RHI/Public/RHICommandList.h
@@ -921,7 +921,8 @@ PRAGMA_ENABLE_DEPRECATION_WARNINGS
 		GDynamicRHI->RHIUpdateTexture2D(*this, Texture, MipIndex, UpdateRegion, SourcePitch, SourceData);
 	}
 
-	FORCEINLINE FTextureRHIRef CreateTexture(const FRHITextureCreateDesc& CreateDesc)
+	[[nodiscard]]
+	FORCEINLINE FRHITextureInitializer CreateTextureInitializer(const FRHITextureCreateDesc& CreateDesc)
 	{
 		LLM_SCOPE(EnumHasAnyFlags(CreateDesc.Flags, TexCreate_RenderTargetable | TexCreate_DepthStencilTargetable) ? ELLMTag::RenderTargets : ELLMTag::Textures);
 
@@ -931,10 +932,16 @@ PRAGMA_ENABLE_DEPRECATION_WARNINGS
 			FRHITextureCreateDesc NewCreateDesc(CreateDesc);
 			NewCreateDesc.SetInitialState(RHIGetDefaultResourceState(CreateDesc.Flags, CreateDesc.BulkData != nullptr));
 
-			return GDynamicRHI->RHICreateTexture(*this, NewCreateDesc);
+			return GDynamicRHI->RHICreateTextureInitializer(*this, NewCreateDesc);
 		}
 
-		return GDynamicRHI->RHICreateTexture(*this, CreateDesc);
+		return GDynamicRHI->RHICreateTextureInitializer(*this, CreateDesc);
+	}
+
+	FORCEINLINE FTextureRHIRef CreateTexture(const FRHITextureCreateDesc& CreateDesc)
+	{
+		FRHITextureInitializer Initializer = CreateTextureInitializer(CreateDesc);
+		return Initializer.Finalize();
 	}
 
 	FORCEINLINE void UpdateFromBufferTexture2D(FRHITexture* Texture, uint32 MipIndex, const struct FUpdateTextureRegion2D& UpdateRegion, uint32 SourcePitch, FRHIBuffer* Buffer, uint32 BufferOffset)
@@ -1310,6 +1317,17 @@ protected:
 		PendingBufferUploads.Remove(InBuffer);
 	}
 
+	void AddPendingTextureUpload(FRHITexture* InTexture)
+	{
+		PendingTextureUploads.Emplace(InTexture);
+	}
+
+	void RemovePendingTextureUpload(FRHITexture* InTexture)
+	{
+		check(PendingTextureUploads.Contains(InTexture));
+		PendingTextureUploads.Remove(InTexture);
+	}
+
 protected:
 	FRHICommandBase*    Root            = nullptr;
 	FRHICommandBase**   CommandLink     = nullptr;
@@ -1485,6 +1503,7 @@ protected:
 	FRHIDrawStats DrawStats {};
 
 	TArray<FRHIBuffer*> PendingBufferUploads;
+	TArray<FRHITexture*> PendingTextureUploads;
 
 public:
 #if WITH_RHI_BREADCRUMBS
@@ -1545,6 +1564,7 @@ private:
 	friend class FRHIComputeCommandList_RecursiveHazardous;
 	friend struct FRHICommandSetGPUMask;
 	friend struct FRHIBufferInitializer;
+	friend struct FRHITextureInitializer;
 
 	template <typename RHICmdListType, typename LAMBDA>
 	friend struct TRHILambdaCommandMultiPipe;
diff --git a/Engine/Source/Runtime/RHI/Public/RHIResources.h b/Engine/Source/Runtime/RHI/Public/RHIResources.h
--- a/Engine/Source/Runtime/RHI/Public/RHIResources.h
+++ b/Engine/Source/Runtime/RHI/Public/RHIResources.h
@@ -1941,6 +1941,18 @@ using FRHITextureCreateInfo = FRHITextureDesc;
 
 extern RHI_API ERHIAccess RHIGetDefaultResourceState(ETextureCreateFlags InUsage, bool bInHasInitialData);
 
+enum class ERHITextureInitAction
+{
+	// Default for the RHI, data can be "undefined" after creation.
+	Default,
+
+	// Upload data from a provided FResourceBulkDataInterface. This data will be discarded after it's used.
+	BulkData,
+
+	// Caller will use FRHITextureInitializer to set the initial texture contents.
+	Initializer,
+};
+
 struct FRHITextureCreateDesc : public FRHITextureDesc
 {
 	static FRHITextureCreateDesc Create(const TCHAR* InDebugName, ETextureDimension InDimension)
@@ -2051,10 +2063,14 @@ struct FRHITextureCreateDesc : public FRHITextureDesc
 		, FResourceBulkDataInterface* InBulkData     = nullptr
 		)
 		: FRHITextureDesc(InDesc)
-		, InitialState   (InInitialState)
-		, DebugName      (InDebugName)
-		, BulkData       (InBulkData)
-	{}
+		, DebugName(InDebugName)
+		, InitialState(InInitialState)
+	{
+		if (InBulkData)
+		{
+			SetInitActionBulkData(InBulkData);
+		}
+	}
 
 	void CheckValidity() const
 	{
@@ -2077,27 +2093,75 @@ struct FRHITextureCreateDesc : public FRHITextureDesc
 	FRHITextureCreateDesc& SetDimension(ETextureDimension InDimension)         { Dimension = InDimension;                  return *this; }
 	FRHITextureCreateDesc& SetFormat(EPixelFormat InFormat)                    { Format = InFormat;                        return *this; }
 	FRHITextureCreateDesc& SetUAVFormat(EPixelFormat InUAVFormat)              { UAVFormat = InUAVFormat;                  return *this; }
-	FRHITextureCreateDesc& SetInitialState(ERHIAccess InInitialState)          { InitialState = InInitialState;            return *this; }
 	FRHITextureCreateDesc& SetDebugName(const TCHAR* InDebugName)              { DebugName = InDebugName;                  return *this; }
+	FRHITextureCreateDesc& SetInitialState(ERHIAccess InInitialState)          { InitialState = InInitialState;            return *this; }
+	FRHITextureCreateDesc& SetInitAction(ERHITextureInitAction InInitAction)   { InitAction = InInitAction;                return *this; }
 	FRHITextureCreateDesc& SetGPUMask(FRHIGPUMask InGPUMask)                   { GPUMask = InGPUMask;                      return *this; }
-	FRHITextureCreateDesc& SetBulkData(FResourceBulkDataInterface* InBulkData) { BulkData = InBulkData;                    return *this; }
-	FRHITextureCreateDesc& DetermineInititialState()                           { if (InitialState == ERHIAccess::Unknown) InitialState = RHIGetDefaultResourceState(Flags, BulkData != nullptr); return *this; }
-	FRHITextureCreateDesc& SetFastVRAMPercentage(float In)                     { FastVRAMPercentage = uint8(FMath::Clamp(In, 0.f, 1.0f) * 0xFF); return *this; }
 	FRHITextureCreateDesc& SetClassName(const FName& InClassName)			   { ClassName = InClassName;				   return *this; }
 	FRHITextureCreateDesc& SetOwnerName(const FName& InOwnerName)			   { OwnerName = InOwnerName;                  return *this; }
-	FName GetTraceClassName() const											   { const static FLazyName FRHITextureName(TEXT("FRHITexture")); return (ClassName == NAME_None) ? FRHITextureName : ClassName; }
 
-	/* The RHI access state that the resource will be created in. */
-	ERHIAccess InitialState = ERHIAccess::Unknown;
+	FRHITextureCreateDesc& SetFastVRAMPercentage(float InFastVRAMPercentage)
+	{
+		FastVRAMPercentage = uint8(FMath::Clamp(InFastVRAMPercentage, 0.f, 1.0f) * 0xFF);
+		return *this;
+	}
+
+	FRHITextureCreateDesc& DetermineInititialState()
+	{
+		if (InitialState == ERHIAccess::Unknown)
+		{
+			InitialState = RHIGetDefaultResourceState(Flags, BulkData != nullptr);
+		}
+		return *this;
+	}
+
+	FRHITextureCreateDesc& SetInitActionNone()
+	{
+		return SetInitAction(ERHITextureInitAction::Default);
+	}
+	FRHITextureCreateDesc& SetInitActionInitializer()
+	{
+		return SetInitAction(ERHITextureInitAction::Initializer);
+	}
+	FRHITextureCreateDesc& SetInitActionBulkData(FResourceBulkDataInterface* InBulkData)
+	{
+		BulkData = InBulkData;
+		return SetInitAction(ERHITextureInitAction::BulkData);
+	}
+
+	UE_DEPRECATED(5.6, "SetBulkData has been renamed to SetInitActionBulkData")
+	FRHITextureCreateDesc& SetBulkData(FResourceBulkDataInterface* InBulkData)
+	{
+		if (InBulkData)
+		{
+			SetInitActionBulkData(InBulkData);
+		}
+		return *this;
+	}
+
+	FName GetTraceClassName() const
+	{
+		const static FLazyName FRHITextureName(TEXT("FRHITexture"));
+		return (ClassName == NAME_None) ? FRHITextureName : ClassName;
+	}
 
 	/* A friendly name for the resource. */
 	const TCHAR* DebugName = nullptr;
 
+	/* The RHI access state that the resource will be created in. */
+	ERHIAccess InitialState = ERHIAccess::Unknown;
+
+	/* Tells how to initialize (or not) the texture's data. */
+	ERHITextureInitAction InitAction = ERHITextureInitAction::Default;
+
 	/* Optional initial data to fill the resource with. */
 	FResourceBulkDataInterface* BulkData = nullptr;
 
-	FName ClassName = NAME_None;	// The owner class of FRHITexture used for Insight asset metadata tracing
-	FName OwnerName = NAME_None;	// The owner name used for Insight asset metadata tracing
+	/* The owner class of FRHITexture used for Insight asset metadata tracing */
+	FName ClassName = NAME_None;
+
+	/* The owner name used for Insight asset metadata tracing */
+	FName OwnerName = NAME_None;
 };
 
 class FRHITexture : public FRHIViewableResource
diff --git a/Engine/Source/Runtime/RHI/Public/RHITextureInitializer.h b/Engine/Source/Runtime/RHI/Public/RHITextureInitializer.h
new file mode 100644
--- /dev/null
+++ b/Engine/Source/Runtime/RHI/Public/RHITextureInitializer.h
@@ -0,0 +1,189 @@
+// Copyright Epic Games, Inc. All Rights Reserved.
+
+#pragma once
+
+#include "RHIResources.h"
+#include "Templates/Function.h"
+
+// 'Minimal' description of an RHITexture, primarily used by FRHITextureInitializer to keep overhead low.
+struct FRHITextureMinimalDesc
+{
+	FIntPoint         Extent = FIntPoint(1, 1);
+	uint16            Depth = 1;
+	uint16            ArraySize = 1;
+	ETextureDimension Dimension = ETextureDimension::Texture2D;
+	EPixelFormat      Format = PF_Unknown;
+
+	FRHITextureMinimalDesc() = default;
+	FRHITextureMinimalDesc(const FRHITextureDesc& InDesc)
+		: Extent(InDesc.Extent)
+		, Depth(InDesc.Depth)
+		, ArraySize(InDesc.ArraySize)
+		, Dimension(InDesc.Dimension)
+		, Format(InDesc.Format)
+	{
+	}
+};
+
+struct FRHITextureSubresourceInitializer
+{
+	void WriteData(const void* InSource, size_t InSize)
+	{
+		FMemory::Memcpy(Data, InSource, InSize);
+	}
+
+	void WriteColor(FColor InColor)
+	{
+		WriteData(&InColor, sizeof(InColor));
+	}
+
+	void*  Data;
+	uint64 Size;
+	uint64 Stride;
+};
+
+// Structure used to allow optimal texture initialization at creation time.
+// Should only ever be obtained by calling RHICreateTextureInitializer().
+// Texture data writing has to be done on individual subresources via the GetSubresource accesors.
+// NO COPIES ALLOWED
+struct FRHITextureInitializer
+{
+	FRHITextureInitializer() = default;
+	FRHITextureInitializer(FRHITextureInitializer&& InOther)
+		: FinalizeCallback(MoveTemp(InOther.FinalizeCallback))
+		, GetSubresourceCallback(MoveTemp(InOther.GetSubresourceCallback))
+		, CommandList(InOther.CommandList)
+		, Texture(InOther.Texture)
+		, WritableData(InOther.WritableData)
+		, WritableSize(InOther.WritableSize)
+		, Desc(InOther.Desc)
+	{
+		InOther.Reset();
+	}
+	~FRHITextureInitializer()
+	{
+		if (CommandList && Texture)
+		{
+			RemovePendingTextureUpload();
+		}
+	}
+
+	const FRHITextureMinimalDesc& GetDesc() const
+	{
+		return Desc;
+	}
+
+	struct FSubresourceIndex
+	{
+		int32 FaceIndex = 0;
+		int32 ArrayIndex = 0;
+		int32 MipIndex = 0;
+	};
+
+	// Get a subresource for any texture type.
+	FRHITextureSubresourceInitializer GetSubresource(FSubresourceIndex SubresourceIndex)
+	{
+		checkf(GetSubresourceCallback, TEXT("Attempting to call GetSubresourceCallback when it is null, make sure you set your InitAction to Initializer."));
+		return GetSubresourceCallback(SubresourceIndex);
+	}
+
+	// Get a subresource for a 2D texture. The only subresources are mip levels.
+	FRHITextureSubresourceInitializer GetTexture2DSubresource(int32 MipIndex)
+	{
+		FSubresourceIndex SubresourceIndex;
+		SubresourceIndex.MipIndex = MipIndex;
+		return GetSubresource(SubresourceIndex);
+	}
+
+	// Get a subresource for a 2D texture array. Subresources are individual array slice mip levels.
+	FRHITextureSubresourceInitializer GetTexture2DArraySubresource(int32 ArrayIndex, int32 MipIndex)
+	{
+		FSubresourceIndex SubresourceIndex;
+		SubresourceIndex.ArrayIndex = ArrayIndex;
+		SubresourceIndex.MipIndex = MipIndex;
+		return GetSubresource(SubresourceIndex);
+	}
+
+	// Get a subresource for a 3D texture. The only subresources are mip levels.
+	FRHITextureSubresourceInitializer GetTexture3DSubresource(int32 MipIndex)
+	{
+		FSubresourceIndex SubresourceIndex;
+		SubresourceIndex.MipIndex = MipIndex;
+		return GetSubresource(SubresourceIndex);
+	}
+
+	// Get a subresource for a cube texture. Subresources are individual face mip levels.
+	FRHITextureSubresourceInitializer GetTextureCubeSubresource(int32 FaceIndex, int32 MipIndex)
+	{
+		FSubresourceIndex SubresourceIndex;
+		SubresourceIndex.FaceIndex = FaceIndex;
+		SubresourceIndex.MipIndex = MipIndex;
+		return GetSubresource(SubresourceIndex);
+	}
+
+	// Get a subresource for a cube texture array. Subresources are individual array slice face mip levels.
+	FRHITextureSubresourceInitializer GetTextureCubeArraySubresource(int32 FaceIndex, int32 ArrayIndex, int32 MipIndex)
+	{
+		FSubresourceIndex SubresourceIndex;
+		SubresourceIndex.FaceIndex = FaceIndex;
+		SubresourceIndex.ArrayIndex = ArrayIndex;
+		SubresourceIndex.MipIndex = MipIndex;
+		return GetSubresource(SubresourceIndex);
+	}
+
+	// 'Finalizes' the initializer and returns the created FRHITexture. The initializer will be reset to an invalid state and should not be used after calling this.
+	RHI_API FTextureRHIRef Finalize();
+
+protected:
+	// @todo dev-pr switch to using IRHIUploadContext
+	using FFinalizeCallback = TUniqueFunction<FTextureRHIRef(FRHICommandListBase&)>;
+
+	// FaceIndex, ArrayIndex, MipIndex
+	using FGetSubresourceCallback = TUniqueFunction<FRHITextureSubresourceInitializer(FSubresourceIndex)>;
+
+	// Should only be called by RHI derived types.
+	RHI_API FRHITextureInitializer(FRHICommandListBase& RHICmdList, FRHITexture* InTexture, void* InWritableData, uint64 InWritableSize, FFinalizeCallback&& InFinalizeCallback, FGetSubresourceCallback&& InGetSubresourceCallback);
+
+	// Remove the texture from the command list. Has to be in cpp file to prevent circular header dependency.
+	RHI_API void RemovePendingTextureUpload();
+
+	// Allow copies only for RHI derived types.
+	FRHITextureInitializer(const FRHITextureInitializer&) = delete;
+	FRHITextureInitializer& operator=(const FRHITextureInitializer&) = delete;
+	FRHITextureInitializer& operator=(FRHITextureInitializer&&) = delete;
+
+	void Reset()
+	{
+		FinalizeCallback = {};
+		GetSubresourceCallback = {};
+		CommandList = {};
+		Texture = {};
+		WritableData = {};
+		WritableSize = {};
+		Desc = FRHITextureMinimalDesc();
+	}
+
+protected:
+	// RHI provided lambda to call when done writing data, returns the created Texture
+	// Only used by the RHI internals, should not be accessed outside of RHIs
+	FFinalizeCallback FinalizeCallback = nullptr;
+
+	// RHI provided lambda to provide pointers to and sizes of individual subresources for writing.
+	// Only used by the RHI internals, should not be accessed outside of RHIs
+	FGetSubresourceCallback GetSubresourceCallback = nullptr;
+
+	// Command list provided on construction, used in finalize.
+	FRHICommandListBase* CommandList = nullptr;
+
+	// Current RHI Texture being initialized. Will only be used for command list validation since each RHI implementation will manage their own texture type.
+	FRHITexture* Texture = nullptr;
+
+	// Pointer to the writable data provided by the RHI
+	void* WritableData = nullptr;
+
+	// Size of the writable data provided by the RHI
+	uint64 WritableSize = 0;
+
+	// Description of the texture being created.
+	FRHITextureMinimalDesc Desc{};
+};
diff --git a/Engine/Source/Runtime/RHI/Public/RHITextureUtils.h b/Engine/Source/Runtime/RHI/Public/RHITextureUtils.h
--- a/Engine/Source/Runtime/RHI/Public/RHITextureUtils.h
+++ b/Engine/Source/Runtime/RHI/Public/RHITextureUtils.h
@@ -6,9 +6,14 @@
 
 namespace UE::RHITextureUtils
 {
+	static uint32 CalculateMipExtent(uint32 Extent, uint32 MipIndex)
+	{
+		return FMath::Max<uint32>(Extent >> MipIndex, 1);
+	}
+
 	static uint32 CalculateMipBlockCount(uint32 Extent, uint32 MipIndex, uint32 BlockSize)
 	{
-		return FMath::Max<uint32>(1, FMath::DivideAndRoundUp(Extent >> MipIndex, BlockSize));
+		return FMath::Max<uint32>(1, FMath::DivideAndRoundUp(CalculateMipExtent(Extent, MipIndex), BlockSize));
 	}
 
 	static FUintVector3 CalculateMipBlockCounts(const FRHITextureDesc& Desc, uint32 MipIndex, const FPixelFormatInfo& PixelFormat)
@@ -26,6 +31,15 @@ namespace UE::RHITextureUtils
 		return CalculateMipBlockCounts(Desc, MipIndex, PixelFormat);
 	}
 
+	static FUintVector3 CalculateMipExtents(const FRHITextureDesc& Desc, uint32 MipIndex)
+	{
+		return FUintVector3(
+			CalculateMipExtent(Desc.Extent.X, MipIndex),
+			CalculateMipExtent(Desc.Extent.Y, MipIndex),
+			CalculateMipExtent(Desc.Depth,    MipIndex)
+		);
+	}
+
 	static uint64 CalculateTextureMipSize(const FRHITextureDesc& Desc, uint32 MipIndex, uint64& OutStride)
 	{
 		const FPixelFormatInfo& PixelFormat = GPixelFormats[Desc.Format];
@@ -38,6 +52,12 @@ namespace UE::RHITextureUtils
 		return SubresourceSize;
 	}
 
+	static uint64 CalculateTextureMipSize(const FRHITextureDesc& Desc, uint32 MipIndex)
+	{
+		uint64 TempStride{};
+		return CalculateTextureMipSize(Desc, MipIndex, TempStride);
+	}
+
 	static uint64 CalculateTexturePlaneSize(const FRHITextureDesc& Desc)
 	{
 		uint64 PlaneSize = 0;
diff --git a/Engine/Source/Runtime/RHI/Public/RHIUtilities.h b/Engine/Source/Runtime/RHI/Public/RHIUtilities.h
--- a/Engine/Source/Runtime/RHI/Public/RHIUtilities.h
+++ b/Engine/Source/Runtime/RHI/Public/RHIUtilities.h
@@ -328,7 +328,7 @@ struct FTextureReadBuffer2D
 		const FRHITextureCreateDesc Desc =
 			FRHITextureCreateDesc::Create2D(InDebugName, SizeX, SizeY, Format)
 			.SetFlags(Flags)
-			.SetBulkData(InBulkData);
+			.SetInitActionBulkData(InBulkData);
 
 		Buffer = RHICreateTexture(Desc);
 		SRV = RHICmdList.CreateShaderResourceView(Buffer, FRHIViewDesc::CreateTextureSRV().SetDimensionFromTexture(Buffer));
diff --git a/Engine/Source/Runtime/RHI/Public/RHIValidation.h b/Engine/Source/Runtime/RHI/Public/RHIValidation.h
--- a/Engine/Source/Runtime/RHI/Public/RHIValidation.h
+++ b/Engine/Source/Runtime/RHI/Public/RHIValidation.h
@@ -529,17 +529,31 @@ public:
 		return RHI->RHIGetTextureMemoryVisualizeData(TextureData, SizeX, SizeY, Pitch, PixelSize);
 	}
 
-	/**
-	* Creates an RHI texture resource.
-	*/
-	virtual FTextureRHIRef RHICreateTexture(FRHICommandListBase& RHICmdList, const FRHITextureCreateDesc& CreateDesc) override final
+	struct FRHIValidationTextureInitializer : public FRHITextureInitializer
+	{
+		FRHIValidationTextureInitializer(const FRHITextureCreateDesc& CreateDesc, FRHITextureInitializer&& InOther)
+			: FRHITextureInitializer(MoveTemp(InOther))
+		{
+			FFinalizeCallback OriginalFinalizeCallback = MoveTemp(FinalizeCallback);
+			FinalizeCallback = [OriginalFinalizeCallback = MoveTemp(OriginalFinalizeCallback)](FRHICommandListBase& RHICmdList)
+				{
+					FTextureRHIRef Texture = OriginalFinalizeCallback(RHICmdList);
+					ensure(Texture->IsBarrierTrackingInitialized());
+					return Texture;
+				};
+
+			// ue-todo: custom GetSubresourceCallback?
+		}
+	};
+
+	[[nodiscard]]
+	virtual FRHITextureInitializer RHICreateTextureInitializer(FRHICommandListBase& RHICmdList, const FRHITextureCreateDesc& CreateDesc) override final
 	{
 		CreateDesc.CheckValidity();
-		FTextureRHIRef Texture = RHI->RHICreateTexture(RHICmdList, CreateDesc);
-		ensure(Texture->IsBarrierTrackingInitialized());
-		return Texture;
+		return FRHIValidationTextureInitializer(CreateDesc, RHI->RHICreateTextureInitializer(RHICmdList, CreateDesc));
 	}
 
+
 	/**
 	* Thread-safe function that can be used to create a texture outside of the
 	* rendering thread. This function can ONLY be called if GRHISupportsAsyncTextureCreation
diff --git a/Engine/Source/Runtime/RHICore/Internal/RHICoreBufferInitializer.h b/Engine/Source/Runtime/RHICore/Internal/RHICoreBufferInitializer.h
--- a/Engine/Source/Runtime/RHICore/Internal/RHICoreBufferInitializer.h
+++ b/Engine/Source/Runtime/RHICore/Internal/RHICoreBufferInitializer.h
@@ -6,6 +6,7 @@
 #include "RHIBufferInitializer.h"
 #include "RHICore.h"
 #include "RHICommandList.h"
+#include "RHICoreInitializerCommon.h"
 
 namespace UE::RHICore
 {
@@ -46,33 +47,6 @@ namespace UE::RHICore
 		}
 	};
 
-	// Utility class for RHIs to use when they need to malloc aligned data up-front and make sure it's freed correctly by the Buffer Initializer lambda destruction
-	struct FScopedMemory : public FNoncopyable
-	{
-		FScopedMemory(void* InPointer)
-			: Pointer(InPointer)
-		{
-		}
-		FScopedMemory(size_t InSize, uint32 InAlignment)
-			: FScopedMemory(FMemory::Malloc(InSize, InAlignment))
-		{
-		}
-		FScopedMemory(FScopedMemory&& InOther)
-			: Pointer(InOther.Pointer)
-		{
-			InOther.Pointer = nullptr;
-		}
-		~FScopedMemory()
-		{
-			if (Pointer)
-			{
-				FMemory::Free(Pointer);
-			}
-		}
-
-		void* Pointer{};
-	};
-
 	// Buffer initializer with custom implementation. This type is necessary for access to the FRHIBufferInitializer protected constructor.
 	struct FCustomBufferInitializer : public FRHIBufferInitializer
 	{
diff --git a/Engine/Source/Runtime/RHICore/Internal/RHICoreInitializerCommon.h b/Engine/Source/Runtime/RHICore/Internal/RHICoreInitializerCommon.h
new file mode 100644
--- /dev/null
+++ b/Engine/Source/Runtime/RHICore/Internal/RHICoreInitializerCommon.h
@@ -0,0 +1,36 @@
+// Copyright Epic Games, Inc. All Rights Reserved.
+
+#pragma once
+
+#include "HAL/UnrealMemory.h"
+#include "Templates/UnrealTemplate.h"
+
+namespace UE::RHICore
+{
+	// Utility class for RHIs to use when they need to malloc aligned data up-front and make sure it's freed correctly by the Initializer lambda destruction
+	struct FInitializerScopedMemory : public FNoncopyable
+	{
+		FInitializerScopedMemory(void* InPointer)
+			: Pointer(InPointer)
+		{
+		}
+		FInitializerScopedMemory(size_t InSize, uint32 InAlignment)
+			: FInitializerScopedMemory(FMemory::Malloc(InSize, InAlignment))
+		{
+		}
+		FInitializerScopedMemory(FInitializerScopedMemory&& InOther)
+			: Pointer(InOther.Pointer)
+		{
+			InOther.Pointer = nullptr;
+		}
+		~FInitializerScopedMemory()
+		{
+			if (Pointer)
+			{
+				FMemory::Free(Pointer);
+			}
+		}
+
+		void* Pointer{};
+	};
+}
diff --git a/Engine/Source/Runtime/RHICore/Internal/RHICoreTextureInitializer.h b/Engine/Source/Runtime/RHICore/Internal/RHICoreTextureInitializer.h
new file mode 100644
--- /dev/null
+++ b/Engine/Source/Runtime/RHICore/Internal/RHICoreTextureInitializer.h
@@ -0,0 +1,130 @@
+// Copyright Epic Games, Inc. All Rights Reserved.
+
+#pragma once
+
+#include "Containers/ResourceArray.h"
+#include "RHITextureInitializer.h"
+#include "RHICore.h"
+#include "RHICommandList.h"
+#include "RHICoreInitializerCommon.h"
+#include "RHITextureUtils.h"
+
+namespace UE::RHICore
+{
+	struct FBaseTextureInitializerImplementation : public FRHITextureInitializer
+	{
+		FBaseTextureInitializerImplementation(FRHICommandListBase& RHICmdList, FRHITexture* InTexture, void* InWritableData, uint64 InWritableSize, FFinalizeCallback&& InFinalizeCallback, FGetSubresourceCallback&& InGetSubresourceCallback)
+			: FRHITextureInitializer(RHICmdList, InTexture, InWritableData, InWritableSize, Forward<FFinalizeCallback>(InFinalizeCallback), Forward<FGetSubresourceCallback>(InGetSubresourceCallback))
+		{
+		}
+		FBaseTextureInitializerImplementation(FRHICommandListBase& RHICmdList, FRHITexture* Texture, FFinalizeCallback&& InFinalizeCallback)
+			: FRHITextureInitializer(RHICmdList, Texture, nullptr, 0, Forward<FFinalizeCallback>(InFinalizeCallback), nullptr)
+		{
+		}
+
+		void* GetWritableData() const
+		{
+			return WritableData;
+		}
+
+		uint64 GetWritableSize() const
+		{
+			return WritableSize;
+		}
+	};
+
+	// Texture initializer that just returns the texture on finalize.
+	struct FDefaultTextureInitializer : public FBaseTextureInitializerImplementation
+	{
+		FDefaultTextureInitializer(FRHICommandListBase& RHICmdList, FRHITexture* Texture)
+			: FBaseTextureInitializerImplementation(RHICmdList, Texture,
+				[Texture = TRefCountPtr<FRHITexture>(Texture)](FRHICommandListBase&) mutable
+				{
+					return MoveTemp(Texture);
+				}
+			)
+		{
+		}
+	};
+
+	struct FDefaultLayoutTextureInitializer : public FBaseTextureInitializerImplementation
+	{
+		FDefaultLayoutTextureInitializer(FRHICommandListBase& RHICmdList, FRHITexture* InTexture, void* InMemory, uint64 InMemorySize, FFinalizeCallback&& FinalizeFunc)
+			: FBaseTextureInitializerImplementation(RHICmdList, InTexture, InMemory, InMemorySize, Forward<FFinalizeCallback>(FinalizeFunc),
+				[Texture = TRefCountPtr<FRHITexture>(InTexture), WritableData = InMemory](FRHITextureInitializer::FSubresourceIndex SubresourceIndex)
+				{
+					const FRHITextureDesc& TextureDesc = Texture->GetDesc();
+
+					uint64 SubresourceStride = 0;
+					uint64 SubresourceSize = 0;
+					const uint64 Offset = UE::RHITextureUtils::CalculateSubresourceOffset(TextureDesc, SubresourceIndex.FaceIndex, SubresourceIndex.ArrayIndex, SubresourceIndex.MipIndex, SubresourceStride, SubresourceSize);
+
+					FRHITextureSubresourceInitializer Result{};
+					Result.Data = reinterpret_cast<uint8*>(WritableData) + Offset;
+					Result.Stride = SubresourceStride;
+					Result.Size = SubresourceSize;
+
+					return Result;
+				}
+			)
+		{
+		}
+	};
+
+	// Texture Initializer that uses Lock/Unlock to upload initial data
+	struct FLockTextureInitializer : public FDefaultLayoutTextureInitializer
+	{
+		FLockTextureInitializer(FRHICommandListBase& RHICmdList, FRHITexture* InTexture, void* InMemory, uint64 InMemorySize)
+			: FDefaultLayoutTextureInitializer(RHICmdList, InTexture, InMemory, InMemorySize,
+				[Texture = TRefCountPtr<FRHITexture>(InTexture), WritableMemory = FInitializerScopedMemory(InMemory)](FRHICommandListBase& InRHICmdList) mutable
+				{
+					FRHICommandListImmediate& RHICmdList = InRHICmdList.GetAsImmediate();
+
+					const FRHITextureDesc& TextureDesc = Texture->GetDesc();
+					const uint32 FaceCount = TextureDesc.IsTextureCube() ? 6 : 1;
+
+					for (uint32 FaceIndex = 0; FaceIndex < FaceCount; FaceIndex++)
+					{
+						for (uint32 ArrayIndex = 0; ArrayIndex < TextureDesc.ArraySize; ArrayIndex++)
+						{
+							for (uint32 MipIndex = 0; MipIndex < TextureDesc.NumMips; MipIndex++)
+							{
+								uint64 SubresourceStride = 0;
+								uint64 SubresourceSize = 0;
+								const uint64 Offset = UE::RHITextureUtils::CalculateSubresourceOffset(TextureDesc, FaceIndex, ArrayIndex, MipIndex, SubresourceStride, SubresourceSize);
+
+								const FRHILockTextureArgs LockArgs = FRHILockTextureArgs::LockCubeFace(Texture, FaceIndex, ArrayIndex, MipIndex, RLM_WriteOnly, false);
+
+								FRHILockTextureResult LockResult = RHICmdList.LockTexture(LockArgs);
+								check(LockResult.ByteCount <= SubresourceSize);
+								check(LockResult.Stride == SubresourceStride);
+
+								FMemory::Memcpy(LockResult.Data, reinterpret_cast<uint8*>(WritableMemory.Pointer) + Offset, SubresourceSize);
+
+								RHICmdList.UnlockTexture(LockArgs);
+							}
+						}
+					}
+
+					return MoveTemp(Texture);
+				})
+		{
+		}
+	};
+
+	static FLockTextureInitializer CreateLockTextureInitializer(FRHICommandListBase& RHICmdList, FRHITexture* Texture, const FRHITextureCreateDesc& CreateDesc)
+	{
+		const uint64 WritableMemorySize = UE::RHITextureUtils::CalculateTextureSize(CreateDesc);
+		void* WritableMemory = FMemory::Malloc(WritableMemorySize, 16);
+
+		return FLockTextureInitializer(RHICmdList, Texture, WritableMemory, WritableMemorySize);
+	}
+
+	static FRHITextureInitializer HandleUnknownTextureInitializerInitAction(FRHICommandListBase& RHICmdList, const FRHITextureCreateDesc& CreateDesc)
+	{
+		UE_LOG(LogRHICore, Fatal, TEXT("Unknown or unhandled ERHITextureInitAction: %d"), static_cast<uint32>(CreateDesc.InitAction));
+
+		FRHITexture* Texture = nullptr;
+		return UE::RHICore::FDefaultTextureInitializer(RHICmdList, Texture);
+	}
+}
diff --git a/Engine/Source/Runtime/RenderCore/Private/GlobalRenderResources.cpp b/Engine/Source/Runtime/RenderCore/Private/GlobalRenderResources.cpp
--- a/Engine/Source/Runtime/RenderCore/Private/GlobalRenderResources.cpp
+++ b/Engine/Source/Runtime/RenderCore/Private/GlobalRenderResources.cpp
@@ -50,17 +50,20 @@ public:
 	// FResource interface.
 	virtual void InitRHI(FRHICommandListBase& RHICmdList) override
 	{
-		// Create the texture RHI.
-		FColorBulkData BulkData(R, G, B, A);
-
 		// BGRA typed UAV is unsupported per D3D spec, use RGBA here.
 		const FRHITextureCreateDesc Desc =
 			FRHITextureCreateDesc::Create2D(TEXT("ColoredTexture"), 1, 1, PF_R8G8B8A8)
 			.SetFlags(ETextureCreateFlags::ShaderResource)
-			.SetBulkData(&BulkData)
-			.SetClassName(TEXT("FColoredTexture"));
+			.SetClassName(TEXT("FColoredTexture"))
+			.SetInitActionInitializer();
 
-		TextureRHI = RHICreateTexture(Desc);
+		FRHITextureInitializer Initializer = RHICmdList.CreateTextureInitializer(Desc);
+
+		const FColor ColorValue(R, G, B, A);
+		Initializer.GetTexture2DSubresource(0).WriteColor(ColorValue);
+
+		// Create the texture RHI.
+		TextureRHI = Initializer.Finalize();
 
 		// Create the sampler state RHI resource.
 		FSamplerStateInitializerRHI SamplerStateInitializer(SF_Point, AM_Wrap, AM_Wrap, AM_Wrap);
@@ -276,31 +279,18 @@ class FBlackVolumeTexture : public FTexture
 public:
 	virtual void InitRHI(FRHICommandListBase& RHICmdList) override
 	{
-		if (GSupportsTexture3D)
-		{
-			// Create the texture.
-			FColorBulkData BulkData(Alpha);
+		FColorBulkData BulkData(Alpha);
 
-			const FRHITextureCreateDesc Desc =
-				FRHITextureCreateDesc::Create3D(TEXT("BlackVolumeTexture3D"), 1, 1, 1, PixelFormat)
-				.SetFlags(ETextureCreateFlags::ShaderResource)
-				.SetBulkData(&BulkData)
-				.SetClassName(TEXT("FBlackVolumeTexture"));
+		FRHITextureCreateDesc Desc =
+			GSupportsTexture3D
+			? FRHITextureCreateDesc::Create3D(TEXT("BlackVolumeTexture3D"), 1, 1, 1, PixelFormat)
+			: FRHITextureCreateDesc::Create2D(TEXT("BlackVolumeTexture2D"), 1, 1, PixelFormat);
 
-			TextureRHI = RHICreateTexture(Desc);
-		}
-		else
-		{
-			// Create a texture, even though it's not a volume texture
-			FColorBulkData BulkData(Alpha);
-
-			const FRHITextureCreateDesc Desc =
-				FRHITextureCreateDesc::Create2D(TEXT("BlackVolumeTexture2D"), 1, 1, PixelFormat)
-				.SetFlags(ETextureCreateFlags::ShaderResource)
-				.SetBulkData(&BulkData);
+		Desc.SetFlags(ETextureCreateFlags::ShaderResource)
+			.SetClassName(TEXT("FBlackVolumeTexture"))
+			.SetInitActionBulkData(&BulkData);
 
-			TextureRHI = RHICreateTexture(Desc);
-		}
+		TextureRHI = RHICmdList.CreateTexture(Desc);
 
 		// Create the sampler state.
 		FSamplerStateInitializerRHI SamplerStateInitializer(SF_Point, AM_Wrap, AM_Wrap, AM_Wrap);
@@ -323,16 +313,19 @@ class FBlackArrayTexture : public FTexture
 public:
 	virtual void InitRHI(FRHICommandListBase& RHICmdList) override
 	{
-		// Create the texture RHI.
-		FColorBulkData BulkData(0);
-
 		const FRHITextureCreateDesc Desc =
 			FRHITextureCreateDesc::Create2DArray(TEXT("BlackArrayTexture"), 1, 1, 1, PF_B8G8R8A8)
 			.SetFlags(ETextureCreateFlags::ShaderResource)
-			.SetBulkData(&BulkData)
-			.SetClassName(TEXT("FBlackArrayTexture"));
+			.SetClassName(TEXT("FBlackArrayTexture"))
+			.SetInitActionInitializer();
+
+		FRHITextureInitializer Initializer = RHICmdList.CreateTextureInitializer(Desc);
 
-		TextureRHI = RHICreateTexture(Desc);
+		const FColor ColorValue(0, 0, 0, 0);
+		Initializer.GetTexture2DSubresource(0).WriteColor(ColorValue);
+
+		// Create the texture RHI.
+		TextureRHI = Initializer.Finalize();
 
 		// Create the sampler state RHI resource.
 		FSamplerStateInitializerRHI SamplerStateInitializer(SF_Point, AM_Wrap, AM_Wrap, AM_Wrap);
@@ -371,28 +364,35 @@ public:
 			FRHITextureCreateDesc::Create2D(TEXT("FMipColorTexture"), TextureSize, TextureSize, PF_B8G8R8A8)
 			.SetNumMips(NumMips)
 			.SetFlags(ETextureCreateFlags::ShaderResource)
-			.SetClassName(TEXT("FMipColorTexture"));
+			.SetClassName(TEXT("FMipColorTexture"))
+			.SetInitActionInitializer();
 
-		TextureRHI = RHICreateTexture(Desc);
+		FRHITextureInitializer Initializer = RHICmdList.CreateTextureInitializer(Desc);
 
 		// Write the contents of the texture.
-		uint32 DestStride;
 		int32 Size = TextureSize;
 		for (int32 MipIndex = 0; MipIndex < NumMips; ++MipIndex)
 		{
-			FColor* DestBuffer = (FColor*)RHILockTexture2D(TextureRHI, MipIndex, RLM_WriteOnly, DestStride, false);
+			const FRHITextureSubresourceInitializer Subresource = Initializer.GetTexture2DSubresource(MipIndex);
+
+			FColor* DestBuffer = reinterpret_cast<FColor*>(Subresource.Data);
+			uint64 DestBufferOffset = 0;
+
 			for (int32 Y = 0; Y < Size; ++Y)
 			{
 				for (int32 X = 0; X < Size; ++X)
 				{
-					DestBuffer[X] = MipColors[NumMips - 1 - MipIndex];
+					DestBuffer[X + DestBufferOffset] = MipColors[NumMips - 1 - MipIndex];
 				}
-				DestBuffer += DestStride / sizeof(FColor);
+				DestBufferOffset += Subresource.Stride / sizeof(FColor);
+				check(DestBufferOffset <= Subresource.Size);
 			}
-			RHIUnlockTexture2D(TextureRHI, MipIndex, false);
+
 			Size >>= 1;
 		}
 
+		TextureRHI = Initializer.Finalize();
+
 		// Create the sampler state RHI resource.
 		FSamplerStateInitializerRHI SamplerStateInitializer(SF_Point, AM_Wrap, AM_Wrap, AM_Wrap);
 		SamplerStateRHI = GetOrCreateSamplerState(SamplerStateInitializer);
@@ -457,31 +457,30 @@ public:
 	virtual void InitRHI(FRHICommandListBase& RHICmdList) override
 	{
 		// Create the texture RHI.
-		const FRHITextureCreateDesc Desc =
+		const FRHITextureCreateDesc CreateDesc =
 			FRHITextureCreateDesc::CreateCube(TEXT("SolidColorCube"), 1, PixelFormat)
 			.SetFlags(ETextureCreateFlags::ShaderResource)
 			.SetClassName(TEXT("FSolidColorTextureCube"))
-			.SetInitialState(ERHIAccess::SRVMask);
+			.SetInitialState(ERHIAccess::SRVMask)
+			.SetInitActionInitializer();
 
-		FTextureRHIRef TextureCube = RHICreateTexture(Desc);
-		TextureRHI = TextureCube;
+		FRHITextureInitializer Initializer = RHICmdList.CreateTextureInitializer(CreateDesc);
 
-		// Write the contents of the texture.
 		for (uint32 FaceIndex = 0; FaceIndex < 6; FaceIndex++)
 		{
-			uint32 DestStride;
-			void* DestBuffer = RHILockTextureCubeFace(TextureCube, FaceIndex, 0, 0, RLM_WriteOnly, DestStride, false);
+			FRHITextureSubresourceInitializer Subresource = Initializer.GetTextureCubeSubresource(FaceIndex, 0);
 			if (bInitToZero)
 			{
-				FMemory::Memzero(DestBuffer, GPixelFormats[PixelFormat].BlockBytes);
+				FMemory::Memzero(Subresource.Data, GPixelFormats[PixelFormat].BlockBytes);
 			}
 			else
 			{
-				FMemory::Memcpy(DestBuffer, &ColorData, sizeof(ColorData));
+				FMemory::Memcpy(Subresource.Data, &ColorData, sizeof(ColorData));
 			}
-			RHIUnlockTextureCubeFace(TextureCube, FaceIndex, 0, 0, false);
 		}
 
+		TextureRHI = Initializer.Finalize();
+
 		// Create the sampler state RHI resource.
 		FSamplerStateInitializerRHI SamplerStateInitializer(SF_Point, AM_Wrap, AM_Wrap, AM_Wrap);
 		SamplerStateRHI = GetOrCreateSamplerState(SamplerStateInitializer);
@@ -511,26 +510,26 @@ public:
 	// FResource interface.
 	virtual void InitRHI(FRHICommandListBase& RHICmdList) override
 	{
-		const TCHAR* Name = TEXT("BlackCubeArray");
-
-		const FRHITextureCreateDesc Desc =
+		const FRHITextureCreateDesc CreateDesc =
 			FRHITextureCreateDesc::CreateCubeArray(TEXT("BlackCubeArray"), 1, 1, PF_B8G8R8A8)
 			.SetFlags(ETextureCreateFlags::ShaderResource)
 			.SetClassName(TEXT("FBlackCubeArrayTexture"))
-			.SetInitialState(ERHIAccess::SRVMask);
+			.SetInitialState(ERHIAccess::SRVMask)
+			.SetInitActionInitializer();
 
-		// Create the texture RHI.
-		TextureRHI = RHICreateTexture(Desc);
+		FRHITextureInitializer Initializer = RHICmdList.CreateTextureInitializer(CreateDesc);
 
 		for (uint32 FaceIndex = 0; FaceIndex < 6; FaceIndex++)
 		{
-			uint32 DestStride;
-			FColor* DestBuffer = (FColor*)RHILockTextureCubeFace(TextureRHI, FaceIndex, 0, 0, RLM_WriteOnly, DestStride, false);
+			FRHITextureSubresourceInitializer Subresource = Initializer.GetTextureCubeSubresource(FaceIndex, 0);
+
 			// Note: alpha is used by reflection environment to say how much of the foreground texture is visible, so 0 says it is completely invisible
-			*DestBuffer = FColor(0, 0, 0, 0);
-			RHIUnlockTextureCubeFace(TextureRHI, FaceIndex, 0, 0, false);
+			Subresource.WriteColor(FColor(0, 0, 0, 0));
 		}
 
+		// Create the texture RHI.
+		TextureRHI = Initializer.Finalize();
+
 		// Create the sampler state RHI resource.
 		FSamplerStateInitializerRHI SamplerStateInitializer(SF_Point, AM_Wrap, AM_Wrap, AM_Wrap);
 		SamplerStateRHI = GetOrCreateSamplerState(SamplerStateInitializer);
@@ -544,25 +543,24 @@ FTexture* GBlackCubeArrayTexture = new TGlobalResource<FBlackCubeArrayTexture, F
 /**
  * A UINT 1x1 texture.
  */
-template <EPixelFormat Format, uint32 R = 0, uint32 G = 0, uint32 B = 0, uint32 A = 0>
-class FUintTexture : public FTextureWithSRV
+class FBlackUintTexture : public FTextureWithSRV
 {
 public:
 	// FResource interface.
 	virtual void InitRHI(FRHICommandListBase& RHICmdList) override
 	{
-		const FRHITextureCreateDesc Desc =
-			FRHITextureCreateDesc::Create2D(TEXT("UintTexture"), 1, 1, Format)
+		const FRHITextureCreateDesc CreateDesc =
+			FRHITextureCreateDesc::Create2D(TEXT("BlackUintTexture"), 1, 1, PF_R32G32B32A32_UINT)
 			.SetFlags(ETextureCreateFlags::ShaderResource)
-			.SetClassName(TEXT("FUintTexture"));
+			.SetClassName(TEXT("FBlackUintTexture"))
+			.SetInitActionInitializer();
 
-		TextureRHI = RHICreateTexture(Desc);
+		FRHITextureInitializer Initializer = RHICmdList.CreateTextureInitializer(CreateDesc);
 
-		// Write the contents of the texture.
-		uint32 DestStride;
-		void* DestBuffer = RHILockTexture2D(TextureRHI, 0, RLM_WriteOnly, DestStride, false);
-		WriteData(DestBuffer);
-		RHIUnlockTexture2D(TextureRHI, 0, false);
+		FRHITextureSubresourceInitializer Subresource = Initializer.GetTextureCubeSubresource(0, 0);
+		FMemory::Memzero(Subresource.Data, Subresource.Size);
+
+		TextureRHI = Initializer.Finalize();
 
 		// Create the sampler state RHI resource.
 		FSamplerStateInitializerRHI SamplerStateInitializer(SF_Point, AM_Wrap, AM_Wrap, AM_Wrap);
@@ -574,48 +572,9 @@ public:
 
 	virtual uint32 GetSizeX() const override { return 1; }
 	virtual uint32 GetSizeY() const override { return 1; }
-
-protected:
-	static int32 GetNumChannels()
-	{
-		return GPixelFormats[Format].NumComponents;
-	}
-
-	static int32 GetBytesPerChannel()
-	{
-		return GPixelFormats[Format].BlockBytes / GPixelFormats[Format].NumComponents;
-	}
-
-	template<typename T>
-	static void DoWriteData(T* DataPtr)
-	{
-		T Values[] = { R, G, B, A };
-		for (int32 i = 0; i < GetNumChannels(); ++i)
-		{
-			DataPtr[i] = Values[i];
-		}
-	}
-
-	static void WriteData(void* DataPtr)
-	{
-		switch (GetBytesPerChannel())
-		{
-		case 1:
-			DoWriteData((uint8*)DataPtr);
-			return;
-		case 2:
-			DoWriteData((uint16*)DataPtr);
-			return;
-		case 4:
-			DoWriteData((uint32*)DataPtr);
-			return;
-		}
-		// Unsupported format
-		check(0);
-	}
 };
 
-FTexture* GBlackUintTexture = new TGlobalResource<FUintTexture<PF_R32G32B32A32_UINT>, FRenderResource::EInitPhase::Pre>;
+FTexture* GBlackUintTexture = new TGlobalResource<FBlackUintTexture, FRenderResource::EInitPhase::Pre>;
 
 
 //////////////////////////////////////////////////////////////////////////////////////////////////////////////////////
diff --git a/Engine/Source/Runtime/Renderer/Private/SystemTextures.cpp b/Engine/Source/Runtime/Renderer/Private/SystemTextures.cpp
--- a/Engine/Source/Runtime/Renderer/Private/SystemTextures.cpp
+++ b/Engine/Source/Runtime/Renderer/Private/SystemTextures.cpp
@@ -14,6 +14,7 @@
 #include "GlobalRenderResources.h"
 #include "Engine/Texture2D.h"
 #include "Engine/Engine.h"
+#include "Containers/ResourceArray.h"
 
 /*-----------------------------------------------------------------------------
  Export System textures (LTC / EnergyConservation / Hair LUT / ...)
@@ -172,25 +173,20 @@ void FSystemTextures::InitializeTextures(FRHICommandListImmediate& RHICmdList, c
 	}
 }
 
-template <typename DataType>
-void SetDummyTextureData(FRHITexture* Texture, const DataType& DummyData)
-{
-	uint32 DestStride;
-	DataType* DestBuffer = (DataType*)RHILockTexture2D(Texture, 0, RLM_WriteOnly, DestStride, false);
-	*DestBuffer = DummyData;
-	RHIUnlockTexture2D(Texture, 0, false);
-}
+const static FLazyName SystemTexturesName(TEXT("FSystemTextures"));
 
 template <typename DataType>
-void SetDummyTextureArrayData(FRHITexture* Texture, const DataType& DummyData)
+static FTextureRHIRef CreateDummyTexture(FRHICommandListImmediate& RHICmdList, const FRHITextureCreateDesc& InCreateDesc, const DataType& DummyData)
 {
-	uint32 DestStride;
-	DataType* DestBuffer = (DataType*)RHILockTexture2DArray(Texture, 0, 0, RLM_WriteOnly, DestStride, false);
-	*DestBuffer = DummyData;
-	RHIUnlockTexture2DArray(Texture, 0, 0, false);
-}
+	FResourceBulkDataArrayView InitData(&DummyData, sizeof(DataType));
 
-const static FLazyName SystemTexturesName(TEXT("FSystemTextures"));
+	FRHITextureCreateDesc CreateDesc(InCreateDesc);
+	CreateDesc.AddFlags(ETextureCreateFlags::ShaderResource);
+	CreateDesc.SetClassName(SystemTexturesName);
+	CreateDesc.SetInitActionBulkData(&InitData);
+
+	return RHICmdList.CreateTexture(CreateDesc);
+}
 
 static TRefCountPtr<IPooledRenderTarget> CreateTexture(FRHICommandListImmediate& RHICmdList, TObjectPtr<class UTexture2D>& InCPUTexture, EPixelFormat InFormat, const TCHAR* InName)
 {
@@ -202,23 +198,25 @@ static TRefCountPtr<IPooledRenderTarget> CreateTexture(FRHICommandListImmediate&
 	const FIntPoint DataSize(Data->SizeX, Data->SizeY);
 
 	const FRHITextureCreateDesc Desc =
-	FRHITextureCreateDesc::Create2D(InName, DataSize.X, DataSize.Y, InFormat)
-	.SetFlags(ETextureCreateFlags::ShaderResource | ETextureCreateFlags::FastVRAM)
-	.SetClassName(SystemTexturesName);
+		FRHITextureCreateDesc::Create2D(InName, DataSize.X, DataSize.Y, InFormat)
+		.SetFlags(ETextureCreateFlags::ShaderResource | ETextureCreateFlags::FastVRAM)
+		.SetClassName(SystemTexturesName)
+		.SetInitActionInitializer();
 
-	FTextureRHIRef Texture = RHICreateTexture(Desc);
+	FRHITextureInitializer Initializer = RHICmdList.CreateTextureInitializer(Desc);
 
 	check(InFormat == PF_FloatRGBA || InFormat == PF_G16R16F);
 	const uint32 ComponentCount = InFormat == PF_G16R16F ? 2u: 4u;
 
 	// Write the contents of the texture.
-	uint32 DestStride;
-	uint8* DestBuffer = (uint8*)RHICmdList.LockTexture2D(Texture, 0, RLM_WriteOnly, DestStride, false);
+	FRHITextureSubresourceInitializer Subresource = Initializer.GetTexture2DSubresource(0);
+
+	uint8* DestBuffer = reinterpret_cast<uint8*>(Subresource.Data);
 	for (int32 y = 0; y < DataSize.Y; ++y)
 	{
 		for (int32 x = 0; x < DataSize.X; ++x)
 		{
-			uint16* Dest = (uint16*)(DestBuffer + x * ComponentCount * sizeof(uint16) + y * DestStride);
+			uint16* Dest = (uint16*)(DestBuffer + x * ComponentCount * sizeof(uint16) + y * Subresource.Stride);
 
 			const FFloat16Color Value = DataView[x + y * DataSize.X];
 			Dest[0] = Value.R.Encoded;
@@ -227,7 +225,8 @@ static TRefCountPtr<IPooledRenderTarget> CreateTexture(FRHICommandListImmediate&
 			if (ComponentCount > 3) { Dest[3] = Value.A.Encoded; }
 		}
 	}
-	RHICmdList.UnlockTexture2D(Texture, 0, false);
+
+	FTextureRHIRef Texture = Initializer.Finalize();
 
 	// Release CPU data which are no longer needed
 	#if !WITH_EDITORONLY_DATA
@@ -257,60 +256,45 @@ void FSystemTextures::InitializeCommonTextures(FRHICommandListImmediate& RHICmdL
 	// Create a BlackDummy texture
 	{
 		const FRHITextureCreateDesc Desc =
-			FRHITextureCreateDesc::Create2D(TEXT("BlackDummy"), 1, 1, PF_B8G8R8A8)
-			.SetFlags(ETextureCreateFlags::ShaderResource)
-			.SetClassName(SystemTexturesName);
+			FRHITextureCreateDesc::Create2D(TEXT("BlackDummy"), 1, 1, PF_B8G8R8A8);
 
-		FTextureRHIRef Texture = RHICreateTexture(Desc);
-		SetDummyTextureData<FColor>(Texture, FColor(0, 0, 0, 0));
+		FTextureRHIRef Texture = CreateDummyTexture<FColor>(RHICmdList, Desc, FColor(0, 0, 0, 0));
 		BlackDummy = CreateRenderTarget(Texture, Desc.DebugName);
 	}
 	
 	// Create a texture array that has a single black dummy slice 
 	{
 		const FRHITextureCreateDesc Desc =
-			FRHITextureCreateDesc::Create2DArray(TEXT("BlackArrayDummy"), 1, 1, 1, PF_B8G8R8A8)
-			.SetFlags(ETextureCreateFlags::ShaderResource)
-			.SetClassName(SystemTexturesName);
+			FRHITextureCreateDesc::Create2DArray(TEXT("BlackArrayDummy"), 1, 1, 1, PF_B8G8R8A8);
 
-		FTextureRHIRef Texture = RHICreateTexture(Desc);
-		SetDummyTextureArrayData<FColor>(Texture, FColor(0, 0, 0, 0));
+		FTextureRHIRef Texture = CreateDummyTexture<FColor>(RHICmdList, Desc, FColor(0, 0, 0, 0));
 		BlackArrayDummy = CreateRenderTarget(Texture, Desc.DebugName);
 	}
 
 	// Create a texture that is a single UInt32 value set to 0
 	{
 		const FRHITextureCreateDesc Desc =
-			FRHITextureCreateDesc::Create2D(TEXT("ZeroUIntDummy"), 1, 1, PF_R32_UINT)
-			.SetFlags(ETextureCreateFlags::ShaderResource)
-			.SetClassName(SystemTexturesName);
+			FRHITextureCreateDesc::Create2D(TEXT("ZeroUIntDummy"), 1, 1, PF_R32_UINT);
 
-		FTextureRHIRef Texture = RHICreateTexture(Desc);
-		SetDummyTextureData<uint32>(Texture, 0u);
+		FTextureRHIRef Texture = CreateDummyTexture<uint32>(RHICmdList, Desc, 0u);
 		ZeroUIntDummy = CreateRenderTarget(Texture, Desc.DebugName);
 	}
 
 	// Create a texture array that is a single UInt32 value set to 0
 	{
 		const FRHITextureCreateDesc Desc =
-			FRHITextureCreateDesc::Create2DArray(TEXT("ZeroUIntArrayDummy"), 1, 1, 1, PF_R32_UINT)
-			.SetFlags(ETextureCreateFlags::ShaderResource)
-			.SetClassName(SystemTexturesName);
+			FRHITextureCreateDesc::Create2DArray(TEXT("ZeroUIntArrayDummy"), 1, 1, 1, PF_R32_UINT);
 
-		FTextureRHIRef Texture = RHICreateTexture(Desc);
-		SetDummyTextureArrayData<uint32>(Texture, 0u);
+		FTextureRHIRef Texture = CreateDummyTexture<uint32>(RHICmdList, Desc, 0u);
 		ZeroUIntArrayDummy = CreateRenderTarget(Texture, Desc.DebugName);
 	}
 
 	// Create a BlackAlphaOneDummy texture
 	{
 		const FRHITextureCreateDesc Desc =
-			FRHITextureCreateDesc::Create2D(TEXT("BlackAlphaOneDummy"), 1, 1, PF_B8G8R8A8)
-			.SetFlags(ETextureCreateFlags::ShaderResource)
-			.SetClassName(SystemTexturesName);
+			FRHITextureCreateDesc::Create2D(TEXT("BlackAlphaOneDummy"), 1, 1, PF_B8G8R8A8);
 
-		FTextureRHIRef Texture = RHICreateTexture(Desc);
-		SetDummyTextureData<FColor>(Texture, FColor(0, 0, 0, 255));
+		FTextureRHIRef Texture = CreateDummyTexture<FColor>(RHICmdList, Desc, FColor(0, 0, 0, 255));
 		BlackAlphaOneDummy = CreateRenderTarget(Texture, Desc.DebugName);
 	}
 
@@ -321,20 +305,16 @@ void FSystemTextures::InitializeCommonTextures(FRHICommandListImmediate& RHICmdL
 			.SetFlags(ETextureCreateFlags::ShaderResource)
 			.SetClassName(SystemTexturesName);
 
-		FTextureRHIRef Texture = RHICreateTexture(Desc);
-		SetDummyTextureData<FColor>(Texture, FColor(0, 255, 0, 255));
+		FTextureRHIRef Texture = CreateDummyTexture<FColor>(RHICmdList, Desc, FColor(0, 255, 0, 255));
 		GreenDummy = CreateRenderTarget(Texture, Desc.DebugName);
 	}
 
 	// Create a DefaultNormal8Bit texture
 	{
 		const FRHITextureCreateDesc Desc =
-			FRHITextureCreateDesc::Create2D(TEXT("DefaultNormal8Bit"), 1, 1, PF_B8G8R8A8)
-			.SetFlags(ETextureCreateFlags::ShaderResource)
-			.SetClassName(SystemTexturesName);
+			FRHITextureCreateDesc::Create2D(TEXT("DefaultNormal8Bit"), 1, 1, PF_B8G8R8A8);
 
-		FTextureRHIRef Texture = RHICreateTexture(Desc);
-		SetDummyTextureData<FColor>(Texture, FColor(128, 128, 128, 255));
+		FTextureRHIRef Texture = CreateDummyTexture<FColor>(RHICmdList, Desc, FColor(128, 128, 128, 255));
 		DefaultNormal8Bit = CreateRenderTarget(Texture, Desc.DebugName);
 	}
 
@@ -345,13 +325,16 @@ void FSystemTextures::InitializeCommonTextures(FRHICommandListImmediate& RHICmdL
 		const FRHITextureCreateDesc Desc =
 			FRHITextureCreateDesc::Create2D(TEXT("PerlinNoiseGradient"), Extent, Extent, PF_B8G8R8A8)
 			.SetFlags(ETextureCreateFlags::ShaderResource)
-			.SetClassName(SystemTexturesName);
-	
-		FTextureRHIRef Texture = RHICreateTexture(Desc);
+			.SetClassName(SystemTexturesName)
+			.SetInitActionInitializer();
+
+		FRHITextureInitializer Initializer = RHICmdList.CreateTextureInitializer(Desc);
 
 		// Write the contents of the texture.
-		uint32 DestStride;
-		uint8* DestBuffer = (uint8*)RHICmdList.LockTexture2D(Texture, 0, RLM_WriteOnly, DestStride, false);
+		FRHITextureSubresourceInitializer Subresource = Initializer.GetTexture2DSubresource(0);
+
+		uint8* DestBuffer = reinterpret_cast<uint8*>(Subresource.Data);
+
 		// seed the pseudo random stream with a good value
 		FRandomStream RandomStream(12345);
 		// Values represent float3 values in the -1..1 range.
@@ -367,44 +350,37 @@ void FSystemTextures::InitializeCommonTextures(FRHICommandListImmediate& RHICmdL
 		{
 			for (int32 x = 0; x < Extent; ++x)
 			{
-				uint32* Dest = (uint32*)(DestBuffer + x * sizeof(uint32) + y * DestStride);
+				uint32* Dest = (uint32*)(DestBuffer + x * sizeof(uint32) + y * Subresource.Stride);
 
 				// pick a random direction (hacky way to overcome the quality issues FRandomStream has)
 				*Dest = gradtable[(uint32)(RandomStream.GetFraction() * 11.9999999f)];
 			}
 		}
-		RHICmdList.UnlockTexture2D(Texture, 0, false);
 
+		FTextureRHIRef Texture = Initializer.Finalize();
 		PerlinNoiseGradient = CreateRenderTarget(Texture, Desc.DebugName);
 	}
 
 	if (GPixelFormats[PF_FloatRGBA].Supported)
 	{
 		const FRHITextureCreateDesc Desc =
-			FRHITextureCreateDesc::Create2D(TEXT("MaxFP16Depth"), 1, 1, PF_FloatRGBA)
-			.SetFlags(ETextureCreateFlags::ShaderResource)
-			.SetClassName(SystemTexturesName);
+			FRHITextureCreateDesc::Create2D(TEXT("MaxFP16Depth"), 1, 1, PF_FloatRGBA);
 
-		FTextureRHIRef Texture = RHICreateTexture(Desc);
-		SetDummyTextureData<FFloat16Color>(Texture, FFloat16Color(FLinearColor(65500.0f, 65500.0f, 65500.0f, 65500.0f)));
+		FTextureRHIRef Texture = CreateDummyTexture<FFloat16Color>(RHICmdList, Desc, FFloat16Color(FLinearColor(65500.0f, 65500.0f, 65500.0f, 65500.0f)));
 		MaxFP16Depth = CreateRenderTarget(Texture, Desc.DebugName);
 	}
 
 	{
 		DepthDummy = static_cast<bool>(ERHIZBuffer::IsInverted) ? BlackDummy : WhiteDummy;
-
 		BlackDepthCube = CreateRenderTarget(GBlackTextureDepthCube->TextureRHI, TEXT("BlackDepthCube"));
 	}
 
 	// Create a dummy stencil SRV.
 	{
 		const FRHITextureCreateDesc Desc =
-			FRHITextureCreateDesc::Create2D(TEXT("StencilDummy"), 1, 1, PF_R8G8B8A8_UINT)
-			.SetFlags(ETextureCreateFlags::ShaderResource)
-			.SetClassName(SystemTexturesName);
+			FRHITextureCreateDesc::Create2D(TEXT("StencilDummy"), 1, 1, PF_R8G8B8A8_UINT);
 
-		FTextureRHIRef Texture = RHICreateTexture(Desc);
-		SetDummyTextureData<FColor>(Texture, FColor::White);
+		FTextureRHIRef Texture = CreateDummyTexture<FColor>(RHICmdList, Desc, FColor::White);
 		StencilDummy = CreateRenderTarget(Texture, Desc.DebugName);
 		StencilDummySRV = RHICmdList.CreateShaderResourceView(StencilDummy->GetRHI(), FRHIViewDesc::CreateTextureSRV().SetDimensionFromTexture(Texture));
 	}
@@ -412,12 +388,9 @@ void FSystemTextures::InitializeCommonTextures(FRHICommandListImmediate& RHICmdL
 	if (GPixelFormats[PF_FloatRGBA].Supported)
 	{
 		const FRHITextureCreateDesc Desc =
-			FRHITextureCreateDesc::Create2D(TEXT("MidGreyDummy"), 1, 1, PF_FloatRGBA)
-			.SetFlags(ETextureCreateFlags::ShaderResource)
-			.SetClassName(SystemTexturesName);
+			FRHITextureCreateDesc::Create2D(TEXT("MidGreyDummy"), 1, 1, PF_FloatRGBA);
 
-		FTextureRHIRef Texture = RHICreateTexture(Desc);
-		SetDummyTextureData<FFloat16Color>(Texture, FFloat16Color(FLinearColor(0.5f, 0.5f, 0.5f, 0.5f)));
+		FTextureRHIRef Texture = CreateDummyTexture<FFloat16Color>(RHICmdList, Desc, FFloat16Color(FLinearColor(0.5f, 0.5f, 0.5f, 0.5f)));
 		MidGreyDummy = CreateRenderTarget(Texture, Desc.DebugName);
 	}
 
@@ -449,17 +422,19 @@ void FSystemTextures::InitializeFeatureLevelDependentTextures(FRHICommandListImm
 		const FRHITextureCreateDesc Desc =
 			FRHITextureCreateDesc::Create2D(TEXT("SobolSampling"), 32, 16, PF_R16_UINT)
 			.SetFlags(ETextureCreateFlags::ShaderResource)
-			.SetClassName(SystemTexturesName);
+			.SetClassName(SystemTexturesName)
+			.SetInitActionInitializer();
+
+		FRHITextureInitializer Initializer = RHICmdList.CreateTextureInitializer(Desc);
 
-		FTextureRHIRef Texture = RHICreateTexture(Desc);
 		// Write the contents of the texture.
-		uint32 DestStride;
-		uint8* DestBuffer = (uint8*)RHICmdList.LockTexture2D(Texture, 0, RLM_WriteOnly, DestStride, false);
+		FRHITextureSubresourceInitializer Subresource = Initializer.GetTexture2DSubresource(0);
+
+		uint8* DestBuffer = reinterpret_cast<uint8*>(Subresource.Data);
 
-		uint16 *Dest;
 		for (int y = 0; y < 16; ++y)
 		{
-			Dest = (uint16*)(DestBuffer + y * DestStride);
+			uint16* Dest = (uint16*)(DestBuffer + y * Subresource.Stride);
 
 			// 16x16 block starting at 0,0 = Sobol X,Y from bottom 4 bits of cell X,Y
 			for (int x = 0; x < 16; ++x, ++Dest)
@@ -473,8 +448,8 @@ void FSystemTextures::InitializeFeatureLevelDependentTextures(FRHICommandListImm
 				*Dest = FSobol::ComputeGPUSpatialSeed(x, y, /* Index = */ 1);
 			}
 		}
-		RHICmdList.UnlockTexture2D(Texture, 0, false);
 
+		FTextureRHIRef Texture = Initializer.Finalize();
 		SobolSampling = CreateRenderTarget(Texture, Desc.DebugName);
 	}
 
@@ -577,13 +552,15 @@ void FSystemTextures::InitializeFeatureLevelDependentTextures(FRHICommandListImm
 		const FRHITextureCreateDesc Desc =
 			FRHITextureCreateDesc::Create2D(TEXT("PreintegratedGF"), Extent, Format)
 			.SetFlags(ETextureCreateFlags::ShaderResource)
-			.SetClassName(SystemTexturesName);
+			.SetClassName(SystemTexturesName)
+			.SetInitActionInitializer();
 
-		FTextureRHIRef Texture = RHICreateTexture(Desc);
+		FRHITextureInitializer Initializer = RHICmdList.CreateTextureInitializer(Desc);
 
 		// Write the contents of the texture.
-		uint32 DestStride;
-		uint8* DestBuffer = (uint8*)RHICmdList.LockTexture2D(Texture, 0, RLM_WriteOnly, DestStride, false);
+		FRHITextureSubresourceInitializer Subresource = Initializer.GetTexture2DSubresource(0);
+
+		uint8* DestBuffer = reinterpret_cast<uint8*>(Subresource.Data);
 
 		// x is NoV, y is roughness
 		for (int32 y = 0; y < Extent.Y; y++)
@@ -665,14 +642,14 @@ void FSystemTextures::InitializeFeatureLevelDependentTextures(FRHICommandListImm
 
 				if (Format == PF_A16B16G16R16)
 				{
-					uint16* Dest = (uint16*)(DestBuffer + x * 8 + y * DestStride);
+					uint16* Dest = (uint16*)(DestBuffer + x * 8 + y * Subresource.Stride);
 					Dest[0] = (int32)(FMath::Clamp(A, 0.0f, 1.0f) * 65535.0f + 0.5f);
 					Dest[1] = (int32)(FMath::Clamp(B, 0.0f, 1.0f) * 65535.0f + 0.5f);
 					Dest[2] = (int32)(FMath::Clamp(C, 0.0f, 1.0f) * 65535.0f + 0.5f);
 				}
 				else if (Format == PF_G16R16)
 				{
-					uint16* Dest = (uint16*)(DestBuffer + x * 4 + y * DestStride);
+					uint16* Dest = (uint16*)(DestBuffer + x * 4 + y * Subresource.Stride);
 					Dest[0] = (int32)(FMath::Clamp(A, 0.0f, 1.0f) * 65535.0f + 0.5f);
 					Dest[1] = (int32)(FMath::Clamp(B, 0.0f, 1.0f) * 65535.0f + 0.5f);
 				}
@@ -680,14 +657,14 @@ void FSystemTextures::InitializeFeatureLevelDependentTextures(FRHICommandListImm
 				{
 					check(Format == PF_R8G8);
 
-					uint8* Dest = (uint8*)(DestBuffer + x * 2 + y * DestStride);
+					uint8* Dest = (uint8*)(DestBuffer + x * 2 + y * Subresource.Stride);
 					Dest[0] = (int32)(FMath::Clamp(A, 0.0f, 1.0f) * 255.f + 0.5f);
 					Dest[1] = (int32)(FMath::Clamp(B, 0.0f, 1.0f) * 255.f + 0.5f);
 				}
 			}
 		}
-		RHICmdList.UnlockTexture2D(Texture, 0, false);
 
+		FTextureRHIRef Texture = Initializer.Finalize();
 		PreintegratedGF = CreateRenderTarget(Texture, Desc.DebugName);
 	}
 
@@ -847,19 +824,20 @@ void FSystemTextures::InitializeFeatureLevelDependentTextures(FRHICommandListImm
 				const FRHITextureCreateDesc Desc =
 					FRHITextureCreateDesc::Create2D(TEXT("SSAORandomization"), Extent, Extent, PF_R8G8)
 					.SetFlags(ETextureCreateFlags::ShaderResource)
-					.SetClassName(SystemTexturesName);
+					.SetClassName(SystemTexturesName)
+					.SetInitActionInitializer();
 
-				FTextureRHIRef Texture = RHICreateTexture(Desc);
+				FRHITextureInitializer Initializer = RHICmdList.CreateTextureInitializer(Desc);
 
 				// Write the contents of the texture.
-				uint32 DestStride;
-				uint8* DestBuffer = (uint8*)RHICmdList.LockTexture2D(Texture, 0, RLM_WriteOnly, DestStride, false);
+				FRHITextureSubresourceInitializer Subresource = Initializer.GetTexture2DSubresource(0);
+				uint8* DestBuffer = (uint8*)Subresource.Data;
 
 				for (int32 y = 0; y < Extent; ++y)
 				{
 					for (int32 x = 0; x < Extent; ++x)
 					{
-						uint8* Dest = (uint8*)(DestBuffer + x * sizeof(uint16) + y * DestStride);
+						uint8* Dest = (uint8*)(DestBuffer + x * sizeof(uint16) + y * Subresource.Stride);
 
 						uint32 Index = (x % 4) + (y % 4) * 4;
 
@@ -868,8 +846,7 @@ void FSystemTextures::InitializeFeatureLevelDependentTextures(FRHICommandListImm
 					}
 				}
 
-				RHICmdList.UnlockTexture2D(Texture, 0, false);
-
+				FTextureRHIRef Texture = Initializer.Finalize();
 				SSAORandomization = CreateRenderTarget(Texture, Desc.DebugName);
 			}
 		}
@@ -891,9 +868,10 @@ void FSystemTextures::InitializeFeatureLevelDependentTextures(FRHICommandListImm
 			.SetDepth(bGTAOPreIngegratedUsingVolumeLUT ? Extent : 1)
 			.SetFormat(PF_R16F)
 			.SetFlags(ETextureCreateFlags::ShaderResource)
-			.SetClassName(SystemTexturesName);
+			.SetClassName(SystemTexturesName)
+			.SetInitAction(bGTAOPreIngegratedUsingVolumeLUT ? ERHITextureInitAction::Default : ERHITextureInitAction::Initializer);
 
-		TRefCountPtr<FRHITexture> Texture = RHICreateTexture(Desc);
+		FRHITextureInitializer Initializer = RHICmdList.CreateTextureInitializer(Desc);
 
 		// Write the contents of the texture.
 		TArray<FFloat16> TempBuffer;
@@ -907,8 +885,8 @@ void FSystemTextures::InitializeFeatureLevelDependentTextures(FRHICommandListImm
 		}
 		else
 		{
-			uint32 DestStride;
-			DestBuffer = (FFloat16*)RHICmdList.LockTexture2D(Texture.GetReference(), 0, RLM_WriteOnly, DestStride, false);
+			FRHITextureSubresourceInitializer Subresource = Initializer.GetTexture2DSubresource(0);
+			DestBuffer = (FFloat16*)Subresource.Data;
 		}
 
 		for (uint32 z = 0; z < Extent; ++z)
@@ -952,6 +930,8 @@ void FSystemTextures::InitializeFeatureLevelDependentTextures(FRHICommandListImm
 			}
 		}
 
+		TRefCountPtr<FRHITexture> Texture = Initializer.Finalize();
+
 		if (bGTAOPreIngegratedUsingVolumeLUT)
 		{
 			FUpdateTextureRegion3D Region(0, 0, 0, 0, 0, 0, Extent, Extent, Extent);
@@ -964,30 +944,22 @@ void FSystemTextures::InitializeFeatureLevelDependentTextures(FRHICommandListImm
 				Extent * Extent * sizeof(FFloat16),
 				(const uint8*)DestBuffer);
 		}
-		else
-		{
-			RHICmdList.UnlockTexture2D(Texture.GetReference(), 0, false);
-		}
 
 		GTAOPreIntegrated = CreateRenderTarget(Texture, Desc.DebugName);
 	}
 
     // Create a texture array that is a single UInt32 value set to 0, this texture is AtomicCompatible on SM6
     {
-        ETextureCreateFlags TextureCreateFlags = ETextureCreateFlags::ShaderResource;
-        
-        if(InFeatureLevel >= ERHIFeatureLevel::SM6)
-        {
-            TextureCreateFlags |= ETextureCreateFlags::AtomicCompatible;
-        }
-        
-        const FRHITextureCreateDesc Desc =
-            FRHITextureCreateDesc::Create2DArray(TEXT("ZeroUIntArrayAtomicCompatDummy"), 1, 1, 1, PF_R32_UINT)
-            .SetFlags(TextureCreateFlags)
-            .SetClassName(SystemTexturesName);
-
-        FTextureRHIRef Texture = RHICreateTexture(Desc);
-        SetDummyTextureArrayData<uint32>(Texture, 0u);
+        FRHITextureCreateDesc Desc =
+			FRHITextureCreateDesc::Create2DArray(TEXT("ZeroUIntArrayAtomicCompatDummy"), 1, 1, 1, PF_R32_UINT)
+			.SetFlags(ETextureCreateFlags::ShaderResource);
+
+		if (InFeatureLevel >= ERHIFeatureLevel::SM6)
+		{
+			Desc.AddFlags(ETextureCreateFlags::AtomicCompatible);
+		}
+
+		FTextureRHIRef Texture = CreateDummyTexture<uint32>(RHICmdList, Desc, 0u);
         ZeroUIntArrayAtomicCompatDummy = CreateRenderTarget(Texture, Desc.DebugName);
     }
     
@@ -1456,40 +1428,29 @@ void InitializeData(const TInType& InData, EPixelFormat InFormat, uint8* OutData
 }
 
 template <typename DataType>
-void SetDefaultTextureData2D(FRHITexture* Texture, const DataType& InData)
+void SetDefaultTextureData2D(FRHITextureInitializer& Initializer, const FRHITextureCreateDesc& InDesc, const DataType& InData)
 {
-	uint8 SrcData[16] = { 0,0,0,0, 0,0,0,0, 0,0,0,0, 0,0,0,0, };
-	uint32 SrcByteCount = 0;
-	const EPixelFormat Format = Texture->GetFormat();
-	InitializeData(ToVector(InData), Format, SrcData, SrcByteCount);
+	FRHITextureSubresourceInitializer Subresource = Initializer.GetTexture2DSubresource(0);
 
-	uint32 DestStride;
-	uint8* Dest = (uint8*)RHILockTexture2D(Texture, 0, RLM_WriteOnly, DestStride, false);
-	FMemory::Memcpy(Dest, SrcData, SrcByteCount);
-	RHIUnlockTexture2D(Texture, 0, false);
+	uint32 SrcByteCount = 0;
+	InitializeData(ToVector(InData), InDesc.Format, (uint8*)Subresource.Data, SrcByteCount);
 }
 
 template <typename DataType>
-void SetDefaultTextureData2DArray(FRHITexture* Texture, const DataType& InData)
+void SetDefaultTextureData2DArray(FRHITextureInitializer& Initializer, const FRHITextureCreateDesc& InDesc, const DataType& InData)
 {
-	uint8 SrcData[16] = { 0,0,0,0, 0,0,0,0, 0,0,0,0, 0,0,0,0, };
-	uint32 SrcByteCount = 0;
-	const EPixelFormat Format = Texture->GetFormat();
-	InitializeData(ToVector(InData), Format, SrcData, SrcByteCount);
+	FRHITextureSubresourceInitializer Subresource = Initializer.GetTexture2DArraySubresource(0, 0);
 
-	uint32 DestStride;
-	uint8* Dest = (uint8*)RHILockTexture2DArray(Texture, 0, 0, RLM_WriteOnly, DestStride, false);
-	FMemory::Memcpy(Dest, SrcData, SrcByteCount);
-	RHIUnlockTexture2DArray(Texture, 0, 0, false);
+	uint32 SrcByteCount = 0;
+	InitializeData(ToVector(InData), InDesc.Format, (uint8*)Subresource.Data, SrcByteCount);
 }
 
 template <typename DataType>
-void SetDefaultTextureData3D(FRHICommandListImmediate& RHICmdList, FRHITexture* Texture, const DataType& InData)
+void SetDefaultTextureData3D(FRHICommandListImmediate& RHICmdList, FRHITexture* Texture, const FRHITextureCreateDesc& InDesc, const DataType& InData)
 {
 	uint8 SrcData[16] = { 0,0,0,0, 0,0,0,0, 0,0,0,0, 0,0,0,0, };
 	uint32 SrcByteCount = 0;
-	const EPixelFormat Format = Texture->GetFormat();
-	InitializeData(ToVector(InData), Format, SrcData, SrcByteCount);
+	InitializeData(ToVector(InData), InDesc.Format, SrcData, SrcByteCount);
 
 	FUpdateTextureRegion3D Region(0, 0, 0, 0, 0, 0, 1, 1, 1);
 	RHICmdList.UpdateTexture3D(
@@ -1505,19 +1466,16 @@ void SetDefaultTextureData3D(FRHICommandListImmediate& RHICmdList, FRHITexture*
 }
 
 template <typename DataType>
-void SetDefaultTextureDataCube(FRHITexture* Texture, const DataType& InData)
+void SetDefaultTextureDataCube(FRHITextureInitializer& Initializer, const FRHITextureCreateDesc& InDesc, const DataType& InData)
 {
 	uint8 SrcData[16] = { 0,0,0,0, 0,0,0,0, 0,0,0,0, 0,0,0,0, };
 	uint32 SrcByteCount = 0;
-	const EPixelFormat Format = Texture->GetFormat();
-	InitializeData(ToVector(InData), Format, SrcData, SrcByteCount);
+	InitializeData(ToVector(InData), InDesc.Format, SrcData, SrcByteCount);
 
 	for (uint32 FaceIt = 0; FaceIt < 6; ++FaceIt)
 	{
-		uint32 DestStride;
-		uint8* Dest = (uint8*)RHILockTextureCubeFace(Texture, FaceIt, 0u, 0u, RLM_WriteOnly, DestStride, false);
-		FMemory::Memcpy(Dest, SrcData, SrcByteCount);
-		RHIUnlockTextureCubeFace(Texture, FaceIt, 0, 0, false);
+		FRHITextureSubresourceInitializer Subresource = Initializer.GetTextureCubeSubresource(FaceIt, 0);
+		Subresource.WriteData(SrcData, SrcByteCount);
 	}
 }
 
@@ -1561,20 +1519,26 @@ FRDGTextureRef GetInternalDefaultTexture(
 	{
 		const FRHITextureCreateDesc Desc =
 			FRHITextureCreateDesc::Create2D(TEXT("DefaultTexture2D"), 1, 1, Format)
-			.SetFlags(ETextureCreateFlags::ShaderResource);
+			.SetFlags(ETextureCreateFlags::ShaderResource)
+			.SetInitActionInitializer();
 
-		FTextureRHIRef Texture = RHICreateTexture(Desc);
-		SetDefaultTextureData2D(Texture, Value);
+		FRHITextureInitializer Initializer = GraphBuilder.RHICmdList.CreateTextureInitializer(Desc);
+		SetDefaultTextureData2D(Initializer, Desc, Value);
+
+		FTextureRHIRef Texture = Initializer.Finalize();
 		Entry.Texture = CreateRenderTarget(Texture, Desc.DebugName);
 	}
 	else if (Dimension == ETextureDimension::Texture2DArray)
 	{
 		const FRHITextureCreateDesc Desc =
 			FRHITextureCreateDesc::Create2DArray(TEXT("DefaultTexture2DArray"), 1, 1, 1, Format)
-			.SetFlags(ETextureCreateFlags::ShaderResource);
+			.SetFlags(ETextureCreateFlags::ShaderResource)
+			.SetInitActionInitializer();
 
-		FTextureRHIRef Texture = RHICreateTexture(Desc);
-		SetDefaultTextureData2DArray(Texture, Value);
+		FRHITextureInitializer Initializer = GraphBuilder.RHICmdList.CreateTextureInitializer(Desc);
+		SetDefaultTextureData2DArray(Initializer, Desc, Value);
+
+		FTextureRHIRef Texture = Initializer.Finalize();
 		Entry.Texture = CreateRenderTarget(Texture, Desc.DebugName);
 	}
 	else if (Dimension == ETextureDimension::Texture3D)
@@ -1584,27 +1548,33 @@ FRDGTextureRef GetInternalDefaultTexture(
 			.SetFlags(ETextureCreateFlags::ShaderResource);
 
 		FTextureRHIRef Texture = RHICreateTexture(Desc);
-		SetDefaultTextureData3D(GraphBuilder.RHICmdList, Texture, Value);
+		SetDefaultTextureData3D(GraphBuilder.RHICmdList, Texture, Desc, Value);
 		Entry.Texture = CreateRenderTarget(Texture, Desc.DebugName);
 	}
 	else if (Dimension == ETextureDimension::TextureCube)
 	{
 		const FRHITextureCreateDesc Desc =
 			FRHITextureCreateDesc::CreateCube(TEXT("DefaultTextureCube"), 1, Format)
-			.SetFlags(ETextureCreateFlags::ShaderResource);
+			.SetFlags(ETextureCreateFlags::ShaderResource)
+			.SetInitActionInitializer();
 
-		FTextureRHIRef Texture = RHICreateTexture(Desc);
-		SetDefaultTextureDataCube(Texture, Value);
+		FRHITextureInitializer Initializer = GraphBuilder.RHICmdList.CreateTextureInitializer(Desc);
+		SetDefaultTextureDataCube(Initializer, Desc, Value);
+
+		FTextureRHIRef Texture = Initializer.Finalize();
 		Entry.Texture = CreateRenderTarget(Texture, Desc.DebugName);
 	}
 	else if (Dimension == ETextureDimension::TextureCubeArray)
 	{
 		const FRHITextureCreateDesc Desc =
 			FRHITextureCreateDesc::CreateCubeArray(TEXT("DefaultTextureCubeArray"), 1, 1, Format)
-			.SetFlags(ETextureCreateFlags::ShaderResource);
+			.SetFlags(ETextureCreateFlags::ShaderResource)
+			.SetInitActionInitializer();
 
-		FTextureRHIRef Texture = RHICreateTexture(Desc);
-		SetDefaultTextureDataCube(Texture, Value);
+		FRHITextureInitializer Initializer = GraphBuilder.RHICmdList.CreateTextureInitializer(Desc);
+		SetDefaultTextureDataCube(Initializer, Desc, Value);
+
+		FTextureRHIRef Texture = Initializer.Finalize();
 		Entry.Texture = CreateRenderTarget(Texture, Desc.DebugName);
 	}
 	else
diff --git a/Engine/Source/Runtime/VulkanRHI/Private/VulkanDynamicRHI.h b/Engine/Source/Runtime/VulkanRHI/Private/VulkanDynamicRHI.h
--- a/Engine/Source/Runtime/VulkanRHI/Private/VulkanDynamicRHI.h
+++ b/Engine/Source/Runtime/VulkanRHI/Private/VulkanDynamicRHI.h
@@ -179,12 +179,13 @@ public:
 	};
 
 	FCreateTextureResult BeginCreateTextureInternal(const FRHITextureCreateDesc& CreateDesc, const FRHITransientHeapAllocation* InTransientHeapAllocation);
-	FVulkanTexture* FinalizeCreateTextureInternal(FRHICommandListBase& RHICmdList, FCreateTextureResult CreateResult, TConstArrayView<uint8> InitialData);
+	FVulkanTexture* FinalizeCreateTextureInternal(FRHICommandListBase& RHICmdList, FCreateTextureResult CreateResult);
 
-	FVulkanTexture* CreateTextureInternal(FRHICommandListBase& RHICmdList, const FRHITextureCreateDesc& CreateDesc, TConstArrayView<uint8> InitialData);
+	FVulkanTexture* CreateTextureInternal(FRHICommandListBase& RHICmdList, const FRHITextureCreateDesc& CreateDesc);
 	FVulkanTexture* CreateTextureInternal(const FRHITextureCreateDesc& CreateDesc, const FRHITransientHeapAllocation& InTransientHeapAllocation);
 
-	virtual FTextureRHIRef RHICreateTexture(FRHICommandListBase& RHICmdList, const FRHITextureCreateDesc& CreateDesc) final override;
+	[[nodiscard]] virtual FRHITextureInitializer RHICreateTextureInitializer(FRHICommandListBase& RHICmdList, const FRHITextureCreateDesc& CreateDesc);
+
 	virtual FTextureRHIRef RHIAsyncCreateTexture2D(uint32 SizeX, uint32 SizeY, uint8 Format, uint32 NumMips, ETextureCreateFlags Flags, ERHIAccess InResourceState, void** InitialMipData, uint32 NumInitialMips, const TCHAR* DebugName, FGraphEventRef& OutCompletionEvent) final override;
 	virtual uint32 RHIComputeMemorySize(FRHITexture* TextureRHI) final override;
 	virtual FTextureRHIRef RHIAsyncReallocateTexture2D(FRHITexture* Texture2D, int32 NewMipCount, int32 NewSizeX, int32 NewSizeY, FThreadSafeCounter* RequestStatus) final override;
diff --git a/Engine/Source/Runtime/VulkanRHI/Private/VulkanResources.h b/Engine/Source/Runtime/VulkanRHI/Private/VulkanResources.h
--- a/Engine/Source/Runtime/VulkanRHI/Private/VulkanResources.h
+++ b/Engine/Source/Runtime/VulkanRHI/Private/VulkanResources.h
@@ -603,6 +603,8 @@ public:
 
 	void AliasTextureResources(FTextureRHIRef& SrcTextureRHI);
 
+	void UploadInitialData(FRHICommandListBase& RHICmdList, VulkanRHI::FStagingBuffer* UploadBuffer);
+
 	// View with all mips/layers
 	FVulkanView* DefaultView = nullptr;
 
diff --git a/Engine/Source/Runtime/VulkanRHI/Private/VulkanTexture.cpp b/Engine/Source/Runtime/VulkanRHI/Private/VulkanTexture.cpp
--- a/Engine/Source/Runtime/VulkanRHI/Private/VulkanTexture.cpp
+++ b/Engine/Source/Runtime/VulkanRHI/Private/VulkanTexture.cpp
@@ -17,6 +17,7 @@
 #include "HAL/LowLevelMemStats.h"
 #include "ProfilingDebugging/AssetMetadataTrace.h"
 #include "RHICoreTexture.h"
+#include "RHICoreTextureInitializer.h"
 
 // This is a workaround for issues with AFBC on Mali GPUs before the G710
 int32 GVulkanDepthStencilForceStorageBit = 0;
@@ -802,18 +803,6 @@ void FVulkanTexture::GetMipSize(uint32 MipIndex, uint64& MipBytes)
 
 	// Size in bytes
 	MipBytes = NumBlocksX * NumBlocksY * BlockBytes * Desc.Depth;
-/*
-#if VULKAN_HAS_DEBUGGING_ENABLED
-	VkImageSubresource SubResource;
-	FMemory::Memzero(SubResource);
-	SubResource.aspectMask = FullAspectMask;
-	SubResource.mipLevel = MipIndex;
-	//SubResource.arrayLayer = 0;
-	VkSubresourceLayout OutLayout;
-	VulkanRHI::vkGetImageSubresourceLayout(Device->GetInstanceHandle(), Image, &SubResource, &OutLayout);
-	ensure(MipBytes >= OutLayout.size);
-#endif
-*/
 }
 
 void FVulkanTexture::SetInitialImageState(FRHICommandListBase& RHICmdList, VkImageLayout InitialLayout, bool bClear, const FClearValueBinding& ClearValueBinding, bool bIsTransientResource)
@@ -1031,7 +1020,7 @@ FVulkanDynamicRHI::FCreateTextureResult FVulkanDynamicRHI::BeginCreateTextureInt
 	return Result;
 }
 
-FVulkanTexture* FVulkanDynamicRHI::FinalizeCreateTextureInternal(FRHICommandListBase& RHICmdList, FCreateTextureResult CreateResult, TConstArrayView<uint8> InitialData)
+FVulkanTexture* FVulkanDynamicRHI::FinalizeCreateTextureInternal(FRHICommandListBase& RHICmdList, FCreateTextureResult CreateResult)
 {
 	FVulkanTexture* Texture = CreateResult.Texture;
 	const FRHITextureDesc& Desc = Texture->GetDesc();
@@ -1042,51 +1031,136 @@ FVulkanTexture* FVulkanDynamicRHI::FinalizeCreateTextureInternal(FRHICommandList
 		{
 			Texture->SetInitialImageState(RHICmdList, CreateResult.DefaultLayout, CreateResult.bClear, Desc.ClearValue, CreateResult.bTransientResource);
 		}
+	}
+
+	return Texture;
+}
+
+size_t VulkanCalculateTextureSize(FVulkanTexture* Texture)
+{
+	const FRHITextureDesc& Desc = Texture->GetDesc();
+	const FPixelFormatInfo& FormatInfo = GPixelFormats[Desc.Format];
+
+	const uint32 WidthInBlocks = FMath::DivideAndRoundUp<uint32>(Desc.Extent.X, FormatInfo.BlockSizeX);
+	const uint32 HeightInBlocks = FMath::DivideAndRoundUp<uint32>(Desc.Extent.Y, FormatInfo.BlockSizeY);
+
+	const VkPhysicalDeviceLimits& Limits = Texture->Device->GetLimits();
+
+	const size_t StagingPitch = static_cast<size_t>(WidthInBlocks) * FormatInfo.BlockBytes;
+	const size_t StagingBufferSize = Align(StagingPitch * HeightInBlocks, Limits.minMemoryMapAlignment);
+	return StagingBufferSize;
+}
+
+namespace UE::VulkanTexture
+{
+	// Vulkan stores subresources by layer first. 
+	// ie: Layer0.Mip0->Layer1.Mip0->Layer0.Mip1->Layer1->Mip1
+	static uint64 CalculateSubresourceOffset(const FRHITextureDesc& Desc, FRHITextureInitializer::FSubresourceIndex SubresourceIndex, uint64& OutStride, uint64& OutSize)
+	{
+		const uint64 FaceCount = Desc.IsTextureCube() ? 6 : 1;
+		const uint64 LayerCount = FaceCount * Desc.ArraySize;
+		const uint64 LayerIndex = SubresourceIndex.FaceIndex + SubresourceIndex.ArrayIndex * FaceCount;
 
-		if (InitialData.Num())
+		uint64 LayerOffset = 0;
+
+		uint64 MipOffset = 0;
+		uint64 MipStride = 0;
+
+		for (int32 Index = 0; Index <= SubresourceIndex.MipIndex; Index++)
 		{
-			// InternalLockWrite leaves the image in VK_IMAGE_LAYOUT_SHADER_READ_ONLY_OPTIMAL, so make sure the requested resource state is SRV.
-			//check(EnumHasAnyFlags(InCreateDesc.InitialState, ERHIAccess::SRVMask));
-			Texture->SetDefaultLayout(VK_IMAGE_LAYOUT_SHADER_READ_ONLY_OPTIMAL);
-
-			// Transfer bulk data
-			VulkanRHI::FStagingBuffer* StagingBuffer = Device->GetStagingManager().AcquireBuffer(InitialData.NumBytes());
-			void* Data = StagingBuffer->GetMappedPointer();
-
-			// Do copy
-			FMemory::Memcpy(Data, InitialData.GetData(), InitialData.NumBytes());
-
-			VkBufferImageCopy Region;
-			FMemory::Memzero(Region);
-			//#todo-rco: Use real Buffer offset when switching to suballocations!
-			Region.bufferOffset = 0;
-			Region.bufferRowLength = Desc.Extent.X;
-			Region.bufferImageHeight = Desc.Extent.Y;
-
-			Region.imageSubresource.mipLevel = 0;
-			Region.imageSubresource.baseArrayLayer = 0;
-			Region.imageSubresource.layerCount = Texture->GetNumberOfArrayLevels();
-			Region.imageSubresource.aspectMask = Texture->GetFullAspectMask();
-
-			Region.imageExtent.width = Region.bufferRowLength;
-			Region.imageExtent.height = Region.bufferImageHeight;
-			Region.imageExtent.depth = Desc.Depth;
-
-			RHICmdList.EnqueueLambda(TEXT("FVulkanTexture::InternalLockWrite"),
-				[Texture, Region, StagingBuffer](FRHICommandListBase& ExecutingCmdList)
-				{
-					FVulkanTexture::InternalLockWrite(FVulkanUploadContext::Get(ExecutingCmdList), Texture, Region, StagingBuffer);
-				});
+			const uint64 MipSize = UE::RHITextureUtils::CalculateTextureMipSize(Desc, Index, MipStride);
+
+			if (Index == SubresourceIndex.MipIndex)
+			{
+				LayerOffset += MipSize * LayerIndex;
+
+				OutSize = MipSize;
+				OutStride = MipStride;
+			}
+			else
+			{
+				LayerOffset += MipSize * LayerCount;
+			}
 		}
+
+		return LayerOffset;
 	}
+}
 
-	return Texture;
+void FVulkanTexture::UploadInitialData(FRHICommandListBase& RHICmdList, VulkanRHI::FStagingBuffer* UploadBuffer)
+{
+	const FRHITextureDesc& Desc = GetDesc();
+
+	// InternalLockWrite leaves the image in VK_IMAGE_LAYOUT_SHADER_READ_ONLY_OPTIMAL, so make sure the requested resource state is SRV.
+	SetDefaultLayout(VK_IMAGE_LAYOUT_SHADER_READ_ONLY_OPTIMAL);
+
+	RHICmdList.EnqueueLambda(TEXT("FVulkanTexture::UploadInitialData"),
+		[this, UploadBuffer](FRHICommandListBase& ExecutingCmdList) mutable
+		{
+			const FRHITextureDesc& Desc = GetDesc();
+
+			const uint32 LayerCount = GetNumberOfArrayLevels();
+
+			TArray<VkBufferImageCopy, TInlineAllocator<12>> Regions;
+			Regions.AddZeroed(Desc.NumMips);
+
+			uint64 MipOffset = 0;
+
+			for (uint32 MipIndex = 0; MipIndex < Desc.NumMips; MipIndex++)
+			{
+				VkBufferImageCopy& Region = Regions[MipIndex];
+
+				check(MipOffset < UploadBuffer->GetSize());
+
+				const FUintVector3 MipExtents = UE::RHITextureUtils::CalculateMipExtents(Desc, MipIndex);
+
+				Region.bufferOffset = MipOffset;
+
+				Region.imageExtent.width = MipExtents.X;
+				Region.imageExtent.height = MipExtents.Y;
+				Region.imageExtent.depth = MipExtents.Z;
+
+				Region.imageSubresource.aspectMask = FullAspectMask;
+				Region.imageSubresource.layerCount = LayerCount;
+				Region.imageSubresource.mipLevel = MipIndex;
+
+				uint64 Stride = 0;
+				const uint64 MipSize = UE::RHITextureUtils::CalculateTextureMipSize(Desc, MipIndex, Stride);
+				const uint64 SubresourceSize = MipSize * LayerCount;
+
+				MipOffset += SubresourceSize;
+			}
+
+			FVulkanUploadContext& Context = FVulkanUploadContext::Get(ExecutingCmdList);
+
+			FVulkanCommandBuffer* CmdBuffer = Context.GetActiveCmdBuffer();
+			ensure(CmdBuffer->IsOutsideRenderPass());
+			VkCommandBuffer StagingCommandBuffer = CmdBuffer->GetHandle();
+
+			const VkImageSubresourceRange SubresourceRange = FVulkanPipelineBarrier::MakeSubresourceRange(FullAspectMask, 0, Desc.NumMips, 0, LayerCount);
+
+			{
+				FVulkanPipelineBarrier Barrier;
+				Barrier.AddImageLayoutTransition(Image, VK_IMAGE_LAYOUT_UNDEFINED, VK_IMAGE_LAYOUT_TRANSFER_DST_OPTIMAL, SubresourceRange);
+				Barrier.Execute(CmdBuffer);
+			}
+
+			VulkanRHI::vkCmdCopyBufferToImage(StagingCommandBuffer, UploadBuffer->GetHandle(), Image, VK_IMAGE_LAYOUT_TRANSFER_DST_OPTIMAL, Regions.Num(), Regions.GetData());
+
+			{
+				FVulkanPipelineBarrier Barrier;
+				Barrier.AddImageLayoutTransition(Image, VK_IMAGE_LAYOUT_TRANSFER_DST_OPTIMAL, VK_IMAGE_LAYOUT_SHADER_READ_ONLY_OPTIMAL, SubresourceRange);
+				Barrier.Execute(CmdBuffer);
+			}
+
+			Device->GetStagingManager().ReleaseBuffer(&Context, UploadBuffer);
+		});
 }
 
-FVulkanTexture* FVulkanDynamicRHI::CreateTextureInternal(FRHICommandListBase& RHICmdList, const FRHITextureCreateDesc& CreateDesc, TConstArrayView<uint8> InitialData)
+FVulkanTexture* FVulkanDynamicRHI::CreateTextureInternal(FRHICommandListBase& RHICmdList, const FRHITextureCreateDesc& CreateDesc)
 {
 	FCreateTextureResult CreateResult = BeginCreateTextureInternal(CreateDesc, nullptr);
-	return FinalizeCreateTextureInternal(RHICmdList, CreateResult, InitialData);
+	return FinalizeCreateTextureInternal(RHICmdList, CreateResult);
 }
 
 FVulkanTexture* FVulkanDynamicRHI::CreateTextureInternal(const FRHITextureCreateDesc& CreateDesc, const FRHITransientHeapAllocation& InTransientHeapAllocation)
@@ -1095,23 +1169,74 @@ FVulkanTexture* FVulkanDynamicRHI::CreateTextureInternal(const FRHITextureCreate
 	return CreateResult.Texture;
 }
 
-FTextureRHIRef FVulkanDynamicRHI::RHICreateTexture(FRHICommandListBase& RHICmdList, const FRHITextureCreateDesc& CreateDesc)
+FRHITextureInitializer FVulkanDynamicRHI::RHICreateTextureInitializer(FRHICommandListBase& RHICmdList, const FRHITextureCreateDesc& CreateDesc)
 {
-	LLM_SCOPE_VULKAN(GetMemoryTagForTextureFlags(CreateDesc.Flags));
-	TConstArrayView<uint8> InitialData;
-	if (CreateDesc.BulkData)
+	LLM_SCOPE_DYNAMIC_STAT_OBJECTPATH_FNAME(CreateDesc.OwnerName, ELLMTagSet::Assets);
+	LLM_SCOPE_DYNAMIC_STAT_OBJECTPATH_FNAME(CreateDesc.GetTraceClassName(), ELLMTagSet::AssetClasses);
+	UE_TRACE_METADATA_SCOPE_ASSET_FNAME(CreateDesc.DebugName, CreateDesc.GetTraceClassName(), CreateDesc.OwnerName);
+
+	if (CreateDesc.InitAction == ERHITextureInitAction::Default)
 	{
-		InitialData = CreateDesc.BulkData->GetBulkDataView<uint8>();
+		return UE::RHICore::FDefaultTextureInitializer(RHICmdList, CreateTextureInternal(RHICmdList, CreateDesc));
 	}
 
-	FVulkanTexture* Texture = CreateTextureInternal(RHICmdList, CreateDesc, InitialData);
+	check(EnumHasAnyFlags(CreateDesc.InitialState, ERHIAccess::SRVMask));
+
+	// TODO: add something to force VK_IMAGE_LAYOUT_SHADER_READ_ONLY_OPTIMAL as the default layout
+
+	FCreateTextureResult CreateResult = BeginCreateTextureInternal(CreateDesc, nullptr);
+	FVulkanTexture* Texture = FinalizeCreateTextureInternal(RHICmdList, CreateResult);
+
+	const uint64 UploadMemorySize = UE::RHITextureUtils::CalculateTextureSize(CreateDesc);
+	const uint64 VulkanUploadSize = VulkanCalculateTextureSize(Texture);
+
+	// Transfer bulk data
+	VulkanRHI::FStagingBuffer* UploadBuffer = Device->GetStagingManager().AcquireBuffer(UploadMemorySize);
+	void* const UploadData = UploadBuffer->GetMappedPointer();
+
+	UE::RHICore::FBaseTextureInitializerImplementation Initializer(RHICmdList, Texture, UploadData, UploadMemorySize,
+		[Texture = TRefCountPtr<FVulkanTexture>(Texture), UploadBuffer](FRHICommandListBase& RHICmdList) mutable
+		{
+			Texture->UploadInitialData(RHICmdList, UploadBuffer);
+			return TRefCountPtr<FRHITexture>(MoveTemp(Texture));
+		},
+		[Texture = TRefCountPtr<FVulkanTexture>(Texture), UploadBuffer](FRHITextureInitializer::FSubresourceIndex SubresourceIndex) mutable
+		{
+			const FRHITextureDesc& TextureDesc = Texture->GetDesc();
+
+			uint64 SubresourceStride = 0;
+			uint64 SubresourceSize = 0;
+			const uint64 Offset = UE::VulkanTexture::CalculateSubresourceOffset(TextureDesc, SubresourceIndex, SubresourceStride, SubresourceSize);
+			check(Offset + SubresourceSize <= UploadBuffer->GetSize());
 
-	if (CreateDesc.BulkData)
+			FRHITextureSubresourceInitializer Result{};
+			Result.Data = reinterpret_cast<uint8*>(UploadBuffer->GetMappedPointer()) + Offset;
+			Result.Stride = SubresourceStride;
+			Result.Size = SubresourceSize;
+
+			return Result;
+		}
+	);
+
+	if (CreateDesc.InitAction == ERHITextureInitAction::BulkData)
 	{
+		check(CreateDesc.BulkData);
+		check(UploadMemorySize >= CreateDesc.BulkData->GetResourceBulkDataSize());
+
+		FMemory::Memcpy(UploadData, CreateDesc.BulkData->GetResourceBulkData(), CreateDesc.BulkData->GetResourceBulkDataSize());
+
+		// Discard the bulk data's contents.
 		CreateDesc.BulkData->Discard();
+
+		return MoveTemp(Initializer);
 	}
 
-	return Texture;
+	if (CreateDesc.InitAction == ERHITextureInitAction::Initializer)
+	{
+		return MoveTemp(Initializer);
+	}
+
+	return UE::RHICore::HandleUnknownTextureInitializerInitAction(RHICmdList, CreateDesc);
 }
 
 FTextureRHIRef FVulkanDynamicRHI::RHIAsyncCreateTexture2D(uint32 SizeX,uint32 SizeY,uint8 Format,uint32 NumMips,ETextureCreateFlags Flags, ERHIAccess InResourceState,void** InitialMipData,uint32 NumInitialMips, const TCHAR* DebugName, FGraphEventRef& OutCompletionEvent)
@@ -1204,7 +1329,7 @@ FTextureRHIRef FVulkanDynamicRHI::AsyncReallocateTexture2D_RenderThread(FRHIComm
 		.DetermineInititialState()
 		.SetOwnerName(OldTexture->GetOwnerName());
 
-	FVulkanTexture* NewTexture = CreateTextureInternal(RHICmdList, CreateDesc, {});
+	FVulkanTexture* NewTexture = CreateTextureInternal(RHICmdList, CreateDesc);
 
 	RHICmdList.EnqueueLambda(TEXT("AsyncReallocateTexture2D"), [OldTexture, NewTexture, NewMipCount, NewSizeX, NewSizeY, RequestStatus](FRHICommandListImmediate& ImmCmdList)
 	{
diff --git a/Engine/Source/Runtime/Windows/D3D11RHI/Private/D3D11Buffer.cpp b/Engine/Source/Runtime/Windows/D3D11RHI/Private/D3D11Buffer.cpp
--- a/Engine/Source/Runtime/Windows/D3D11RHI/Private/D3D11Buffer.cpp
+++ b/Engine/Source/Runtime/Windows/D3D11RHI/Private/D3D11Buffer.cpp
@@ -201,7 +201,7 @@ FRHIBufferInitializer FD3D11DynamicRHI::RHICreateBufferInitializer(FRHICommandLi
 		void* UploadMemory = FMemory::Malloc(CreateDesc.Size, 16);
 
 		return UE::RHICore::FCustomBufferInitializer(RHICmdList, Buffer, UploadMemory, CreateDesc.Size,
-			[this, Buffer = TRefCountPtr<FD3D11Buffer>(Buffer), UploadMemory = UE::RHICore::FScopedMemory(UploadMemory)](FRHICommandListBase& RHICmdList) mutable
+			[this, Buffer = TRefCountPtr<FD3D11Buffer>(Buffer), UploadMemory = UE::RHICore::FInitializerScopedMemory(UploadMemory)](FRHICommandListBase& RHICmdList) mutable
 			{
 				FinalizeCreateBufferInternal(Buffer, TConstArrayView<uint8>(reinterpret_cast<const uint8*>(UploadMemory.Pointer), Buffer->GetDesc().Size));
 				return TRefCountPtr<FRHIBuffer>(MoveTemp(Buffer));
diff --git a/Engine/Source/Runtime/Windows/D3D11RHI/Private/D3D11RHIPrivate.h b/Engine/Source/Runtime/Windows/D3D11RHI/Private/D3D11RHIPrivate.h
--- a/Engine/Source/Runtime/Windows/D3D11RHI/Private/D3D11RHIPrivate.h
+++ b/Engine/Source/Runtime/Windows/D3D11RHI/Private/D3D11RHIPrivate.h
@@ -471,8 +471,7 @@ public:
 	void  FinalizeCreateTexture3DInternal(FD3D11Texture* Texture, TConstArrayView<D3D11_SUBRESOURCE_DATA> InitialData);
 	FD3D11Texture* FinalizeCreateTextureInternal(FD3D11Texture* Texture, const FRHITextureDesc& InDesc, TConstArrayView<uint8> InitialData);
 	[[nodiscard]] FD3D11Texture* CreateTextureInternal(const FRHITextureCreateDesc& CreateDesc, TConstArrayView<uint8> InitialData);
-
-	virtual FTextureRHIRef RHICreateTexture(FRHICommandListBase& RHICmdList, const FRHITextureCreateDesc& CreateDesc) final override;
+	[[nodiscard]] virtual FRHITextureInitializer RHICreateTextureInitializer(FRHICommandListBase& RHICmdList, const FRHITextureCreateDesc& CreateDesc) final override;
 	virtual FTextureRHIRef RHIAsyncCreateTexture2D(uint32 SizeX, uint32 SizeY, uint8 Format, uint32 NumMips, ETextureCreateFlags Flags, ERHIAccess InResourceState, void** InitialMipData, uint32 NumInitialMips, const TCHAR* DebugName, FGraphEventRef& OutCompletionEvent) final override;
 	virtual uint32 RHIComputeMemorySize(FRHITexture* TextureRHI) final override;
 	virtual void RHIAsyncCopyTexture2DCopy(FRHITexture* NewTexture2DRHI, FRHITexture* Texture2DRHI, int32 NewMipCount, int32 NewSizeX, int32 NewSizeY, FThreadSafeCounter* RequestStatus);
diff --git a/Engine/Source/Runtime/Windows/D3D11RHI/Private/D3D11Texture.cpp b/Engine/Source/Runtime/Windows/D3D11RHI/Private/D3D11Texture.cpp
--- a/Engine/Source/Runtime/Windows/D3D11RHI/Private/D3D11Texture.cpp
+++ b/Engine/Source/Runtime/Windows/D3D11RHI/Private/D3D11Texture.cpp
@@ -23,6 +23,7 @@
 #include "RHICoreStats.h"
 #include "RHICoreTexture.h"
 #include "RHIUtilities.h"
+#include "RHICoreTextureInitializer.h"
 #include "RHITextureUtils.h"
 
 int64 FD3D11GlobalStats::GDedicatedVideoMemory = 0;
@@ -817,21 +818,46 @@ FD3D11Texture* FD3D11DynamicRHI::CreateTextureInternal(const FRHITextureCreateDe
 	return FinalizeCreateTextureInternal(Texture, CreateDesc, InitialData);
 }
 
-FTextureRHIRef FD3D11DynamicRHI::RHICreateTexture(FRHICommandListBase&, const FRHITextureCreateDesc& CreateDesc)
+FRHITextureInitializer FD3D11DynamicRHI::RHICreateTextureInitializer(FRHICommandListBase& RHICmdList, const FRHITextureCreateDesc& CreateDesc)
 {
-	TConstArrayView<uint8> InitialData;
-	if (CreateDesc.BulkData)
+	LLM_SCOPE_DYNAMIC_STAT_OBJECTPATH_FNAME(CreateDesc.OwnerName, ELLMTagSet::Assets);
+	LLM_SCOPE_DYNAMIC_STAT_OBJECTPATH_FNAME(CreateDesc.GetTraceClassName(), ELLMTagSet::AssetClasses);
+	UE_TRACE_METADATA_SCOPE_ASSET_FNAME(CreateDesc.DebugName, CreateDesc.GetTraceClassName(), CreateDesc.OwnerName);
+
+	if (CreateDesc.InitAction == ERHITextureInitAction::Default)
 	{
-		InitialData = CreateDesc.BulkData->GetBulkDataView<uint8>();
+		return UE::RHICore::FDefaultTextureInitializer(RHICmdList, CreateTextureInternal(CreateDesc, {}));
 	}
 
-	FD3D11Texture* Texture = CreateTextureInternal(CreateDesc, InitialData);
-	if (CreateDesc.BulkData)
+	if (CreateDesc.InitAction == ERHITextureInitAction::BulkData)
 	{
+		check(CreateDesc.BulkData);
+
+		FD3D11Texture* Texture = CreateTextureInternal(CreateDesc, CreateDesc.BulkData->GetBulkDataView<uint8>());
+
+		// Discard the bulk data's contents.
 		CreateDesc.BulkData->Discard();
+
+		return UE::RHICore::FDefaultTextureInitializer(RHICmdList, Texture);
 	}
 
-	return Texture;
+	if (CreateDesc.InitAction == ERHITextureInitAction::Initializer)
+	{
+		FD3D11Texture* Texture = BeginCreateTextureInternal(CreateDesc);
+
+		const uint64 UploadMemorySize = UE::RHITextureUtils::CalculateTextureSize(CreateDesc);
+		void* UploadMemory = FMemory::Malloc(UploadMemorySize, 16);
+
+		return UE::RHICore::FDefaultLayoutTextureInitializer(RHICmdList, Texture, UploadMemory, UploadMemorySize,
+			[this, Texture = TRefCountPtr<FD3D11Texture>(Texture), UploadMemory = UE::RHICore::FInitializerScopedMemory(UploadMemory), UploadMemorySize](FRHICommandListBase& RHICmdList) mutable
+			{
+				TConstArrayView<uint8> InitialData(reinterpret_cast<const uint8*>(UploadMemory.Pointer), UploadMemorySize);
+
+				return TRefCountPtr<FRHITexture>(FinalizeCreateTextureInternal(Texture, Texture->GetDesc(), InitialData));
+			});
+	}
+
+	return UE::RHICore::HandleUnknownTextureInitializerInitAction(RHICmdList, CreateDesc);
 }
 
 FTextureRHIRef FD3D11DynamicRHI::RHIAsyncCreateTexture2D(uint32 SizeX, uint32 SizeY, uint8 Format, uint32 NumMips, ETextureCreateFlags Flags, ERHIAccess InResourceState, void** InitialMipData, uint32 NumInitialMips, const TCHAR* DebugName, FGraphEventRef& OutCompletionEvent)
