Disable Halfresolution LocalFog in the RealtimeSkyCapture pass.

diff --git a/Engine/Source/Runtime/Renderer/Private/MobileBasePassRendering.cpp b/Engine/Source/Runtime/Renderer/Private/MobileBasePassRendering.cpp
--- a/Engine/Source/Runtime/Renderer/Private/MobileBasePassRendering.cpp
+++ b/Engine/Source/Runtime/Renderer/Private/MobileBasePassRendering.cpp
@@ -350,7 +350,7 @@ void SetupMobileBasePassUniformParameters(
 	BasePassParameters.LFV = View.LocalFogVolumeViewData.UniformParametersStruct;
 
 	// We need to compose the half resolution LFV texture when rendering mesh with Sky materials. So that the Fog passes remains cheap and we can keep the stencil test on the fog pass.
-	if (BasePass > EMobileBasePass::DepthPrePass)
+	if (BasePass > EMobileBasePass::DepthPrePass && !bForRealtimeSkyCapture)
 	{
 		// HalfResLocalFogVolumeView is rendered after the depth pre pass so we only bind it after the depth pre pass
 		BasePassParameters.bApplyHalfResLocalFogToSkyMeshes = View.LocalFogVolumeViewData.bUseHalfResLocalFogVolume ? 1 : 0;
