Enable GPUScene support for Niagara mesh rendering on mobile platforms

diff --git a/Engine/Plugins/FX/Niagara/Shaders/Private/NiagaraMeshVertexFactory.ush b/Engine/Plugins/FX/Niagara/Shaders/Private/NiagaraMeshVertexFactory.ush
--- a/Engine/Plugins/FX/Niagara/Shaders/Private/NiagaraMeshVertexFactory.ush
+++ b/Engine/Plugins/FX/Niagara/Shaders/Private/NiagaraMeshVertexFactory.ush
@@ -864,6 +864,11 @@ uint VertexFactoryGetPrimitiveId(FVertexFactoryInterpolantsVSToPS Interpolants)
 #endif
 }
 
+float4 VertexFactoryGetInstanceHitProxyId(FVertexFactoryInput Input, FVertexFactoryIntermediates Intermediates)
+{
+	return 0;
+}
+
 float4 VertexFactoryGetTranslatedPrimitiveVolumeBounds(FVertexFactoryInterpolantsVSToPS Interpolants)
 {
 	return float4(0,0,0,0);
diff --git a/Engine/Plugins/FX/Niagara/Source/Niagara/Private/NiagaraRendererMeshes.cpp b/Engine/Plugins/FX/Niagara/Source/Niagara/Private/NiagaraRendererMeshes.cpp
--- a/Engine/Plugins/FX/Niagara/Source/Niagara/Private/NiagaraRendererMeshes.cpp
+++ b/Engine/Plugins/FX/Niagara/Source/Niagara/Private/NiagaraRendererMeshes.cpp
@@ -318,8 +318,7 @@ void FNiagaraRendererMeshes::PrepareParticleMeshRenderData(FRHICommandListBase&
 	ParticleMeshRenderData.Collector = &Collector;
 
 #if NIAGARA_ENABLE_GPU_SCENE_MESHES
-	ParticleMeshRenderData.bUseGPUScene = FeatureLevel > ERHIFeatureLevel::ES3_1	
-		&& UseGPUScene(SceneProxy->GetScene().GetShaderPlatform(), FeatureLevel);
+	ParticleMeshRenderData.bUseGPUScene = UseGPUScene(SceneProxy->GetScene().GetShaderPlatform(), FeatureLevel);
 #else
 	ParticleMeshRenderData.bUseGPUScene = false;
 #endif
diff --git a/Engine/Plugins/FX/Niagara/Source/NiagaraVertexFactories/Private/NiagaraMeshVertexFactory.cpp b/Engine/Plugins/FX/Niagara/Source/NiagaraVertexFactories/Private/NiagaraMeshVertexFactory.cpp
--- a/Engine/Plugins/FX/Niagara/Source/NiagaraVertexFactories/Private/NiagaraMeshVertexFactory.cpp
+++ b/Engine/Plugins/FX/Niagara/Source/NiagaraVertexFactories/Private/NiagaraMeshVertexFactory.cpp
@@ -144,7 +144,7 @@ void FNiagaraMeshVertexFactory::InitRHI(FRHICommandListBase& RHICmdList)
 	if (bAddPrimitiveIDElement)
 	{
 		// TODO: Support GPU Scene on mobile? Maybe only for CPU particles?
-		AddPrimitiveIdStreamElement(EVertexInputStreamType::Default, Elements, 13, 0xFF);
+		AddPrimitiveIdStreamElement(EVertexInputStreamType::Default, Elements, 13, 13);
 	}
 #endif
 
@@ -173,13 +173,17 @@ void FNiagaraMeshVertexFactory::ModifyCompilationEnvironment(const FVertexFactor
 
 #if NIAGARA_ENABLE_GPU_SCENE_MESHES
 	const ERHIFeatureLevel::Type MaxSupportedFeatureLevel = GetMaxSupportedFeatureLevel(Parameters.Platform);
-
-	// TODO: Support GPU Scene on mobile?
-	const bool bUseGPUScene = UseGPUScene(Parameters.Platform, MaxSupportedFeatureLevel) && MaxSupportedFeatureLevel > ERHIFeatureLevel::ES3_1;
+	const bool bUseGPUScene = UseGPUScene(Parameters.Platform, MaxSupportedFeatureLevel);
 	const bool bSupportsPrimitiveIdStream = Parameters.VertexFactoryType->SupportsPrimitiveIdStream();
 	
 	OutEnvironment.SetDefine(TEXT("VF_SUPPORTS_PRIMITIVE_SCENE_DATA"), bSupportsPrimitiveIdStream && bUseGPUScene);
 	OutEnvironment.SetDefine(TEXT("VF_REQUIRES_PER_INSTANCE_CUSTOM_DATA"), bSupportsPrimitiveIdStream && bUseGPUScene);
+
+	// Mobile GPUScene implementation relies on USE_INSTANCE_CULLING define for rendering instanced meshes
+	if (bUseGPUScene && MaxSupportedFeatureLevel == ERHIFeatureLevel::ES3_1)
+	{
+		OutEnvironment.SetDefine(TEXT("USE_INSTANCE_CULLING"), 1);
+	}
 #endif
 
 	if (RHISupportsManualVertexFetch(Parameters.Platform))
diff --git a/Engine/Source/Runtime/Renderer/Private/MobileShadingRenderer.cpp b/Engine/Source/Runtime/Renderer/Private/MobileShadingRenderer.cpp
--- a/Engine/Source/Runtime/Renderer/Private/MobileShadingRenderer.cpp
+++ b/Engine/Source/Runtime/Renderer/Private/MobileShadingRenderer.cpp
@@ -722,17 +722,6 @@ void FMobileSceneRenderer::InitViews(
 		}		
 	}
 
-	{
-		RDG_CSV_STAT_EXCLUSIVE_SCOPE(GraphBuilder, UpdateGPUScene);
-
-		for (int32 ViewIndex = 0; ViewIndex < AllViews.Num(); ViewIndex++)
-		{
-			FViewInfo& View = *AllViews[ViewIndex];
-			Scene->GPUScene.UploadDynamicPrimitiveShaderDataForView(GraphBuilder, View);
-			Scene->GPUScene.DebugRender(GraphBuilder, GetSceneUniforms(), View);
-		}
-	}
-
 	if (bRendererOutputFinalSceneColor)
 	{
 		const bool bDynamicShadows = ViewFamily.EngineShowFlags.DynamicShadows;
@@ -1169,6 +1158,30 @@ void FMobileSceneRenderer::Render(FRDGBuilder& GraphBuilder, const FSceneRenderU
 	{
 		BeginOcclusionScope(GraphBuilder, Views);
 	}
+
+	// Notify the FX system that the scene is about to be rendered.
+	// TODO: These should probably be moved to scene extensions
+	if (FXSystem)
+	{
+		FXSystem->PreRender(GraphBuilder, GetSceneViews(), GetSceneUniforms(), true /*bAllowGPUParticleUpdate*/);
+		if (FGPUSortManager* GPUSortManager = FXSystem->GetGPUSortManager())
+		{
+			// if GPUSortManager::OnPostRenderOpaque is called below (from RenderOpaqueFX) we must also call OnPreRender (as it sets up
+			// the internal state of the GPUSortManager).  Any optimization to skip this block needs to take that into consideration.
+			GPUSortManager->OnPreRender(GraphBuilder);
+		}
+	}
+
+	{
+		RDG_CSV_STAT_EXCLUSIVE_SCOPE(GraphBuilder, UpdateGPUScene);
+
+		for (int32 ViewIndex = 0; ViewIndex < AllViews.Num(); ViewIndex++)
+		{
+			FViewInfo& View = *AllViews[ViewIndex];
+			Scene->GPUScene.UploadDynamicPrimitiveShaderDataForView(GraphBuilder, View);
+			Scene->GPUScene.DebugRender(GraphBuilder, GetSceneUniforms(), View);
+		}
+	}
 	
 	GetSceneExtensionsRenderers().UpdateViewData(GraphBuilder, ViewDataManager);
 
@@ -1230,19 +1243,6 @@ void FMobileSceneRenderer::Render(FRDGBuilder& GraphBuilder, const FSceneRenderU
 			SetDummyForwardLightUniformBufferOnViews(GraphBuilder, ShaderPlatform, Views);
 		}
 
-		// Notify the FX system that the scene is about to be rendered.
-		// TODO: These should probably be moved to scene extensions
-		if (FXSystem)
-		{
-			FXSystem->PreRender(GraphBuilder, GetSceneViews(), GetSceneUniforms(), true /*bAllowGPUParticleUpdate*/);
-			if (FGPUSortManager* GPUSortManager = FXSystem->GetGPUSortManager())
-			{
-				// if GPUSortManager::OnPostRenderOpaque is called below (from RenderOpaqueFX) we must also call OnPreRender (as it sets up
-				// the internal state of the GPUSortManager).  Any optimization to skip this block needs to take that into consideration.
-				GPUSortManager->OnPreRender(GraphBuilder);
-			}
-		}
-
 		// Generate the Sky/Atmosphere look up tables
 		const bool bShouldRenderSkyAtmosphere = ShouldRenderSkyAtmosphere(Scene, ViewFamily.EngineShowFlags);
 		if (bShouldRenderSkyAtmosphere)
diff --git a/Engine/Source/Runtime/Renderer/Private/SceneHitProxyRendering.cpp b/Engine/Source/Runtime/Renderer/Private/SceneHitProxyRendering.cpp
--- a/Engine/Source/Runtime/Renderer/Private/SceneHitProxyRendering.cpp
+++ b/Engine/Source/Runtime/Renderer/Private/SceneHitProxyRendering.cpp
@@ -565,12 +565,6 @@ void FMobileSceneRenderer::RenderHitProxies(FRDGBuilder& GraphBuilder, const FSc
 
 	// Find the visible primitives.
 	InitViews(GraphBuilder, SceneTexturesConfig, InstanceCullingManager, nullptr, InitViewTaskDatas);
-	
-	GetSceneExtensionsRenderers().UpdateViewData(GraphBuilder, ViewDataManager);
-	GetSceneExtensionsRenderers().UpdateSceneUniformBuffer(GraphBuilder, GetSceneUniforms());
-
-	GetSceneExtensionsRenderers().PreRender(GraphBuilder);
-	GEngine->GetPreRenderDelegateEx().Broadcast(GraphBuilder);
 
 	// Notify the FX system that the scene is about to be rendered.
 	// TODO: These should probably be moved to scene extensions
@@ -591,6 +585,23 @@ void FMobileSceneRenderer::RenderHitProxies(FRDGBuilder& GraphBuilder, const FSc
 		}
 	}
 
+	{
+		RDG_CSV_STAT_EXCLUSIVE_SCOPE(GraphBuilder, UpdateGPUScene);
+
+		for (int32 ViewIndex = 0; ViewIndex < AllViews.Num(); ViewIndex++)
+		{
+			FViewInfo& View = *AllViews[ViewIndex];
+			Scene->GPUScene.UploadDynamicPrimitiveShaderDataForView(GraphBuilder, View);
+			Scene->GPUScene.DebugRender(GraphBuilder, GetSceneUniforms(), View);
+		}
+	}
+	
+	GetSceneExtensionsRenderers().UpdateViewData(GraphBuilder, ViewDataManager);
+	GetSceneExtensionsRenderers().UpdateSceneUniformBuffer(GraphBuilder, GetSceneUniforms());
+
+	GetSceneExtensionsRenderers().PreRender(GraphBuilder);
+	GEngine->GetPreRenderDelegateEx().Broadcast(GraphBuilder);
+
 	TArray<Nanite::FRasterResults, TInlineAllocator<2>> NaniteRasterResults;
 	::DoRenderHitProxies(GraphBuilder, this, HitProxyTexture, HitProxyDepthTexture, NaniteRasterResults, InstanceCullingManager);
 
