Heterogeneous Volumes: Fixes to LODInfo

1. Confine bIsPespective to Spot Lights to prevent improper frustum culling for point and rect lights
2. Apply PreViewTranslation to correctly determine WorldOrigin when building Camera AVSM

diff --git a/Engine/Source/Runtime/Renderer/Private/HeterogeneousVolumes/HeterogeneousVolumesLiveShadingPipeline.cpp b/Engine/Source/Runtime/Renderer/Private/HeterogeneousVolumes/HeterogeneousVolumesLiveShadingPipeline.cpp
--- a/Engine/Source/Runtime/Renderer/Private/HeterogeneousVolumes/HeterogeneousVolumesLiveShadingPipeline.cpp
+++ b/Engine/Source/Runtime/Renderer/Private/HeterogeneousVolumes/HeterogeneousVolumesLiveShadingPipeline.cpp
@@ -2312,7 +2312,8 @@ bool RenderVolumetricShadowMapForLightWithLiveShading(
 	LODInfo.FOV = PI / 4.0f;
 	LODInfo.NearClippingDistance = HeterogeneousVolumes::GetShadowNearClippingDistance();
 	LODInfo.DownsampleFactor = 1.0;
-	LODInfo.bIsPerspective = (LightType != LightType_Directional);
+	// TODO: Support alternate test for rect lights
+	LODInfo.bIsPerspective = (LightType == LightType_Spot);
 
 	// Iterate over shadow-casting volumes
 	bool bHasShadowCastingVolume = false;
@@ -2649,12 +2650,13 @@ void RenderAdaptiveVolumetricCameraMapWithLiveShading(
 			TEXT("HeterogeneousVolume.VolumetricShadowLinkedListBuffer")
 		);
 
+		FVector PreViewTranslation = View.ViewMatrices.GetPreViewTranslation();
+		FMatrix WorldToTranslatedWorldMatrix = FTranslationMatrix(PreViewTranslation);
+
 		HeterogeneousVolumes::FLODInfo LODInfo;
 		LODInfo.WorldSceneBounds = WorldVolumeBounds;
-		LODInfo.WorldOrigin = FVector(TranslatedWorldOrigin);
+		LODInfo.WorldOrigin = FVector(TranslatedWorldOrigin) - PreViewTranslation;
 		LODInfo.ViewRect = FIntRect(0, 0, ShadowMapResolution.X, ShadowMapResolution.Y);
-		FVector PreViewTranslation = View.ViewMatrices.GetPreViewTranslation();
-		FMatrix WorldToTranslatedWorldMatrix = FTranslationMatrix(PreViewTranslation);
 		GetViewFrustumBounds(LODInfo.WorldShadowFrustum, WorldToTranslatedWorldMatrix * FMatrix(TranslatedWorldToShadow[0]), true);
 		LODInfo.FOV = FOV;
 		LODInfo.NearClippingDistance = HeterogeneousVolumes::GetShadowNearClippingDistance();
