VirtualTexture - Fix Metal Validation error when PageTable Texture Reference is unitialized during one frame
Merge two calls into one  (CreateTextureReference + RHIUpdateTextureReference into CreateTextureReference)

#tests Local test on macOS running the editor with Metal Validation Layer

diff --git a/Engine/Source/Runtime/Renderer/Private/VT/VirtualTextureSpace.cpp b/Engine/Source/Runtime/Renderer/Private/VT/VirtualTextureSpace.cpp
--- a/Engine/Source/Runtime/Renderer/Private/VT/VirtualTextureSpace.cpp
+++ b/Engine/Source/Runtime/Renderer/Private/VT/VirtualTextureSpace.cpp
@@ -128,10 +128,9 @@ void FVirtualTextureSpace::InitRHI(FRHICommandListBase& RHICmdList)
 	for (uint32 TextureIndex = 0u; TextureIndex < GetNumPageTableTextures(); ++TextureIndex)
 	{
 		FTextureEntry& TextureEntry = PageTable[TextureIndex];
-		TextureEntry.TextureReferenceRHI = RHICmdList.CreateTextureReference();
+		TextureEntry.TextureReferenceRHI = RHICmdList.CreateTextureReference(GBlackUintTexture->TextureRHI);
 	}
-	PageTableIndirection.TextureReferenceRHI = RHICmdList.CreateTextureReference();
-	RHICmdList.UpdateTextureReference(PageTableIndirection.TextureReferenceRHI, GBlackUintTexture->TextureRHI);
+	PageTableIndirection.TextureReferenceRHI = RHICmdList.CreateTextureReference(GBlackUintTexture->TextureRHI);
 }
 
 void FVirtualTextureSpace::ReleaseRHI()
