[VSM] Add experimental Nanite sustained load throttling through dynamic res.
When enabled, VSM Nanite ShadowDepths performance is monitored. When a performance threshold is reached, VSM resolution ramps down by biasing the LOD up.
The bias is determined per VSM with a heuristic based on historical nanite cluster counts, i.e. expensive shadow maps are biased more than cheap ones (when data is available).
The performance target is set with one of two cvars:
r.Shadow.Virtual.DynamicRes.TimeBudgetMs: use the DynamicRes system to time ShadowDepths.
r.Shadow.Virtual.DynamicRes.LoadBudget: estimate global ShadowDepths cost based on Nanite cluster data.
This system is likely most effective in scenes with little or no caching, as biasing may mark pages that were previously unmarked.

diff --git a/Engine/Shaders/Private/VirtualShadowMaps/VirtualShadowMapPageMarking.usf b/Engine/Shaders/Private/VirtualShadowMaps/VirtualShadowMapPageMarking.usf
--- a/Engine/Shaders/Private/VirtualShadowMaps/VirtualShadowMapPageMarking.usf
+++ b/Engine/Shaders/Private/VirtualShadowMaps/VirtualShadowMapPageMarking.usf
@@ -34,6 +34,8 @@ RWTexture2D<uint> OutPageRequestFlags;
 RWTexture2D<uint> OutPageReceiverMasks;
 StructuredBuffer<int> DirectionalLightIds;
 
+StructuredBuffer<uint> ThrottleBuffer;
+
 float PageDilationBorderSizeDirectional;
 float PageDilationBorderSizeLocal;
 uint InputType;
@@ -41,14 +43,50 @@ uint bCullBackfacingPixels;
 uint NumDirectionalLightSmInds;
 
 float DOFBiasStrength;
+float ThrottleMaxBiasDirectional;
+float ThrottleMaxBiasLocal;
 
 #ifdef GeneratePageFlagsFromFroxelsCS
 int PassId;
 #endif
 
+float GetSustainedLoadThrottlingResolutionBiasDirectional(FVirtualShadowMapProjectionShaderData BaseProjectionData, float BiasedLevel)
+{
+	float ResolutionBias = 0.0f;
+
+#if PERMUTATION_THROTTLING
+	const FVirtualShadowMapHandle VSMHandle = BaseProjectionData.VirtualShadowMapHandle.MakeOffset(BiasedLevel - BaseProjectionData.ClipmapLevel);
+
+	float ThrottleIntensity = asfloat(ThrottleBuffer[VSM_TB_SIZEOF_HEADER + VSM_TB_SIZEOF_ENTRY * VSMHandle.GetDataIndex() + VSM_TB_ENTRY_THROTTLE]);
+	
+	ResolutionBias = ThrottleIntensity * ThrottleMaxBiasDirectional;
+#endif
+
+	return ResolutionBias;
+}
+
+float GetSustainedLoadThrottlingResolutionBiasLocal(FVirtualShadowMapProjectionShaderData ProjectionData)
+{
+	float ResolutionBias = 0.0f;
+
+#if PERMUTATION_THROTTLING
+	const FVirtualShadowMapHandle VSMHandle = ProjectionData.VirtualShadowMapHandle;
+
+	float ThrottleIntensity = asfloat(ThrottleBuffer[VSM_TB_SIZEOF_HEADER + VSM_TB_SIZEOF_ENTRY * VSMHandle.GetDataIndex() + VSM_TB_ENTRY_THROTTLE]);
+	
+	ResolutionBias = ThrottleIntensity * ThrottleMaxBiasLocal;
+#endif
+
+	return ResolutionBias;
+}
+
 int GetBiasedClipmapLevel(FVirtualShadowMapProjectionShaderData BaseProjectionData, float3 TranslatedWorldPosition, float ExtraBias)
 {
 	float BiasedLevel = CalcAbsoluteClipmapLevel(BaseProjectionData, TranslatedWorldPosition) + BaseProjectionData.ResolutionLodBias + VirtualShadowMap.GlobalResolutionLodBias + ExtraBias;
+
+	float ThrottlingBias = GetSustainedLoadThrottlingResolutionBiasDirectional(BaseProjectionData, BiasedLevel);
+	BiasedLevel += ThrottlingBias;
+
 	return int(floor(BiasedLevel));
 }
 
@@ -58,6 +96,8 @@ uint GetMipLevelLocal(FVirtualShadowMapProjectionShaderData ProjectionData, floa
 	float3 ViewToShadowTranslation = DFFastLocalSubtractDemote(ProjectionData.PreViewTranslation, PrimaryView.PreViewTranslation);
 	float3 ShadowTranslatedWorldPosition = TranslatedWorldPosition + ViewToShadowTranslation;
 
+	ExtraBias += GetSustainedLoadThrottlingResolutionBiasLocal(ProjectionData);
+
 	float Footprint = VirtualShadowMapCalcPixelFootprintLocal(ProjectionData, ShadowTranslatedWorldPosition, SceneDepth);
 	return GetMipLevelLocal(Footprint, VirtualShadowMap.MipModeLocal, ProjectionData.ResolutionLodBias, VirtualShadowMap.GlobalResolutionLodBias, ExtraBias);
 }
diff --git a/Engine/Shaders/Private/VirtualShadowMaps/VirtualShadowMapPrintStats.usf b/Engine/Shaders/Private/VirtualShadowMaps/VirtualShadowMapPrintStats.usf
--- a/Engine/Shaders/Private/VirtualShadowMaps/VirtualShadowMapPrintStats.usf
+++ b/Engine/Shaders/Private/VirtualShadowMaps/VirtualShadowMapPrintStats.usf
@@ -40,6 +40,25 @@ void PrintLeftAlign(inout FShaderPrintContext Context, uint Value)
 	Print(Context, Value);
 }
 
+void PrintLeftAlign(inout FShaderPrintContext Context, float Value)
+{
+	int Len = 9;
+	
+	uint Tmp = Value;
+	while (Tmp >= 10)
+	{
+		Tmp /= 10;
+		Len--;
+	}
+
+	while (Len > 0)
+	{
+		PrintSymbol(Context, _SPC_);
+		Len--;
+	}
+	
+	Print(Context, Value);
+}
 
 void PrintUnits(inout FShaderPrintContext Context, uint Value)
 {
@@ -68,6 +87,9 @@ uint StatsMessageTimestampLow;
 StructuredBuffer<uint> NanitePerformanceFeedback;
 uint NanitePerformanceFeedbackNumEntries;
 
+StructuredBuffer<uint> ThrottleBuffer;
+uint bThrottlingEnabled;
+
 StructuredBuffer<FNaniteStats> NaniteStats;
 
 void SendStatusMessage()
@@ -171,11 +193,12 @@ void PrintStats()
 	const bool bShowPhysicalPages      = (ShowStatsValue == 1) || ShowStatsValue == 2;
 	const bool bShowNonNaniteInstances = (ShowStatsValue == 1) || ShowStatsValue == 3;
 
-	float TopMargin = 0.05f;
-	float HeadlineX = 0.47f;
-	float ItemX = 0.48f;
-	float ValueX = 0.58f;
+	const float TopMargin = 0.05f;
+	const float HeadlineX = 0.47f;
+	const float ItemX = 0.48f;
 	FShaderPrintContext Context = InitShaderPrintContext(true, float2(HeadlineX, TopMargin));
+	const int NumLabelCharacters = 26;
+	const float ValueX = HeadlineX + Context.Config.FontSpacing.x * NumLabelCharacters;
 
 	if (bShowPhysicalPages)
 	{
@@ -255,10 +278,29 @@ void PrintStats()
 
 		Context.Pos.x = HeadlineX;
 		Print(Context, TEXT("Global Resolution Bias   "), FontWhite);
-		Print(Context, VirtualShadowMap.GlobalResolutionLodBias);
+		Context.Pos.x = ValueX;
+		PrintLeftAlign(Context, VirtualShadowMap.GlobalResolutionLodBias);
 		Newline(Context);
 	}
 
+	float GlobalComputeThrottleIntensity = 0.0f;
+	if (bThrottlingEnabled)
+	{
+		GlobalComputeThrottleIntensity = asfloat(ThrottleBuffer[VSM_TB_HEADER_TOTAL_THROTTLE]);
+	}
+
+	Context.Pos.x = HeadlineX;
+	Print(Context, TEXT("Compute Throttle"), FontWhite);
+	Context.Pos *= float2(Context.Config.Resolution);
+	float ValueWidth = Context.Config.FontSpacing.x * 9.5;
+	float BarStartX = ValueX * Context.Config.Resolution.x;
+	float BarEndX = (ValueX + ValueWidth) * Context.Config.Resolution.x;
+	float BarHeight = 8.0f;
+	AddFilledQuadSS(Context, float2(BarStartX, Context.Pos.y - BarHeight/2), float2(BarEndX, Context.Pos.y + BarHeight/2), float4(0,0,0,0.5));
+	AddFilledQuadSS(Context, float2(BarStartX, Context.Pos.y - BarHeight/2), float2(lerp(BarStartX, BarEndX, GlobalComputeThrottleIntensity), Context.Pos.y + BarHeight/2), float4(1,0,0,1.0));
+	Context.Pos /= float2(Context.Config.Resolution);
+	Newline(Context);
+
 	if (bShowNonNaniteInstances)
 	{
 		Context.Pos.x = HeadlineX;
@@ -336,15 +378,39 @@ void PrintStats()
 	Context.Pos.y -= HistogramScale.y / Context.Config.Resolution.y;
 
 	Context.Pos *= float2(Context.Config.Resolution);
-	AddFilledQuadSS(Context, Context.Pos, Context.Pos + float2(HistogramScale.x * max(64, NaniteGeoHistogramNonZeroBins), HistogramScale.y), float4(HistogramMaxIndicatorColor.rgb * 0.25,0.4));
+	float4 HistogramRect = float4(
+		Context.Pos,
+		Context.Pos + float2(HistogramScale.x * max(50, NaniteGeoHistogramNonZeroBins), HistogramScale.y)
+	);
+	AddFilledQuadSS(Context, HistogramRect.xy, HistogramRect.zw, float4(HistogramMaxIndicatorColor.rgb * 0.25,0.4));
 	Context.Pos.y += HistogramScale.y;
 	for (int Index = 0; Index < NANITE_GEO_HISTOGRAM_NUM_BINS; Index++)
 	{
 		float HWClusters = (NaniteGeoHistogramHWClusters[Index] / float(HistogramMax)) * HistogramScale.y;
 		float SWClusters = (NaniteGeoHistogramSWClusters[Index] / float(HistogramMax)) * HistogramScale.y;
+		
+		bool bIsSinglePage = Index < VSM_GENERATE_STATS_GROUPSIZE;
+		FVirtualShadowMapHandle VSMHandle;
+		if (bIsSinglePage)
+		{
+			VSMHandle = FVirtualShadowMapHandle::MakeFromId(Index);
+		}
+		else
+		{
+			VSMHandle = FVirtualShadowMapHandle::MakeFromFullIndex(Index - VSM_GENERATE_STATS_GROUPSIZE);
+		}
+		float ThrottleIntensity = 0.0f;
+		if (bThrottlingEnabled)
+		{
+			ThrottleIntensity = asfloat(ThrottleBuffer[VSM_TB_SIZEOF_HEADER + VSM_TB_SIZEOF_ENTRY * VSMHandle.GetDataIndex() + VSM_TB_ENTRY_THROTTLE]);
+		}
+
+		bool bHasClusterData = HWClusters > 0 || SWClusters > 0;
+
 		if (HWClusters > 0) { AddFilledQuadSS(Context, Context.Pos, Context.Pos + float2(HistogramScale.x, -HWClusters), float4(0,0,1,1.0)); }
 		if (SWClusters > 0) { AddFilledQuadSS(Context, Context.Pos + float2(0, -HWClusters), Context.Pos + float2(0, -HWClusters) + float2(HistogramScale.x, -SWClusters), float4(0,1,1,1.0)); }
-		if (SWClusters + HWClusters > 0) { Context.Pos.x += HistogramScale.x; }
+		if (bHasClusterData && ThrottleIntensity > 0) { AddFilledQuadSS(Context, Context.Pos, Context.Pos + float2(HistogramScale.x, ThrottleIntensity * HistogramScale.y), float4(1,0,0,1.0)); }
+		if (bHasClusterData) { Context.Pos.x += HistogramScale.x; }
 	}
 	Context.Pos /= float2(Context.Config.Resolution);
 	Newline(Context);
diff --git a/Engine/Shaders/Private/VirtualShadowMaps/VirtualShadowMapThrottle.usf b/Engine/Shaders/Private/VirtualShadowMaps/VirtualShadowMapThrottle.usf
new file mode 100644
--- /dev/null
+++ b/Engine/Shaders/Private/VirtualShadowMaps/VirtualShadowMapThrottle.usf
@@ -0,0 +1,127 @@
+// Copyright Epic Games, Inc. All Rights Reserved.
+
+#include "../Common.ush"
+#include "VirtualShadowMapHandle.ush"
+#include "/Engine/Shared/VirtualShadowMapDefinitions.h"
+
+StructuredBuffer<uint> PrevThrottleBuffer;
+StructuredBuffer<uint> PrevNanitePerformanceFeedback;
+
+StructuredBuffer<FNextVirtualShadowMapData> NextVirtualShadowMapData;
+uint NextVirtualShadowMapDataCount;
+
+uint NumThrottleBufferEntries;
+RWStructuredBuffer<uint> OutThrottleBuffer;
+
+float DynResThrottleStrength;
+float ThrottleLoadBudget;
+float ThrottleEMAHistoryWeight;
+
+[numthreads(64, 1, 1)]
+void ProcessPrevFramePerfDataCS(uint DispatchThreadId : SV_DispatchThreadID)
+{
+	const uint TotalClustersHW = PrevNanitePerformanceFeedback[VSM_NPF_HEADER_TOTAL_HW_CLUSTERS];
+	const uint TotalClustersSW = PrevNanitePerformanceFeedback[VSM_NPF_HEADER_TOTAL_SW_CLUSTERS];
+
+	if (DispatchThreadId == 0)
+	{
+		OutThrottleBuffer[VSM_TB_HEADER_TOTAL_HW_CLUSTERS] = TotalClustersHW;
+		OutThrottleBuffer[VSM_TB_HEADER_TOTAL_SW_CLUSTERS] = TotalClustersSW;
+	}
+
+	if(DispatchThreadId < NextVirtualShadowMapDataCount)
+	{
+		FVirtualShadowMapHandle PrevVSMHandle = FVirtualShadowMapHandle::MakeFromId(DispatchThreadId);
+		const FNextVirtualShadowMapData NextVSMData = NextVirtualShadowMapData[PrevVSMHandle.GetDataIndex()];
+
+		FVirtualShadowMapHandle NextVSMHandle = FVirtualShadowMapHandle::MakeFromId(NextVSMData.NextVirtualShadowMapId);
+		if (NextVSMHandle.IsValid())
+		{
+			const uint NPFOffset = VSM_NPF_SIZEOF_HEADER + VSM_NPF_SIZEOF_ENTRY * PrevVSMHandle.GetDataIndex();
+			const uint ClustersHW = PrevNanitePerformanceFeedback[NPFOffset + VSM_NPF_ENTRY_HW_CLUSTERS];
+			const uint ClustersSW = PrevNanitePerformanceFeedback[NPFOffset + VSM_NPF_ENTRY_SW_CLUSTERS];
+
+			const uint ThrottleBufferOffset = VSM_TB_SIZEOF_HEADER + VSM_TB_SIZEOF_ENTRY * NextVSMHandle.GetDataIndex();
+			OutThrottleBuffer[ThrottleBufferOffset + VSM_TB_ENTRY_HW_CLUSTERS] = ClustersHW;
+			OutThrottleBuffer[ThrottleBufferOffset + VSM_TB_ENTRY_SW_CLUSTERS] = ClustersSW;
+			OutThrottleBuffer[ThrottleBufferOffset + VSM_TB_ENTRY_THROTTLE] = PrevThrottleBuffer[VSM_TB_SIZEOF_HEADER + VSM_TB_SIZEOF_ENTRY * PrevVSMHandle.GetDataIndex() + VSM_TB_ENTRY_THROTTLE];
+
+			InterlockedMax(OutThrottleBuffer[VSM_TB_HEADER_MAX_HW_CLUSTERS], ClustersHW);
+			InterlockedMax(OutThrottleBuffer[VSM_TB_HEADER_MAX_SW_CLUSTERS], ClustersSW);
+		}
+	}
+}
+
+float CalculateCostHeuristic(const float NaniteClustersHW, const float NaniteClustersSW)
+{
+	// todo: can significantly improve the fit of this heuristic with page count data (requested, cached, invalidated)
+	return 0.00012176f * NaniteClustersHW + 0.00002700f * NaniteClustersSW;
+}
+
+[numthreads(64, 1, 1)]
+void UpdateThrottleParametersCS(uint DispatchThreadId : SV_DispatchThreadID)
+{
+	const uint TotalClustersHW = OutThrottleBuffer[VSM_TB_HEADER_TOTAL_HW_CLUSTERS];
+	const uint TotalClustersSW = OutThrottleBuffer[VSM_TB_HEADER_TOTAL_SW_CLUSTERS];
+	const uint MaxClustersHW = OutThrottleBuffer[VSM_TB_HEADER_MAX_HW_CLUSTERS];
+	const uint MaxClustersSW = OutThrottleBuffer[VSM_TB_HEADER_MAX_SW_CLUSTERS];
+	
+	float GlobalThrottleIntensity;
+	{
+		float TotalCostHeuristic = CalculateCostHeuristic(TotalClustersHW, TotalClustersSW);
+	
+		bool bThrottleBasedOnDynRes = DynResThrottleStrength >= 0;
+		if (bThrottleBasedOnDynRes)
+		{
+			GlobalThrottleIntensity = DynResThrottleStrength;
+		}
+		else
+		{
+			// hand-tuned parameters
+			const float ThrottleResponse = 0.5f;
+			const float ThrottleMaxStepUp = 0.2f;
+			const float ThrottleMaxStepDown = -0.01f;
+
+			float TotalThrottleHeuristic = TotalCostHeuristic - ThrottleLoadBudget;
+			float ThrottleModifier = clamp(TotalThrottleHeuristic * ThrottleResponse, ThrottleMaxStepDown, ThrottleMaxStepUp);
+	
+			float PrevThrottleIntensity = asfloat(PrevThrottleBuffer[VSM_TB_HEADER_TOTAL_THROTTLE]);
+			GlobalThrottleIntensity = saturate(PrevThrottleIntensity + ThrottleModifier);
+		}
+	}
+
+	if (DispatchThreadId == 0)
+	{
+		OutThrottleBuffer[VSM_TB_HEADER_TOTAL_THROTTLE] = asuint(GlobalThrottleIntensity);
+	}
+
+	if(DispatchThreadId < NumThrottleBufferEntries)
+	{
+		FVirtualShadowMapHandle VSMHandle = FVirtualShadowMapHandle::MakeFromId(DispatchThreadId);
+
+		uint ThrottleBufferOffset = VSM_TB_SIZEOF_HEADER + VSM_TB_SIZEOF_ENTRY * VSMHandle.GetDataIndex();
+
+		uint ClustersHW = OutThrottleBuffer[ThrottleBufferOffset + VSM_TB_ENTRY_HW_CLUSTERS];
+		uint ClustersSW = OutThrottleBuffer[ThrottleBufferOffset + VSM_TB_ENTRY_SW_CLUSTERS];
+
+		float ThrottleIntensity;
+		if (ClustersHW + ClustersSW == 0)
+		{
+			// Likely no data available, fall back to global throttle value
+			ThrottleIntensity = GlobalThrottleIntensity;
+		}
+		else
+		{
+			float PrevThrottleIntensity = asfloat(OutThrottleBuffer[ThrottleBufferOffset + VSM_TB_ENTRY_THROTTLE]);
+			float MaxCostHeuristic = CalculateCostHeuristic(MaxClustersHW, MaxClustersSW);
+
+			float CostHeuristic = CalculateCostHeuristic(ClustersHW, ClustersSW);
+			float CostContibutionHeuristic = CostHeuristic / MaxCostHeuristic;
+
+			ThrottleIntensity = saturate(CostContibutionHeuristic * GlobalThrottleIntensity);
+			ThrottleIntensity = ThrottleIntensity * (1.0 - ThrottleEMAHistoryWeight) + PrevThrottleIntensity * ThrottleEMAHistoryWeight;
+		}
+
+		OutThrottleBuffer[ThrottleBufferOffset + VSM_TB_ENTRY_THROTTLE] = asuint(ThrottleIntensity);
+	}
+}
diff --git a/Engine/Shaders/Shared/VirtualShadowMapDefinitions.h b/Engine/Shaders/Shared/VirtualShadowMapDefinitions.h
--- a/Engine/Shaders/Shared/VirtualShadowMapDefinitions.h
+++ b/Engine/Shaders/Shared/VirtualShadowMapDefinitions.h
@@ -105,6 +105,18 @@ VirtualShadowMapDefinitions.h: used in virtual shadow map shaders and C++ code t
 #define VSM_NPF_ENTRY_SW_CLUSTERS						1
 #define VSM_NPF_SIZEOF_ENTRY							2
 
+// Throttle Buffer
+#define VSM_TB_HEADER_TOTAL_HW_CLUSTERS				0
+#define VSM_TB_HEADER_TOTAL_SW_CLUSTERS				1
+#define VSM_TB_HEADER_MAX_HW_CLUSTERS				2
+#define VSM_TB_HEADER_MAX_SW_CLUSTERS				3
+#define VSM_TB_HEADER_TOTAL_THROTTLE				4
+#define VSM_TB_SIZEOF_HEADER						5
+#define VSM_TB_ENTRY_HW_CLUSTERS					0
+#define VSM_TB_ENTRY_SW_CLUSTERS					1
+#define VSM_TB_ENTRY_THROTTLE						2
+#define VSM_TB_SIZEOF_ENTRY							3
+
 #ifdef __cplusplus
 #include "HLSLTypeAliases.h"
 #include "HLSLMathAliases.h"
diff --git a/Engine/Source/Runtime/Renderer/Private/VirtualShadowMaps/VirtualShadowMapArray.cpp b/Engine/Source/Runtime/Renderer/Private/VirtualShadowMaps/VirtualShadowMapArray.cpp
--- a/Engine/Source/Runtime/Renderer/Private/VirtualShadowMaps/VirtualShadowMapArray.cpp
+++ b/Engine/Source/Runtime/Renderer/Private/VirtualShadowMaps/VirtualShadowMapArray.cpp
@@ -519,6 +519,49 @@ static TAutoConsoleVariable<bool> CVarVSMUseReceiverMaskLocal(
 	ECVF_RenderThreadSafe
 );
 
+extern TAutoConsoleVariable<float> CVarVSMDynamicResolutionMaxLodBias;
+extern TAutoConsoleVariable<float> CVarVSMDynamicResolutionMaxLodBiasDirectional;
+extern TAutoConsoleVariable<float> CVarVSMDynamicResolutionMaxLodBiasLocal;
+
+static TAutoConsoleVariable<float> CVarTimeBudgetMs(
+	TEXT("r.Shadow.Virtual.DynamicRes.TimeBudgetMs"),
+	DynamicRenderScaling::FHeuristicSettings::kBudgetMsDisabled,
+	TEXT("Frame's time budget for VSM ShadowDepths in milliseconds. Non-Nanite not included.\n")
+	TEXT("Disabled when set to 0. Experimental."),
+	ECVF_RenderThreadSafe
+);
+
+static TAutoConsoleVariable<float> CVarThrottleLoadBudget(
+	TEXT("r.Shadow.Virtual.DynamicRes.LoadBudget"),
+	0.0f,
+	TEXT("Scale LOD bias based on an estimate of the ShadowDepths load. Lower value means more aggressive resolution downscaling.\n")
+	TEXT("Disabled when TimeBudgetMs is active, or if set to 0. Experimental."),
+	ECVF_RenderThreadSafe
+);
+
+static TAutoConsoleVariable<float> CVarThrottleEMAHistoryWeight(
+	TEXT("r.Shadow.Virtual.DynamicRes.PerVSMBiasHistoryWeight"),
+	0.9f,
+	TEXT("Smoothen out the per-VSM LOD bias over time. Value between 0 and 1."),
+	ECVF_RenderThreadSafe
+);
+
+static DynamicRenderScaling::FHeuristicSettings GetDynamicVSMResolutionSettings()
+{
+	DynamicRenderScaling::FHeuristicSettings BucketSetting;
+	BucketSetting.Model = DynamicRenderScaling::EHeuristicModel::Linear;
+	BucketSetting.bModelScalesWithPrimaryScreenPercentage = false;
+	// Resolution fraction is used as throttle strength, not directly as resolution
+	BucketSetting.MinResolutionFraction = DynamicRenderScaling::PercentageToFraction(0.0f);
+	BucketSetting.MaxResolutionFraction = DynamicRenderScaling::PercentageToFraction(100.0f);
+	BucketSetting.BudgetMs = CVarTimeBudgetMs.GetValueOnAnyThread();
+	BucketSetting.ChangeThreshold = DynamicRenderScaling::PercentageToFraction(1.0f);
+	BucketSetting.TargetedHeadRoom = DynamicRenderScaling::PercentageToFraction(5.0f); // 5% headroom
+	BucketSetting.UpperBoundQuantization = DynamicRenderScaling::FHeuristicSettings::kDefaultUpperBoundQuantization;
+	return BucketSetting;
+}
+DynamicRenderScaling::FBudget GDynamicResolutionVSMNaniteBudget(TEXT("DynamicResolutionVSMNanite"), &GetDynamicVSMResolutionSettings);
+
 bool IsVirtualShadowMapLocalReceiverMaskEnabled()
 {
 	return CVarVSMUseReceiverMaskLocal.GetValueOnRenderThread() != 0;
@@ -911,14 +954,64 @@ class FPruneLightGridCS : public FVirtualShadowMapPageManagementShader
 };
 IMPLEMENT_GLOBAL_SHADER(FPruneLightGridCS, "/Engine/Private/VirtualShadowMaps/VirtualShadowMapPageMarking.usf", "PruneLightGridCS", SF_Compute);
 
+BEGIN_SHADER_PARAMETER_STRUCT(FThrottlingParameters, )
+	SHADER_PARAMETER(float, DynResThrottleStrength)
+	SHADER_PARAMETER(float, ThrottleLoadBudget)
+	SHADER_PARAMETER(float, ThrottleEMAHistoryWeight)
+	SHADER_PARAMETER(float, ThrottleMaxBiasDirectional)
+	SHADER_PARAMETER(float, ThrottleMaxBiasLocal)
+END_SHADER_PARAMETER_STRUCT()
+
+class FProcessPrevFramePerfDataCS : public FVirtualShadowMapGlobalShader
+{
+	DECLARE_GLOBAL_SHADER(FProcessPrevFramePerfDataCS);
+	SHADER_USE_PARAMETER_STRUCT(FProcessPrevFramePerfDataCS, FVirtualShadowMapGlobalShader)
+public:
+	static constexpr int32 NumThreadsPerGroup = 64;
+
+	BEGIN_SHADER_PARAMETER_STRUCT(FParameters, )
+		SHADER_PARAMETER_RDG_BUFFER_SRV(StructuredBuffer<uint>, PrevThrottleBuffer)
+		SHADER_PARAMETER_RDG_BUFFER_SRV(StructuredBuffer<uint>, PrevNanitePerformanceFeedback)
+		SHADER_PARAMETER_RDG_BUFFER_SRV(StructuredBuffer<FNextVirtualShadowMapData>, NextVirtualShadowMapData)
+		SHADER_PARAMETER(uint32, NextVirtualShadowMapDataCount)
+
+		SHADER_PARAMETER_STRUCT_INCLUDE(FThrottlingParameters, ThrottlingParameters)
+
+		SHADER_PARAMETER_RDG_BUFFER_UAV(RWBuffer<uint>, OutThrottleBuffer)
+	END_SHADER_PARAMETER_STRUCT()
+};
+IMPLEMENT_GLOBAL_SHADER(FProcessPrevFramePerfDataCS, "/Engine/Private/VirtualShadowMaps/VirtualShadowMapThrottle.usf", "ProcessPrevFramePerfDataCS", SF_Compute);
+
+class FUpdateThrottleParametersCS : public FVirtualShadowMapGlobalShader
+{
+	DECLARE_GLOBAL_SHADER(FUpdateThrottleParametersCS);
+	SHADER_USE_PARAMETER_STRUCT(FUpdateThrottleParametersCS, FVirtualShadowMapGlobalShader)
+public:
+	static constexpr int32 NumThreadsPerGroup = 64;
+
+	BEGIN_SHADER_PARAMETER_STRUCT(FParameters, )
+		SHADER_PARAMETER_RDG_BUFFER_SRV(StructuredBuffer<uint>, PrevThrottleBuffer)
+
+		SHADER_PARAMETER_STRUCT_INCLUDE(FThrottlingParameters, ThrottlingParameters)
+
+		SHADER_PARAMETER(uint32, NumThrottleBufferEntries)
+		SHADER_PARAMETER_RDG_BUFFER_UAV(RWBuffer<uint>, OutThrottleBuffer)
+	END_SHADER_PARAMETER_STRUCT()
+};
+IMPLEMENT_GLOBAL_SHADER(FUpdateThrottleParametersCS, "/Engine/Private/VirtualShadowMaps/VirtualShadowMapThrottle.usf", "UpdateThrottleParametersCS", SF_Compute);
+
+BEGIN_SHADER_PARAMETER_STRUCT(FThrottlingShaderParameters, )
+	SHADER_PARAMETER_RDG_BUFFER_SRV(StructuredBuffer<uint>, NanitePerformanceFeedback)
+END_SHADER_PARAMETER_STRUCT()
+
 BEGIN_SHADER_PARAMETER_STRUCT(FBaseGeneratePageFlagsParameters, )
 	SHADER_PARAMETER_RDG_UNIFORM_BUFFER(FVirtualShadowMapUniformParameters, VirtualShadowMap)
 	SHADER_PARAMETER_RDG_UNIFORM_BUFFER(FSceneTextureUniformParameters, SceneTexturesStruct)
 	SHADER_PARAMETER_STRUCT_REF(FViewUniformShaderParameters, View)
 	SHADER_PARAMETER_RDG_UNIFORM_BUFFER(FForwardLightUniformParameters, ForwardLightStruct)
-	SHADER_PARAMETER_RDG_TEXTURE_UAV( RWTexture2D<uint>, OutPageRequestFlags)
-	SHADER_PARAMETER_RDG_TEXTURE_UAV( RWTexture2D<uint>, OutPageReceiverMasks)
-	SHADER_PARAMETER_RDG_BUFFER_SRV(StructuredBuffer< uint >, DirectionalLightIds)
+	SHADER_PARAMETER_RDG_TEXTURE_UAV(RWTexture2D<uint>, OutPageRequestFlags)
+	SHADER_PARAMETER_RDG_TEXTURE_UAV(RWTexture2D<uint>, OutPageReceiverMasks)
+	SHADER_PARAMETER_RDG_BUFFER_SRV(StructuredBuffer<uint>, DirectionalLightIds)
 	SHADER_PARAMETER_RDG_UNIFORM_BUFFER(FSubstrateGlobalUniformParameters, Substrate)
 	SHADER_PARAMETER(float, PageDilationBorderSizeDirectional)
 	SHADER_PARAMETER(float, PageDilationBorderSizeLocal)
@@ -929,6 +1022,8 @@ BEGIN_SHADER_PARAMETER_STRUCT(FBaseGeneratePageFlagsParameters, )
 	SHADER_PARAMETER(uint32, FirstPersonPixelRequestLevelClamp)	
 	SHADER_PARAMETER(float, DOFBiasStrength)
 	SHADER_PARAMETER_STRUCT_INCLUDE(DiaphragmDOF::FDOFCocModelShaderParameters, CocModel)
+	SHADER_PARAMETER_STRUCT_INCLUDE(FThrottlingParameters, ThrottlingParameters)
+	SHADER_PARAMETER_RDG_BUFFER_SRV(StructuredBuffer<uint>, ThrottleBuffer)
 END_SHADER_PARAMETER_STRUCT()
 
 class FGeneratePageFlagsFromPixelsCS : public FVirtualShadowMapPageManagementShader
@@ -939,7 +1034,8 @@ class FGeneratePageFlagsFromPixelsCS : public FVirtualShadowMapPageManagementSha
 	class FInputType : SHADER_PERMUTATION_INT("PERMUTATION_INPUT_TYPE", 2); 
 	class FWaterDepth : SHADER_PERMUTATION_BOOL("PERMUTATION_WATER_DEPTH"); 
 	class FTranslucencyDepth : SHADER_PERMUTATION_BOOL("PERMUTATION_TRANSLUCENCY_DEPTH"); 
-	using FPermutationDomain = TShaderPermutationDomain<FInputType, FWaterDepth, FTranslucencyDepth>;
+	class FThrottling : SHADER_PERMUTATION_BOOL("PERMUTATION_THROTTLING"); 
+	using FPermutationDomain = TShaderPermutationDomain<FInputType, FWaterDepth, FTranslucencyDepth, FThrottling>;
 	
 	static bool ShouldCompilePermutation(const FGlobalShaderPermutationParameters& Parameters) 
 	{ 
@@ -980,7 +1076,8 @@ class FGeneratePageFlagsFromFroxelsCS : public FVirtualShadowMapPageManagementSh
 	SHADER_USE_PARAMETER_STRUCT(FGeneratePageFlagsFromFroxelsCS, FVirtualShadowMapPageManagementShader)
 
 	class FDebugRenderDim : SHADER_PERMUTATION_BOOL("DEBUG_DRAW_GENERATE_FROM_FROXELS"); 
-	using FPermutationDomain = TShaderPermutationDomain<FDebugRenderDim>;
+	class FThrottling : SHADER_PERMUTATION_BOOL("PERMUTATION_THROTTLING"); 
+	using FPermutationDomain = TShaderPermutationDomain<FDebugRenderDim, FThrottling>;
 
 	BEGIN_SHADER_PARAMETER_STRUCT(FParameters, )
 		SHADER_PARAMETER_STRUCT_INCLUDE(FBaseGeneratePageFlagsParameters, Base)
@@ -2050,6 +2147,110 @@ void FVirtualShadowMapArray::BuildPageAllocations(
 		);
 	}
 
+	const FRDGSystemTextures& SystemTextures = FRDGSystemTextures::Get(GraphBuilder);
+
+	bThrottlingEnabled = 
+		CVarTimeBudgetMs.GetValueOnRenderThread() > 0 ||
+		CVarThrottleLoadBudget.GetValueOnRenderThread() > 0;
+
+	if (bThrottlingEnabled)
+	{
+		ThrottleBufferRDG = GraphBuilder.CreateBuffer(FRDGBufferDesc::CreateStructuredDesc(sizeof(uint32), VSM_TB_SIZEOF_HEADER + VSM_TB_SIZEOF_ENTRY * GetNumShadowMapSlots()), TEXT("ThrottleBuffer"));
+		AddClearUAVPass(GraphBuilder, GraphBuilder.CreateUAV(ThrottleBufferRDG), 0);
+	}
+	else
+	{
+		ThrottleBufferRDG = GSystemTextures.GetDefaultStructuredBuffer(GraphBuilder, sizeof(uint32), 0U);
+	}
+
+	FThrottlingParameters ThrottlingParameters {};
+	{
+		if (GDynamicResolutionVSMNaniteBudget.GetSettings().IsEnabled())
+		{
+			ThrottlingParameters.DynResThrottleStrength = FMath::Clamp(1.0f - SceneRenderer.DynamicResolutionFractions[GDynamicResolutionVSMNaniteBudget], 0.0f, 1.0f);
+		}
+		else
+		{
+			ThrottlingParameters.DynResThrottleStrength = -1.0f;
+		}
+
+		float ThrottleLoadBudget = CVarThrottleLoadBudget.GetValueOnRenderThread();
+		ThrottleLoadBudget = (ThrottleLoadBudget <= 0) ? TNumericLimits<float>::Max() : ThrottleLoadBudget;
+		ThrottlingParameters.ThrottleLoadBudget = ThrottleLoadBudget;
+
+		// Memory pressure lod bias is applied first. How much lod bias do we have left for compute-time pressure?
+		const float GlobalMaxBias = CVarVSMDynamicResolutionMaxLodBias.GetValueOnRenderThread();
+		const float DirectionalMaxBias = CVarVSMDynamicResolutionMaxLodBiasDirectional.GetValueOnRenderThread();
+		const float LocalMaxBias = CVarVSMDynamicResolutionMaxLodBiasLocal.GetValueOnRenderThread();
+		
+		float MaxLodBiasDirectional = FMath::Max(0, FMath::Min(GlobalMaxBias, DirectionalMaxBias) 	- CacheManager->GetGlobalResolutionLodBias());
+		float MaxLodBiasLocal 		= FMath::Max(0, FMath::Min(GlobalMaxBias, LocalMaxBias) 		- CacheManager->GetGlobalResolutionLodBias());
+		
+		ThrottlingParameters.ThrottleMaxBiasDirectional = MaxLodBiasDirectional;
+		ThrottlingParameters.ThrottleMaxBiasLocal = MaxLodBiasLocal;
+
+		ThrottlingParameters.ThrottleEMAHistoryWeight = CVarThrottleEMAHistoryWeight.GetValueOnRenderThread();
+	}
+
+	if (bThrottlingEnabled && CacheManager->GetPrevBuffers().NanitePerformanceFeedback)
+	{
+		FProcessPrevFramePerfDataCS::FParameters* PassParameters = GraphBuilder.AllocParameters<FProcessPrevFramePerfDataCS::FParameters>();
+		
+		check(CacheManager->GetPrevBuffers().ThrottleBuffer);
+		FRDGBufferRef PrevThrottleBuffer = GraphBuilder.RegisterExternalBuffer(CacheManager->GetPrevBuffers().ThrottleBuffer);
+		PassParameters->PrevThrottleBuffer = GraphBuilder.CreateSRV(PrevThrottleBuffer);
+
+		FRDGBufferRef PrevNanitePerformanceFeedback = GraphBuilder.RegisterExternalBuffer(CacheManager->GetPrevBuffers().NanitePerformanceFeedback);
+		PassParameters->PrevNanitePerformanceFeedback = GraphBuilder.CreateSRV(PrevNanitePerformanceFeedback);
+
+		FRDGBufferRef NextVirtualShadowMapData = CreateStructuredBuffer(GraphBuilder, TEXT("Shadow.Virtual.NextVirtualShadowMapData"), NextData);
+		PassParameters->NextVirtualShadowMapData = GraphBuilder.CreateSRV(NextVirtualShadowMapData);
+		PassParameters->NextVirtualShadowMapDataCount = NextData.Num();
+
+		PassParameters->ThrottlingParameters = ThrottlingParameters;
+
+		PassParameters->OutThrottleBuffer = GraphBuilder.CreateUAV(ThrottleBufferRDG);
+
+		auto ComputeShader = GetGlobalShaderMap(Scene.GetFeatureLevel())->GetShader<FProcessPrevFramePerfDataCS>();
+		ClearUnusedGraphResources(ComputeShader, PassParameters);
+
+		FComputeShaderUtils::AddPass(
+			GraphBuilder,
+			RDG_EVENT_NAME("ProcessPrevFramePerfData"),
+			ComputeShader,
+			PassParameters,
+			FIntVector(FMath::Max(1, FMath::DivideAndRoundUp(NextData.Num(), FProcessPrevFramePerfDataCS::NumThreadsPerGroup)), 1, 1)
+		);
+	}
+
+	if (bThrottlingEnabled && CacheManager->GetPrevBuffers().ThrottleBuffer)
+	{
+		FUpdateThrottleParametersCS::FParameters* PassParameters = GraphBuilder.AllocParameters<FUpdateThrottleParametersCS::FParameters>();
+		
+		FRDGBufferRef PrevThrottleBuffer = GraphBuilder.RegisterExternalBuffer(CacheManager->GetPrevBuffers().ThrottleBuffer);
+		PassParameters->PrevThrottleBuffer = GraphBuilder.CreateSRV(PrevThrottleBuffer);
+
+		PassParameters->ThrottlingParameters = ThrottlingParameters;
+
+		PassParameters->NumThrottleBufferEntries = GetNumShadowMapSlots();
+		PassParameters->OutThrottleBuffer = GraphBuilder.CreateUAV(ThrottleBufferRDG);
+
+		auto ComputeShader = GetGlobalShaderMap(Scene.GetFeatureLevel())->GetShader<FUpdateThrottleParametersCS>();
+		ClearUnusedGraphResources(ComputeShader, PassParameters);
+
+		FComputeShaderUtils::AddPass(
+			GraphBuilder,
+			RDG_EVENT_NAME("UpdateThrottleParameters"),
+			ComputeShader,
+			PassParameters,
+			FIntVector(FMath::DivideAndRoundUp(GetNumShadowMapSlots(), FUpdateThrottleParametersCS::NumThreadsPerGroup), 1, 1)
+		);
+	}
+
+	uint32 NanitePerformanceFeedbackSize = VSM_NPF_SIZEOF_HEADER + VSM_NPF_SIZEOF_ENTRY * uint32(GetNumShadowMapSlots());
+	NanitePerformanceFeedbackRDG = GraphBuilder.CreateBuffer(FRDGBufferDesc::CreateStructuredDesc(sizeof(uint32), NanitePerformanceFeedbackSize), TEXT("NanitePerformanceFeedback"));
+	AddClearUAVPass(GraphBuilder, GraphBuilder.CreateUAV(NanitePerformanceFeedbackRDG), 0);
+
 	for (int32 ViewIndex = 0; ViewIndex < Views.Num(); ++ViewIndex)
 	{
 		RDG_EVENT_SCOPE_CONDITIONAL(GraphBuilder, Views.Num() > 1, "View%d", ViewIndex);
@@ -2069,8 +2270,6 @@ void FVirtualShadowMapArray::BuildPageAllocations(
 
 		FRDGBufferRef DirectionalLightIdsRDG = CreateStructuredBuffer(GraphBuilder, TEXT("Shadow.Virtual.DirectionalLightIds"), DirectionalLightIds);
 
-		const FRDGSystemTextures& SystemTextures = FRDGSystemTextures::Get(GraphBuilder);
-
 		FRDGBufferRef ScreenSpaceGridBoundsRDG = nullptr;
 			
 		{
@@ -2177,6 +2376,8 @@ void FVirtualShadowMapArray::BuildPageAllocations(
 				BaseParameters.MipModeLocal = CVarMarkPixelPagesMipModeLocal.GetValueOnRenderThread();
 				BaseParameters.FirstPersonPixelRequestBias = CVarFirstPersonPixelRequestBias.GetValueOnRenderThread();
 				BaseParameters.FirstPersonPixelRequestLevelClamp = CVarFirstPersonPixelRequestLevelClamp.GetValueOnRenderThread();
+				BaseParameters.ThrottleBuffer = GraphBuilder.CreateSRV(ThrottleBufferRDG);
+				BaseParameters.ThrottlingParameters = ThrottlingParameters;
 
 				DiaphragmDOF::FPhysicalCocModel CocModel;
 				CocModel.Compile(View);
@@ -2243,6 +2444,7 @@ void FVirtualShadowMapArray::BuildPageAllocations(
 					{
 						FGeneratePageFlagsFromPixelsCS::FPermutationDomain PermutationVector;
 						PermutationVector.Set<FGeneratePageFlagsFromPixelsCS::FInputType>(static_cast<uint32>(InputType));
+						PermutationVector.Set<FGeneratePageFlagsFromPixelsCS::FThrottling>(bThrottlingEnabled);
 						auto ComputeShader = View.ShaderMap->GetShader<FGeneratePageFlagsFromPixelsCS>(PermutationVector);
 
 						check(View.HairStrandsViewData.VisibilityData.TileData.IsValid());
@@ -2261,6 +2463,7 @@ void FVirtualShadowMapArray::BuildPageAllocations(
 						PermutationVector.Set<FGeneratePageFlagsFromPixelsCS::FInputType>(static_cast<uint32>(InputType));
 						PermutationVector.Set<FGeneratePageFlagsFromPixelsCS::FWaterDepth>(bWaterEnabled);
 						PermutationVector.Set<FGeneratePageFlagsFromPixelsCS::FTranslucencyDepth>(bFrontLayerEnabled);
+						PermutationVector.Set<FGeneratePageFlagsFromPixelsCS::FThrottling>(bThrottlingEnabled);
 						auto ComputeShader = View.ShaderMap->GetShader<FGeneratePageFlagsFromPixelsCS>(PermutationVector);
 						FComputeShaderUtils::AddPass(
 							GraphBuilder,
@@ -2295,6 +2498,7 @@ void FVirtualShadowMapArray::BuildPageAllocations(
 							PermutationVector.Set<FGeneratePageFlagsFromFroxelsCS::FDebugRenderDim>(true);
 							ShaderPrint::SetParameters(GraphBuilder, View.ShaderPrintData, PassParameters->ShaderPrintStruct);
 						}
+						PermutationVector.Set<FGeneratePageFlagsFromFroxelsCS::FThrottling>(bThrottlingEnabled);
 
 						auto ComputeShader = View.ShaderMap->GetShader<FGeneratePageFlagsFromFroxelsCS>(PermutationVector);
 						FComputeShaderUtils::AddPass(
@@ -2331,10 +2535,6 @@ void FVirtualShadowMapArray::BuildPageAllocations(
 
 	AddClearPageTableUAVPass(GraphBuilder, PageTableRDG, 0u);
 	AddClearPageTableUAVPass(GraphBuilder, PageFlagsRDG, 0u);
-	
-	int NanitePerformanceFeedbackSize = VSM_NPF_SIZEOF_HEADER + VSM_NPF_SIZEOF_ENTRY * GetNumShadowMapSlots();
-	NanitePerformanceFeedbackRDG = GraphBuilder.CreateBuffer(FRDGBufferDesc::CreateStructuredDesc(sizeof(uint32), NanitePerformanceFeedbackSize), TEXT("NanitePerformanceFeedback"));
-	AddClearUAVPass(GraphBuilder, GraphBuilder.CreateUAV(NanitePerformanceFeedbackRDG), 0);
 
 	// Update cached or newly invalidated pages with respect to the new requests
 	{	
@@ -2668,6 +2868,8 @@ class FVirtualSmLogStatsCS : public FVirtualShadowMapPageManagementShader
 		SHADER_PARAMETER_RDG_BUFFER_SRV(StructuredBuffer<FNaniteStats>, NaniteStats)
 		SHADER_PARAMETER_RDG_BUFFER_SRV(StructuredBuffer<uint>, NanitePerformanceFeedback)
 		SHADER_PARAMETER(uint32, NanitePerformanceFeedbackNumEntries)
+		SHADER_PARAMETER_RDG_BUFFER_SRV(StructuredBuffer<uint>, ThrottleBuffer)
+		SHADER_PARAMETER(uint32, bThrottlingEnabled)
 		SHADER_PARAMETER(int, ShowStatsValue)
 		SHADER_PARAMETER(uint32, StatusMessageId)
 		SHADER_PARAMETER(uint32, StatsMessageId)
@@ -2735,6 +2937,8 @@ void FVirtualShadowMapArray::LogStats(FRDGBuilder& GraphBuilder, const FViewInfo
 		
 		PassParameters->NanitePerformanceFeedback = GraphBuilder.CreateSRV(NanitePerformanceFeedbackRDG);
 		PassParameters->NanitePerformanceFeedbackNumEntries = GetNumShadowMapSlots();
+		PassParameters->ThrottleBuffer = GraphBuilder.CreateSRV(ThrottleBufferRDG);
+		PassParameters->bThrottlingEnabled = bThrottlingEnabled;
 	}
 	
 	auto ComputeShader = View.ShaderMap->GetShader<FVirtualSmLogStatsCS>(PermutationVector);
@@ -3239,6 +3443,7 @@ void FVirtualShadowMapArray::RenderVirtualShadowMapsNanite(
 
 	TRACE_CPUPROFILER_EVENT_SCOPE(FVirtualShadowMapArray::RenderVirtualShadowMapsNanite);
 	RDG_EVENT_SCOPE(GraphBuilder, "RenderVirtualShadowMaps(Nanite)");
+	DynamicRenderScaling::FRDGScope DynamicVSMResolutionScope(GraphBuilder, GDynamicResolutionVSMNaniteBudget); // Only covers nanite for now, because resolution scaling is not very effective for non-nanite.
 
 	const FIntPoint VirtualShadowSize = GetPhysicalPoolSize();
 	const FIntRect VirtualShadowViewRect = FIntRect(0, 0, VirtualShadowSize.X, VirtualShadowSize.Y);
diff --git a/Engine/Source/Runtime/Renderer/Private/VirtualShadowMaps/VirtualShadowMapArray.h b/Engine/Source/Runtime/Renderer/Private/VirtualShadowMaps/VirtualShadowMapArray.h
--- a/Engine/Source/Runtime/Renderer/Private/VirtualShadowMaps/VirtualShadowMapArray.h
+++ b/Engine/Source/Runtime/Renderer/Private/VirtualShadowMaps/VirtualShadowMapArray.h
@@ -462,6 +462,10 @@ public:
 	// Buffer that receives stats from Nanite about the geometry that was rendered into each VSM
 	FRDGBufferRef NanitePerformanceFeedbackRDG = nullptr;
 
+	// Buffer that stores data related to past and current VSM load and performance
+	FRDGBufferRef ThrottleBufferRDG = nullptr;
+	bool bThrottlingEnabled = false;
+
 	// List(s) of physical pages used during allocation/updates
 	// These can be saved frame to frame to allow keeping an LRU-sorted order for cached pages
 	FRDGBufferRef PhysicalPageListsRDG = nullptr;
diff --git a/Engine/Source/Runtime/Renderer/Private/VirtualShadowMaps/VirtualShadowMapCacheManager.cpp b/Engine/Source/Runtime/Renderer/Private/VirtualShadowMaps/VirtualShadowMapCacheManager.cpp
--- a/Engine/Source/Runtime/Renderer/Private/VirtualShadowMaps/VirtualShadowMapCacheManager.cpp
+++ b/Engine/Source/Runtime/Renderer/Private/VirtualShadowMaps/VirtualShadowMapCacheManager.cpp
@@ -107,14 +107,33 @@ static TAutoConsoleVariable<int32> CVarVSMReservedResource(
 	ECVF_RenderThreadSafe
 );
 
-static TAutoConsoleVariable<float> CVarVSMDynamicResolutionMaxLodBias(
+TAutoConsoleVariable<float> CVarVSMDynamicResolutionMaxLodBias(
 	TEXT("r.Shadow.Virtual.DynamicRes.MaxResolutionLodBias"),
 	2.0f,
-	TEXT("As page allocation approaches the pool capacity, VSM resolution ramps down by biasing the LOD up, similar to 'ResolutionLodBiasDirectional'.\n")
+	TEXT("As memory or compute-time cost limits are approached, VSM resolution ramps down by biasing the LOD up, similar to 'ResolutionLodBiasDirectional'.\n")
+	TEXT("Memory pressure is based on page allocation compared to pool capacity.\n")
 	TEXT("This is the maximum LOD bias to clamp to for global dynamic shadow resolution reduction. 0 = disabled"),
 	ECVF_RenderThreadSafe
 );
 
+TAutoConsoleVariable<float> CVarVSMDynamicResolutionMaxLodBiasDirectional(
+	TEXT("r.Shadow.Virtual.DynamicRes.MaxComputeResolutionLodBiasDirectional"),
+	99999.0f,
+	TEXT("As compute-time cost limits are approached, VSM resolution ramps down by biasing the LOD up.\n")
+	TEXT("The maximum LOD bias that is applied is min(this variable, MaxResolutionLodBias).\n")
+	TEXT("This variable applies to directional lights only."),
+	ECVF_RenderThreadSafe
+);
+
+TAutoConsoleVariable<float> CVarVSMDynamicResolutionMaxLodBiasLocal(
+	TEXT("r.Shadow.Virtual.DynamicRes.MaxComputeResolutionLodBiasLocal"),
+	99999.0f,
+	TEXT("As compute-time cost limits are approached, VSM resolution ramps down by biasing the LOD up.\n")
+	TEXT("The maximum LOD bias that is applied is min(this variable, MaxResolutionLodBias).\n")
+	TEXT("This variable applies to local lights only."),
+	ECVF_RenderThreadSafe
+);
+
 static TAutoConsoleVariable<float> CVarVSMDynamicResolutionMaxPagePoolLoadFactor(
 	TEXT("r.Shadow.Virtual.DynamicRes.MaxPagePoolLoadFactor"),
 	0.85f,
@@ -1281,6 +1300,8 @@ void FVirtualShadowMapArrayCacheManager::ExtractFrameData(
 		GraphBuilder.QueueBufferExtraction(VirtualShadowMapArray.UncachedPageRectBoundsRDG, &PrevBuffers.UncachedPageRectBounds);
 		GraphBuilder.QueueBufferExtraction(VirtualShadowMapArray.AllocatedPageRectBoundsRDG, &PrevBuffers.AllocatedPageRectBounds);
 		GraphBuilder.QueueTextureExtraction(VirtualShadowMapArray.PageFlagsRDG, &PrevBuffers.PageFlags);
+		GraphBuilder.QueueBufferExtraction(VirtualShadowMapArray.NanitePerformanceFeedbackRDG, &PrevBuffers.NanitePerformanceFeedback);
+		GraphBuilder.QueueBufferExtraction(VirtualShadowMapArray.ThrottleBufferRDG, &PrevBuffers.ThrottleBuffer);
 		GraphBuilder.QueueTextureExtraction(VirtualShadowMapArray.PageReceiverMasksRDG, &PrevBuffers.PageReceiverMasks);
 
 		if (IsCacheEnabled())
diff --git a/Engine/Source/Runtime/Renderer/Private/VirtualShadowMaps/VirtualShadowMapCacheManager.h b/Engine/Source/Runtime/Renderer/Private/VirtualShadowMaps/VirtualShadowMapCacheManager.h
--- a/Engine/Source/Runtime/Renderer/Private/VirtualShadowMaps/VirtualShadowMapCacheManager.h
+++ b/Engine/Source/Runtime/Renderer/Private/VirtualShadowMaps/VirtualShadowMapCacheManager.h
@@ -233,6 +233,8 @@ struct FVirtualShadowMapArrayFrameData
 	TRefCountPtr<FRDGPooledBuffer>				PhysicalPageLists;
 	TRefCountPtr<IPooledRenderTarget>			PageRequestFlags;
 	TRefCountPtr<IPooledRenderTarget>			PageReceiverMasks;
+	TRefCountPtr<FRDGPooledBuffer>				NanitePerformanceFeedback;
+	TRefCountPtr<FRDGPooledBuffer>				ThrottleBuffer;
 	uint64 GetGPUSizeBytes(bool bLogSizes) const;
 };
 
diff --git a/Engine/Source/Runtime/Renderer/Private/VirtualShadowMaps/VirtualShadowMapShaders.h b/Engine/Source/Runtime/Renderer/Private/VirtualShadowMaps/VirtualShadowMapShaders.h
--- a/Engine/Source/Runtime/Renderer/Private/VirtualShadowMaps/VirtualShadowMapShaders.h
+++ b/Engine/Source/Runtime/Renderer/Private/VirtualShadowMaps/VirtualShadowMapShaders.h
@@ -9,18 +9,14 @@
 #include "DataDrivenShaderPlatformInfo.h"
 #include "GlobalShader.h"
 
-class FVirtualShadowMapPageManagementShader : public FGlobalShader
+class FVirtualShadowMapGlobalShader : public FGlobalShader
 {
 public:
-	// Kernel launch group sizes
-	static constexpr uint32 DefaultCSGroupXY = 8;
-	static constexpr uint32 DefaultCSGroupX = 256;
-
-	FVirtualShadowMapPageManagementShader()
+	FVirtualShadowMapGlobalShader()
 	{
 	}
 
-	FVirtualShadowMapPageManagementShader(const ShaderMetaType::CompiledShaderInitializerType& Initializer)
+	FVirtualShadowMapGlobalShader(const ShaderMetaType::CompiledShaderInitializerType& Initializer)
 		: FGlobalShader(Initializer)
 	{
 	}
@@ -31,22 +27,43 @@ public:
 			DoesPlatformSupportVirtualShadowMaps(Parameters.Platform);
 	}
 
+	static void ModifyCompilationEnvironment(const FGlobalShaderPermutationParameters& Parameters, FShaderCompilerEnvironment& OutEnvironment)
+	{
+		FGlobalShader::ModifyCompilationEnvironment(Parameters, OutEnvironment);
+		FVirtualShadowMapArray::SetShaderDefines(OutEnvironment);
+
+		OutEnvironment.CompilerFlags.Add(CFLAG_HLSL2021);
+	}
+};
+
+class FVirtualShadowMapPageManagementShader : public FVirtualShadowMapGlobalShader
+{
+public:
+	// Kernel launch group sizes
+	static constexpr uint32 DefaultCSGroupXY = 8;
+	static constexpr uint32 DefaultCSGroupX = 256;
+
+	FVirtualShadowMapPageManagementShader()
+	{
+	}
+
+	FVirtualShadowMapPageManagementShader(const ShaderMetaType::CompiledShaderInitializerType& Initializer)
+		: FVirtualShadowMapGlobalShader(Initializer)
+	{
+	}
+
 	/**
 	* Can be overridden by FVertexFactory subclasses to modify their compile environment just before compilation occurs.
 	*/
 	static void ModifyCompilationEnvironment(const FGlobalShaderPermutationParameters& Parameters, FShaderCompilerEnvironment& OutEnvironment)
 	{
-		FGlobalShader::ModifyCompilationEnvironment(Parameters, OutEnvironment);
-
-		FVirtualShadowMapArray::SetShaderDefines(OutEnvironment);
+		FVirtualShadowMapGlobalShader::ModifyCompilationEnvironment(Parameters, OutEnvironment);
 
 		OutEnvironment.SetDefine(TEXT("VSM_DEFAULT_CS_GROUP_X"), DefaultCSGroupX);
 		OutEnvironment.SetDefine(TEXT("VSM_DEFAULT_CS_GROUP_XY"), DefaultCSGroupXY);
 
 		FForwardLightingParameters::ModifyCompilationEnvironment(Parameters.Platform, OutEnvironment);
 		OutEnvironment.SetDefine(TEXT("VF_SUPPORTS_PRIMITIVE_SCENE_DATA"), 1);
-
-		OutEnvironment.CompilerFlags.Add(CFLAG_HLSL2021);
 	}
 };
 
