Fix incorrect shadow mask on hair strands when using VSM and light functions at the same time.

diff --git a/Engine/Source/Runtime/Renderer/Private/LightRendering.cpp b/Engine/Source/Runtime/Renderer/Private/LightRendering.cpp
--- a/Engine/Source/Runtime/Renderer/Private/LightRendering.cpp
+++ b/Engine/Source/Runtime/Renderer/Private/LightRendering.cpp
@@ -2135,12 +2135,18 @@ void FDeferredShadingSceneRenderer::RenderLights(
 				{
 					if (bDrawLightFunction)
 					{
-						const bool bLightFunctionRendered = RenderLightFunction(GraphBuilder, SceneTextures, &LightSceneInfo, ScreenShadowMaskTexture, bDrawShadows, false, false);
-						bUsedShadowMaskTexture |= bLightFunctionRendered;
-
-						if (CVarAppliedLightFunctionOnHair.GetValueOnRenderThread() > 0 && bLightFunctionRendered && ScreenShadowMaskSubPixelTexture)
+						if (bIsHairPass)
+						{
+							if (CVarAppliedLightFunctionOnHair.GetValueOnRenderThread() > 0 && ScreenShadowMaskSubPixelTexture)
+							{
+								const bool bLightFunctionRendered = RenderLightFunction(GraphBuilder, SceneTextures, &LightSceneInfo, ScreenShadowMaskSubPixelTexture, bDrawShadows, false, true);
+								bUsedShadowMaskTexture |= bLightFunctionRendered;
+							}
+						}
+						else
 						{
-							RenderLightFunction(GraphBuilder, SceneTextures, &LightSceneInfo, ScreenShadowMaskSubPixelTexture, bDrawShadows, false, true);
+							const bool bLightFunctionRendered = RenderLightFunction(GraphBuilder, SceneTextures, &LightSceneInfo, ScreenShadowMaskTexture, bDrawShadows, false, false);
+							bUsedShadowMaskTexture |= bLightFunctionRendered;
 						}
 					}
 
diff --git a/Engine/Source/Runtime/Renderer/Private/Shadows/ShadowSceneRenderer.cpp b/Engine/Source/Runtime/Renderer/Private/Shadows/ShadowSceneRenderer.cpp
--- a/Engine/Source/Runtime/Renderer/Private/Shadows/ShadowSceneRenderer.cpp
+++ b/Engine/Source/Runtime/Renderer/Private/Shadows/ShadowSceneRenderer.cpp
@@ -1012,7 +1012,7 @@ void FShadowSceneRenderer::ApplyVirtualShadowMapProjectionForLight(
 						VirtualShadowMapArray,
 						InputType,
 						LocalLightSetup.VirtualShadowMapId,
-						VirtualShadowMapMaskBits,
+						InputType == EVirtualShadowMapProjectionInputType::HairStrands ? VirtualShadowMapMaskBitsHairStrands : VirtualShadowMapMaskBits,
 						OutputScreenShadowMaskTexture);
 				}
 				else
