PostProcess:  UserSceneTexture fixes for mobile renderer.  Shader compile error fixes and debug display rendering.

diff --git a/Engine/Shaders/Private/MaterialTemplate.ush b/Engine/Shaders/Private/MaterialTemplate.ush
--- a/Engine/Shaders/Private/MaterialTemplate.ush
+++ b/Engine/Shaders/Private/MaterialTemplate.ush
@@ -3342,6 +3342,10 @@ float4 SceneTextureLookup(float2 UV, int SceneTextureIndex, bool bFiltered)
 	return float4(0, 0, 0, 0);
 }
 
+#define SceneTextureFetch(SceneTextureIndex, PixelOffset) SceneTextureFetchFunc(Parameters, SceneTextureIndex, PixelOffset)
+
+#endif // SHADING_PATH_DEFERRED
+
 /**
  * Utility function and macro to streamline Custom HLSL fetching from SceneTextures or UserSceneTextures with a pixel offset from default UVs, including
  * multiplication by texture size and clamping.  Example that fetches at offset 1,1 from the default sample:
@@ -3354,18 +3358,18 @@ float4 SceneTextureFetchFunc(FMaterialPixelParameters Parameters, int SceneTextu
 	return float4(0,0,0,0);
 #else
 	float2 UV = ClampSceneTextureUV(GetDefaultSceneTextureUV(Parameters, SceneTextureIndex) + GetSceneTextureViewSize(SceneTextureIndex).zw * PixelOffset, SceneTextureIndex);
+#if SHADING_PATH_MOBILE
+	return MobileSceneTextureLookup(Parameters, SceneTextureIndex, UV);
+#else
 	return SceneTextureLookup(UV, SceneTextureIndex, true);
 #endif
+#endif
 }
 
 float4 SceneTextureFetchFunc(FMaterialPixelParameters Parameters, int SceneTextureIndex, float PixelOffsetX, float PixelOffsetY)
 {
 	return SceneTextureFetchFunc(Parameters, SceneTextureIndex, float2(PixelOffsetX, PixelOffsetY));
 }
-
-#define SceneTextureFetch(SceneTextureIndex, PixelOffset) SceneTextureFetchFunc(Parameters, SceneTextureIndex, PixelOffset)
-
-#endif // SHADING_PATH_DEFERRED
 #endif // NEEDS_SCENE_TEXTURES
 
 #if SHADING_PATH_DEFERRED
diff --git a/Engine/Source/Runtime/Renderer/Private/PostProcess/PostProcessing.cpp b/Engine/Source/Runtime/Renderer/Private/PostProcess/PostProcessing.cpp
--- a/Engine/Source/Runtime/Renderer/Private/PostProcess/PostProcessing.cpp
+++ b/Engine/Source/Runtime/Renderer/Private/PostProcess/PostProcessing.cpp
@@ -2028,59 +2028,8 @@ void AddPostProcessingPasses(
 	}
 	#endif
 
-	#if !UE_BUILD_SHIPPING
-	{
-		int32 UserSceneTextureDebug = CVarUserSceneTextureDebug.GetValueOnRenderThread();
-		bool bEnableUserSceneTextureDebug = false;
-
-		if (UserSceneTextureDebug == 1)
-		{
-			// Enable always
-			bEnableUserSceneTextureDebug = true;
-		}
-		else if (UserSceneTextureDebug == 2 && GAreScreenMessagesEnabled)
-		{
-			// Enable conditionally if there are errors
-			const FSceneTextures& SceneTextures = View.GetSceneTextures();
-			for (const FUserSceneTextureEventData& EventData : SceneTextures.UserSceneTextureEvents)
-			{
-				if (EventData.Event == EUserSceneTextureEvent::MissingInput || EventData.Event == EUserSceneTextureEvent::CollidingInput)
-				{
-					bEnableUserSceneTextureDebug = true;
-					break;
-				}
-			}
-
-			for (auto& UserSceneTextureElement : SceneTextures.UserSceneTextures)
-			{
-				if (bEnableUserSceneTextureDebug)
-				{
-					break;
-				}
-				for (FTransientUserSceneTexture& UserSceneTexture : UserSceneTextureElement.Value)
-				{
-					if (!UserSceneTexture.bUsed)
-					{
-						bEnableUserSceneTextureDebug = true;
-						break;
-					}
-				}
-			}
-		}
-		else if (UserSceneTextureDebug == 3)
-		{
-			// Enable conditionally for view with texture being visualized
-			if (GVisualizeTexture.IsRequestedView())
-			{
-				bEnableUserSceneTextureDebug = true;
-			}
-		}
-
-		if (bEnableUserSceneTextureDebug)
-		{
-			AddUserSceneTextureDebugPass(GraphBuilder, View, ViewIndex, SceneColor);
-		}
-	}
+#if !UE_BUILD_SHIPPING
+	AddUserSceneTextureDebugPass(GraphBuilder, View, ViewIndex, SceneColor);
 #endif
 
 	if (PassSequence.IsEnabled(EPass::AlphaInvert))
@@ -3408,6 +3357,10 @@ void AddMobilePostProcessingPasses(FRDGBuilder& GraphBuilder, FScene* Scene, con
 		SceneColor = AddHMDDistortionPass(GraphBuilder, View, PassInputs);
 	}
 
+#if !UE_BUILD_SHIPPING
+	AddUserSceneTextureDebugPass(GraphBuilder, View, ViewIndex, SceneColor);
+#endif
+
 	// Copy the scene color to back buffer in case there is no post process, such as LDR MSAA.
 	if (SceneColor.Texture != ViewFamilyOutput.Texture)
 	{
@@ -3543,6 +3496,57 @@ static float CanvasDrawShadowedStringReturnWidth(FCanvas& Canvas, float PrintX,
 
 static void AddUserSceneTextureDebugPass(FRDGBuilder& GraphBuilder, const FViewInfo& View, int32 ViewIndex, FScreenPassTexture Output)
 {
+	int32 UserSceneTextureDebug = CVarUserSceneTextureDebug.GetValueOnRenderThread();
+	bool bEnableUserSceneTextureDebug = false;
+
+	if (UserSceneTextureDebug == 1)
+	{
+		// Enable always
+		bEnableUserSceneTextureDebug = true;
+	}
+	else if (UserSceneTextureDebug == 2 && GAreScreenMessagesEnabled)
+	{
+		// Enable conditionally if there are errors
+		const FSceneTextures& SceneTextures = View.GetSceneTextures();
+		for (const FUserSceneTextureEventData& EventData : SceneTextures.UserSceneTextureEvents)
+		{
+			if (EventData.Event == EUserSceneTextureEvent::MissingInput || EventData.Event == EUserSceneTextureEvent::CollidingInput)
+			{
+				bEnableUserSceneTextureDebug = true;
+				break;
+			}
+		}
+
+		for (auto& UserSceneTextureElement : SceneTextures.UserSceneTextures)
+		{
+			if (bEnableUserSceneTextureDebug)
+			{
+				break;
+			}
+			for (FTransientUserSceneTexture& UserSceneTexture : UserSceneTextureElement.Value)
+			{
+				if (!UserSceneTexture.bUsed)
+				{
+					bEnableUserSceneTextureDebug = true;
+					break;
+				}
+			}
+		}
+	}
+	else if (UserSceneTextureDebug == 3)
+	{
+		// Enable conditionally for view with texture being visualized
+		if (GVisualizeTexture.IsRequestedView())
+		{
+			bEnableUserSceneTextureDebug = true;
+		}
+	}
+
+	if (!bEnableUserSceneTextureDebug)
+	{
+		return;
+	}
+
 	const FSceneTextures& SceneTextures = View.GetSceneTextures();
 	if (!SceneTextures.UserSceneTextureEvents.IsEmpty())
 	{
