SLW - fixed water removing cloud due to writing depth = near clip plane in SingleLayerWaterRefractionCullingPS.

diff --git a/Engine/Shaders/Private/SingleLayerWaterRefractionCulling.usf b/Engine/Shaders/Private/SingleLayerWaterRefractionCulling.usf
--- a/Engine/Shaders/Private/SingleLayerWaterRefractionCulling.usf
+++ b/Engine/Shaders/Private/SingleLayerWaterRefractionCulling.usf
@@ -67,6 +67,7 @@ void SingleLayerWaterRefractionMaskPS(
 #ifdef SingleLayerWaterRefractionCullingPS
 
 Texture2D<float> WaterRefractionCullingTexture;
+Texture2D<float> WaterDepthTexture;
 
 void SingleLayerWaterRefractionCullingPS(
 	in float4 SvPosition : SV_Position,
@@ -85,7 +86,9 @@ void SingleLayerWaterRefractionCullingPS(
 		discard;
 	}
 
-	OutDepth = NearDepthValue;
+	// Always a water pixel at this stage
+	const float WaterDeviceZ = WaterDepthTexture.Load(int3(SvPosition.xy, 0)).x;
+	OutDepth = WaterDeviceZ;
 
 #if SHADING_MASK_EXPORT
 	OutShadingMask = 0;
diff --git a/Engine/Source/Runtime/Renderer/Private/SingleLayerWaterRendering.cpp b/Engine/Source/Runtime/Renderer/Private/SingleLayerWaterRendering.cpp
--- a/Engine/Source/Runtime/Renderer/Private/SingleLayerWaterRendering.cpp
+++ b/Engine/Source/Runtime/Renderer/Private/SingleLayerWaterRendering.cpp
@@ -448,6 +448,7 @@ class FSingleLayerWaterRefractionCullingPS : public FGlobalShader
 	BEGIN_SHADER_PARAMETER_STRUCT(FParameters, )
 		SHADER_PARAMETER_STRUCT_REF(FViewUniformShaderParameters, View)
 		SHADER_PARAMETER_RDG_TEXTURE(Texture2D, WaterRefractionCullingTexture)
+		SHADER_PARAMETER_RDG_TEXTURE(Texture2D, WaterDepthTexture)
 		RENDER_TARGET_BINDING_SLOTS()
 	END_SHADER_PARAMETER_STRUCT()
 
@@ -1129,6 +1130,7 @@ FSingleLayerWaterPrePassResult* FDeferredShadingSceneRenderer::RenderSingleLayer
 			FSingleLayerWaterRefractionCullingPS::FParameters* PassParameters = GraphBuilder.AllocParameters<FSingleLayerWaterRefractionCullingPS::FParameters>();
 			PassParameters->View = View.ViewUniformBuffer;
 			PassParameters->WaterRefractionCullingTexture = Result->RefractionMaskTexture;
+			PassParameters->WaterDepthTexture = Result->DepthPrepassTexture.Resolve;
 			PassParameters->RenderTargets.DepthStencil = FDepthStencilBinding(SceneTextures.Depth.Target, ERenderTargetLoadAction::ELoad, ERenderTargetLoadAction::ELoad, FExclusiveDepthStencil::DepthWrite_StencilNop);
 			if (bNaniteShadingMaskExport)
 			{
