VolumetricCloud - added sample count clamp cvar for fine tuning performance.

diff --git a/Engine/Shaders/Private/SkyAtmosphere.usf b/Engine/Shaders/Private/SkyAtmosphere.usf
--- a/Engine/Shaders/Private/SkyAtmosphere.usf
+++ b/Engine/Shaders/Private/SkyAtmosphere.usf
@@ -5,6 +5,10 @@
 =============================================================================*/
 
 
+// Change this to force recompilation of all volumetric cloud material shaders
+#pragma message("UESHADERMETADATA_VERSION 046409B1-EC51-4A31-AD82-2E8D54C0DA6A")
+
+
 #include "Common.ush"
 #include "SceneTexturesCommon.ush"
 #include "/Engine/Shared/EnvironmentComponentsFlags.h"
diff --git a/Engine/Shaders/Private/VolumetricCloud.usf b/Engine/Shaders/Private/VolumetricCloud.usf
--- a/Engine/Shaders/Private/VolumetricCloud.usf
+++ b/Engine/Shaders/Private/VolumetricCloud.usf
@@ -6,7 +6,7 @@
 
 
  // Change this to force recompilation of all volumetric cloud material shaders
-#pragma message("UESHADERMETADATA_VERSION 89190FC1-44B2-474A-8CB6-05332B45FE36")
+#pragma message("UESHADERMETADATA_VERSION CCE8E085-EEC0-48B0-A7FA-910D86534098")
 
 #include "Common.ush"
 #include "ColorMap.ush"
@@ -713,11 +713,16 @@ void MainCommon(in FMaterialPixelParameters MaterialParameters, in float4 SvPosi
 		TMax = max(TMin, TMax);
 
 		// When we do trace, we want a minimum amount of sample count. This avoid situation with 0 sample count due to InvDistanceToSampleCountMax.
-		const uint  IStepCount = max(RenderVolumetricCloudParameters.SampleCountMin, RenderVolumetricCloudParameters.SampleCountMax * saturate((TMax - TMin) * RenderVolumetricCloudParameters.InvDistanceToSampleCountMax));
-		const float StepCount = float(IStepCount);
+		const uint  EvaluatedIStepCount = max(RenderVolumetricCloudParameters.SampleCountMin, RenderVolumetricCloudParameters.SampleCountMax * saturate((TMax - TMin) * RenderVolumetricCloudParameters.InvDistanceToSampleCountMax));
+		const float StepCount = float(EvaluatedIStepCount);
 		const float StepT = (TMax - TMin) / StepCount; // StepT is step distance in centimeters
 		const float dtMeters = StepT * CENTIMETER_TO_METER;
 
+		// After the sample count for a trace has been evaluated, we can still clamp the sample count.
+		// This is useful to avoid long running GPU waves which can increase the cost, especially when started later on the async compute pipe.
+		// It should not be set too low otherwise the visual result will look incomplete.
+		const uint  IStepCount = min(EvaluatedIStepCount, RenderVolumetricCloudParameters.SampleCountClamp);
+
 		// This is the scattered sky light luminance assuming a uniform phase function.
 		// TODO Approximate some form of occlusion
 		// TODO have some directionality: use SH from sky atmosphere? Or GetSkySHDiffuse/GetSkySHDiffuseSimple for dynamic SkyLightCapture.
@@ -1447,7 +1452,7 @@ void MainCommon(in FMaterialPixelParameters MaterialParameters, in float4 SvPosi
 			}
 
 			t += StepT;
-			}
+		}
 
 #if CLOUD_DEBUG_VIEW_MODE
 		if (RenderVolumetricCloudParameters.CloudDebugViewMode >= 1)
diff --git a/Engine/Source/Runtime/Renderer/Private/VolumetricCloudRendering.cpp b/Engine/Source/Runtime/Renderer/Private/VolumetricCloudRendering.cpp
--- a/Engine/Source/Runtime/Renderer/Private/VolumetricCloudRendering.cpp
+++ b/Engine/Source/Runtime/Renderer/Private/VolumetricCloudRendering.cpp
@@ -58,6 +58,11 @@ static TAutoConsoleVariable<int32> CVarVolumetricCloudSampleMinCount(
 	TEXT("The minimum number of samples to take along a ray. This can help with quality for volume close to the camera, e.g. if cloud layer is also used as low altitude fog. SampleMinCount should remain relatively small because it is applied to all tracing process."),
 	ECVF_RenderThreadSafe | ECVF_Scalability);
 
+static TAutoConsoleVariable<int32> CVarVolumetricCloudSampleClampCount(
+	TEXT("r.VolumetricCloud.SampleClampCount"), 2147483647,
+	TEXT("After the sample count for a trace has been evaluated, we can still clamp the sample count. This is useful to avoid long running GPU waves which can increase the cost, especially when started later on the async compute pipe. It should not be set too low otherwise the visual result will look incomplete. It is tied to all the sample count cvar, so it needs to be setup for all quality levels."),
+	ECVF_RenderThreadSafe | ECVF_Scalability);
+
 static TAutoConsoleVariable<float> CVarVolumetricCloudViewRaySampleMaxCount(
 	TEXT("r.VolumetricCloud.ViewRaySampleMaxCount"), 768,
 	TEXT("The maximum number of samples taken while ray marching view primary rays."),
@@ -1690,7 +1695,8 @@ void FSceneRenderer::InitVolumetricCloudsForViews(FRDGBuilder& GraphBuilder, boo
 			CloudGlobalShaderParams.TracingStartDistanceFromCamera		= FMath::Max(0.0f, KilometersToCentimeters * CloudProxy.TracingStartDistanceFromCamera);
 			CloudGlobalShaderParams.TracingMaxDistance					= FMath::Max(0.0f, KilometersToCentimeters * CloudProxy.TracingMaxDistance);
 			CloudGlobalShaderParams.TracingMaxDistanceMode				= CloudProxy.TracingMaxDistanceMode;
-
+			
+			CloudGlobalShaderParams.SampleCountClamp					= FMath::Max(0, CVarVolumetricCloudSampleClampCount.GetValueOnAnyThread());
 			CloudGlobalShaderParams.SampleCountMin						= FMath::Max(0, CVarVolumetricCloudSampleMinCount.GetValueOnAnyThread());
 			CloudGlobalShaderParams.SampleCountMax						= FMath::Max(2.0f, FMath::Min(UVolumetricCloudComponent::BaseViewRaySampleCount   * CloudProxy.ViewSampleCountScale,       CVarVolumetricCloudViewRaySampleMaxCount.GetValueOnAnyThread()));
 			CloudGlobalShaderParams.ShadowSampleCountMax				= FMath::Max(2.0f, FMath::Min(UVolumetricCloudComponent::BaseShadowRaySampleCount * CloudProxy.ShadowViewSampleCountScale, CVarVolumetricCloudShadowViewRaySampleMaxCount.GetValueOnAnyThread()));
diff --git a/Engine/Source/Runtime/Renderer/Private/VolumetricCloudRendering.h b/Engine/Source/Runtime/Renderer/Private/VolumetricCloudRendering.h
--- a/Engine/Source/Runtime/Renderer/Private/VolumetricCloudRendering.h
+++ b/Engine/Source/Runtime/Renderer/Private/VolumetricCloudRendering.h
@@ -34,6 +34,7 @@ BEGIN_SHADER_PARAMETER_STRUCT(FVolumetricCloudCommonShaderParameters, )
 	SHADER_PARAMETER(float, TracingStartMaxDistance)
 	SHADER_PARAMETER(int32, TracingMaxDistanceMode)
 	SHADER_PARAMETER(float, TracingMaxDistance)
+	SHADER_PARAMETER(int32, SampleCountClamp)
 	SHADER_PARAMETER(int32, SampleCountMin)
 	SHADER_PARAMETER(int32, SampleCountMax)
 	SHADER_PARAMETER(float, InvDistanceToSampleCountMax)
