VolumetricCloud - fixed compilation warning.

diff --git a/Engine/Shaders/Private/SkyAtmosphere.usf b/Engine/Shaders/Private/SkyAtmosphere.usf
--- a/Engine/Shaders/Private/SkyAtmosphere.usf
+++ b/Engine/Shaders/Private/SkyAtmosphere.usf
@@ -6,7 +6,7 @@
 
 
 // Change this to force recompilation of all volumetric cloud material shaders
-#pragma message("UESHADERMETADATA_VERSION 046409B1-EC51-4A31-AD82-2E8D54C0DA6A")
+#pragma message("UESHADERMETADATA_VERSION 849A6B91-7442-4D77-AA37-77C896193CD5")
 
 
 #include "Common.ush"
diff --git a/Engine/Shaders/Private/VolumetricCloud.usf b/Engine/Shaders/Private/VolumetricCloud.usf
--- a/Engine/Shaders/Private/VolumetricCloud.usf
+++ b/Engine/Shaders/Private/VolumetricCloud.usf
@@ -6,7 +6,7 @@
 
 
  // Change this to force recompilation of all volumetric cloud material shaders
-#pragma message("UESHADERMETADATA_VERSION CCE8E085-EEC0-48B0-A7FA-910D86534098")
+#pragma message("UESHADERMETADATA_VERSION 26C4467C-50E5-414D-86B5-3F3C27A85891")
 
 #include "Common.ush"
 #include "ColorMap.ush"
diff --git a/Engine/Source/Runtime/Renderer/Private/VolumetricCloudRendering.cpp b/Engine/Source/Runtime/Renderer/Private/VolumetricCloudRendering.cpp
--- a/Engine/Source/Runtime/Renderer/Private/VolumetricCloudRendering.cpp
+++ b/Engine/Source/Runtime/Renderer/Private/VolumetricCloudRendering.cpp
@@ -1696,7 +1696,7 @@ void FSceneRenderer::InitVolumetricCloudsForViews(FRDGBuilder& GraphBuilder, boo
 			CloudGlobalShaderParams.TracingMaxDistance					= FMath::Max(0.0f, KilometersToCentimeters * CloudProxy.TracingMaxDistance);
 			CloudGlobalShaderParams.TracingMaxDistanceMode				= CloudProxy.TracingMaxDistanceMode;
 			
-			CloudGlobalShaderParams.SampleCountClamp					= FMath::Max(0, CVarVolumetricCloudSampleClampCount.GetValueOnAnyThread());
+			CloudGlobalShaderParams.SampleCountClamp					= (uint32) FMath::Max(0, CVarVolumetricCloudSampleClampCount.GetValueOnAnyThread());
 			CloudGlobalShaderParams.SampleCountMin						= FMath::Max(0, CVarVolumetricCloudSampleMinCount.GetValueOnAnyThread());
 			CloudGlobalShaderParams.SampleCountMax						= FMath::Max(2.0f, FMath::Min(UVolumetricCloudComponent::BaseViewRaySampleCount   * CloudProxy.ViewSampleCountScale,       CVarVolumetricCloudViewRaySampleMaxCount.GetValueOnAnyThread()));
 			CloudGlobalShaderParams.ShadowSampleCountMax				= FMath::Max(2.0f, FMath::Min(UVolumetricCloudComponent::BaseShadowRaySampleCount * CloudProxy.ShadowViewSampleCountScale, CVarVolumetricCloudShadowViewRaySampleMaxCount.GetValueOnAnyThread()));
diff --git a/Engine/Source/Runtime/Renderer/Private/VolumetricCloudRendering.h b/Engine/Source/Runtime/Renderer/Private/VolumetricCloudRendering.h
--- a/Engine/Source/Runtime/Renderer/Private/VolumetricCloudRendering.h
+++ b/Engine/Source/Runtime/Renderer/Private/VolumetricCloudRendering.h
@@ -34,7 +34,7 @@ BEGIN_SHADER_PARAMETER_STRUCT(FVolumetricCloudCommonShaderParameters, )
 	SHADER_PARAMETER(float, TracingStartMaxDistance)
 	SHADER_PARAMETER(int32, TracingMaxDistanceMode)
 	SHADER_PARAMETER(float, TracingMaxDistance)
-	SHADER_PARAMETER(int32, SampleCountClamp)
+	SHADER_PARAMETER(uint32, SampleCountClamp)
 	SHADER_PARAMETER(int32, SampleCountMin)
 	SHADER_PARAMETER(int32, SampleCountMax)
 	SHADER_PARAMETER(float, InvDistanceToSampleCountMax)
