Fixed: Incorrect filtering out of shading models in an editor builds

diff --git a/Engine/Source/Runtime/Engine/Public/Materials/MaterialInterface.h b/Engine/Source/Runtime/Engine/Public/Materials/MaterialInterface.h
--- a/Engine/Source/Runtime/Engine/Public/Materials/MaterialInterface.h
+++ b/Engine/Source/Runtime/Engine/Public/Materials/MaterialInterface.h
@@ -1259,7 +1259,7 @@ public:
 #endif // WITH_EDITOR
 
 	/** Filter out ShadingModels field to a shader platform settings */
-	static void FilterOutPlatformShadingModels(EShaderPlatform Platform, FMaterialShadingModelField& ShadingModels);
+	ENGINE_API static void FilterOutPlatformShadingModels(EShaderPlatform Platform, FMaterialShadingModelField& ShadingModels);
 	
 #if WITH_EDITOR
 	TUniquePtr<FMaterialInsights> MaterialInsight;
diff --git a/Engine/Source/Runtime/Renderer/Private/MobileBasePass.cpp b/Engine/Source/Runtime/Renderer/Private/MobileBasePass.cpp
--- a/Engine/Source/Runtime/Renderer/Private/MobileBasePass.cpp
+++ b/Engine/Source/Runtime/Renderer/Private/MobileBasePass.cpp
@@ -759,7 +759,12 @@ bool FMobileBasePassMeshProcessor::TryAddMeshBatch(const FMeshBatch& RESTRICT Me
 {
 	if (ShouldDraw(Material))
 	{
-		const FMaterialShadingModelField ShadingModels = Material.GetShadingModels();
+		FMaterialShadingModelField ShadingModels = Material.GetShadingModels();
+#if WITH_EDITOR
+		// Non-Editor builds filter out shading models on a material load
+		const EShaderPlatform ShaderPlatform = GetFeatureLevelShaderPlatform(FeatureLevel);
+		UMaterialInterface::FilterOutPlatformShadingModels(ShaderPlatform, ShadingModels);
+#endif
 		const bool bSingleLayerWater = ShadingModels.HasShadingModel(MSM_SingleLayerWater);
 		const bool bCanReceiveCSM = bSingleLayerWater || ((Flags & EFlags::CanReceiveCSM) == EFlags::CanReceiveCSM);
 		const EBlendMode BlendMode = Material.GetBlendMode();
