[VSM] Add unrolled shadow projection permutations for 2 and 4 samples per ray and remove that for 0 samples to optimize the common use cases.

#tests VSM regression testing - should only affect performance.

diff --git a/Engine/Shaders/Private/VirtualShadowMaps/VirtualShadowMapSMRTTemplate.ush b/Engine/Shaders/Private/VirtualShadowMaps/VirtualShadowMapSMRTTemplate.ush
--- a/Engine/Shaders/Private/VirtualShadowMaps/VirtualShadowMapSMRTTemplate.ush
+++ b/Engine/Shaders/Private/VirtualShadowMaps/VirtualShadowMapSMRTTemplate.ush
@@ -48,7 +48,7 @@ FSMRTResult SMRTRayCast(
 	bool bValidHit =  false;
 #if SMRT_TEMPLATE_STATIC_SAMPLES_PER_RAY >= 0
 	NumSteps = SMRT_TEMPLATE_STATIC_SAMPLES_PER_RAY;
-	[unroll]
+	UNROLL
 #endif
 	for (int i = 0; i <= NumSteps; i++)
 	{
diff --git a/Engine/Source/Runtime/RenderCore/Public/ShaderPermutationUtils.h b/Engine/Source/Runtime/RenderCore/Public/ShaderPermutationUtils.h
--- a/Engine/Source/Runtime/RenderCore/Public/ShaderPermutationUtils.h
+++ b/Engine/Source/Runtime/RenderCore/Public/ShaderPermutationUtils.h
@@ -121,4 +121,17 @@ namespace UE::ShaderPermutationUtils
 			}
 		}
 	}
+
+	template <typename DimensionType>
+	bool DoesDimensionContainValue(const typename DimensionType::Type& Value)
+	{
+		for (int32 PermutationId = 0; PermutationId < DimensionType::PermutationCount; ++PermutationId)
+		{
+			if (Value == DimensionType::FromDimensionValueId(PermutationId))
+			{
+				return true;
+			}
+		}
+		return false;
+	}
 }
diff --git a/Engine/Source/Runtime/Renderer/Private/VirtualShadowMaps/VirtualShadowMapProjection.cpp b/Engine/Source/Runtime/Renderer/Private/VirtualShadowMaps/VirtualShadowMapProjection.cpp
--- a/Engine/Source/Runtime/Renderer/Private/VirtualShadowMaps/VirtualShadowMapProjection.cpp
+++ b/Engine/Source/Runtime/Renderer/Private/VirtualShadowMaps/VirtualShadowMapProjection.cpp
@@ -27,6 +27,7 @@
 #include "HairStrands/HairStrandsData.h"
 #include "BasePassRendering.h"
 #include "BlueNoise.h"
+#include "ShaderPermutationUtils.h"
 
 #define MAX_TEST_PERMUTATION 0
 
@@ -105,7 +106,7 @@ class FVirtualShadowMapProjectionCS : public FGlobalShader
 	class FUseTileList				: SHADER_PERMUTATION_BOOL("USE_TILE_LIST");
 	class FFirstPersonShadow		: SHADER_PERMUTATION_BOOL("FIRST_PERSON_SHADOW");
 	// -1 means dynamic count
-	class FSMRTStaticSampleCount	: SHADER_PERMUTATION_RANGE_INT("SMRT_TEMPLATE_STATIC_SAMPLES_PER_RAY", -1, 2);
+	class FSMRTStaticSampleCount	: SHADER_PERMUTATION_SPARSE_INT("SMRT_TEMPLATE_STATIC_SAMPLES_PER_RAY", -1, 2, 4);
 	// Used for A/B testing a change that affects reg allocation, etc.
 	class FTestDim					: SHADER_PERMUTATION_INT("TEST_PERMUTATION", MAX_TEST_PERMUTATION+1);
 
@@ -192,6 +193,12 @@ class FVirtualShadowMapProjectionCS : public FGlobalShader
 		{
 			return false;
 		}
+		
+		// FOnePassProjectionDim is not used together with the tile list
+		if( PermutationVector.Get< FUseTileList >() && PermutationVector.Get< FOnePassProjectionDim >() )
+		{
+			return false;
+		}
 
 		return DoesPlatformSupportVirtualShadowMaps(Parameters.Platform);
 	}
@@ -259,6 +266,9 @@ static void RenderVirtualShadowMapProjectionCommon(
 	bool bOnePassProjection = LightProxy == nullptr;
 	if (bOnePassProjection)
 	{
+		// These are mutually exclusive in practice, enforce to cut down on permutations
+		check(!bUseTileList);
+
 		// One pass projection
 		PassParameters->ForwardLightStruct = View.ForwardLightingResources.ForwardLightUniformBuffer;
 		PassParameters->OutShadowMaskBits = GraphBuilder.CreateUAV( OutputTexture );
@@ -294,7 +304,7 @@ static void RenderVirtualShadowMapProjectionCommon(
 	// If the requested samples per ray matches one of our static permutations, pick that one
 	// Otherwise use the dynamic samples per ray permutation (-1).
 	int SamplesPerRay = (bDirectionalLight ? VirtualShadowMapArray.UniformParameters.SMRTSamplesPerRayDirectional : VirtualShadowMapArray.UniformParameters.SMRTSamplesPerRayLocal);
-	int StaticSamplesPerRay = SamplesPerRay	== 0 ? SamplesPerRay : -1;
+	int StaticSamplesPerRay = UE::ShaderPermutationUtils::DoesDimensionContainValue<FVirtualShadowMapProjectionCS::FSMRTStaticSampleCount>(SamplesPerRay) ? SamplesPerRay : -1;
 	float ExtrapolateMaxSlope = (bDirectionalLight ? VirtualShadowMapArray.UniformParameters.SMRTExtrapolateMaxSlopeDirectional : VirtualShadowMapArray.UniformParameters.SMRTExtrapolateMaxSlopeLocal);
 
 	FVirtualShadowMapProjectionCS::FPermutationDomain PermutationVector;
