Only track global distance field coverage if platform supports volume texture atomics.

diff --git a/Engine/Shaders/Private/DistanceField/GlobalDistanceFieldCompositeObjects.usf b/Engine/Shaders/Private/DistanceField/GlobalDistanceFieldCompositeObjects.usf
--- a/Engine/Shaders/Private/DistanceField/GlobalDistanceFieldCompositeObjects.usf
+++ b/Engine/Shaders/Private/DistanceField/GlobalDistanceFieldCompositeObjects.usf
@@ -245,6 +245,7 @@ void CompositeObjectsIntoPagesCS(
 
 	}
 
+#if COMPOSITE_COVERAGE_ATLAS
 	GroupMemoryBarrierWithGroupSync();
 
 	// Mark page containing coverage values below 1
@@ -252,6 +253,7 @@ void CompositeObjectsIntoPagesCS(
 	{
 		InterlockedOr(RWPageTableCombinedTexture[PageTableTextureCoord], SharedCompositeMask);
 	}
+#endif
 }
 
 RWStructuredBuffer<uint4> RWPageObjectGridBuffer;
diff --git a/Engine/Source/Runtime/Renderer/Private/GlobalDistanceField.cpp b/Engine/Source/Runtime/Renderer/Private/GlobalDistanceField.cpp
--- a/Engine/Source/Runtime/Renderer/Private/GlobalDistanceField.cpp
+++ b/Engine/Source/Runtime/Renderer/Private/GlobalDistanceField.cpp
@@ -304,6 +304,11 @@ bool UseGlobalDistanceField(const FDistanceFieldAOParameters& Parameters)
 	return UseGlobalDistanceField() && Parameters.GlobalMaxOcclusionDistance > 0;
 }
 
+static bool SupportsGlobalDistanceFieldCoverage(EShaderPlatform Platform)
+{
+	return RHISupportsVolumeTextureAtomics(Platform);
+}
+
 FGlobalDistanceFieldParameters2 SetupGlobalDistanceFieldParameters(const FGlobalDistanceFieldParameterData& ParameterData)
 {
 	FGlobalDistanceFieldParameters2 ShaderParameters;
@@ -889,30 +894,41 @@ static void ComputeUpdateRegionsAndUpdateViewState(
 
 			const FIntVector CoverageAtlasTextureSize = GlobalDistanceField::GetCoverageAtlasSize(bLumenEnabled, View.FinalPostProcessSettings.LumenSceneViewDistance);
 
-			if (bLumenEnabled
-				&& (!GlobalDistanceFieldData.CoverageAtlasTexture
-				|| GlobalDistanceFieldData.CoverageAtlasTexture->GetDesc().Extent.X != CoverageAtlasTextureSize.X
-				|| GlobalDistanceFieldData.CoverageAtlasTexture->GetDesc().Extent.Y != CoverageAtlasTextureSize.Y
-				|| GlobalDistanceFieldData.CoverageAtlasTexture->GetDesc().Depth != CoverageAtlasTextureSize.Z))
+			if (bLumenEnabled && SupportsGlobalDistanceFieldCoverage(View.GetShaderPlatform()))
+
 			{
-				FPooledRenderTargetDesc VolumeDesc = FPooledRenderTargetDesc(FPooledRenderTargetDesc::CreateVolumeDesc(
-					CoverageAtlasTextureSize.X,
-					CoverageAtlasTextureSize.Y,
-					CoverageAtlasTextureSize.Z,
-					PF_R8,
-					FClearValueBinding::None,
-					TexCreate_None,
-					TexCreate_ShaderResource | TexCreate_UAV | TexCreate_ReduceMemoryWithTilingMode | TexCreate_3DTiling,
-					false));
+				if (!GlobalDistanceFieldData.CoverageAtlasTexture
+					|| GlobalDistanceFieldData.CoverageAtlasTexture->GetDesc().Extent.X != CoverageAtlasTextureSize.X
+					|| GlobalDistanceFieldData.CoverageAtlasTexture->GetDesc().Extent.Y != CoverageAtlasTextureSize.Y
+					|| GlobalDistanceFieldData.CoverageAtlasTexture->GetDesc().Depth != CoverageAtlasTextureSize.Z)
+				{
+					FPooledRenderTargetDesc VolumeDesc = FPooledRenderTargetDesc(FPooledRenderTargetDesc::CreateVolumeDesc(
+						CoverageAtlasTextureSize.X,
+						CoverageAtlasTextureSize.Y,
+						CoverageAtlasTextureSize.Z,
+						PF_R8,
+						FClearValueBinding::None,
+						TexCreate_None,
+						TexCreate_ShaderResource | TexCreate_UAV | TexCreate_ReduceMemoryWithTilingMode | TexCreate_3DTiling,
+						false));
 
-				GRenderTargetPool.FindFreeElement(
-					RHICmdList,
-					VolumeDesc,
-					GlobalDistanceFieldData.CoverageAtlasTexture,
-					TEXT("GlobalDistanceField.CoverageAtlas")
-				);
+					GRenderTargetPool.FindFreeElement(
+						RHICmdList,
+						VolumeDesc,
+						GlobalDistanceFieldData.CoverageAtlasTexture,
+						TEXT("GlobalDistanceField.CoverageAtlas")
+					);
 
-				bSharedDataReallocated = true;
+					bSharedDataReallocated = true;
+				}
+			}
+			else
+			{
+				if (GlobalDistanceFieldData.CoverageAtlasTexture != nullptr)
+				{
+					bSharedDataReallocated = true;
+				}
+				GlobalDistanceFieldData.CoverageAtlasTexture = nullptr;
 			}
 
 			GlobalDistanceFieldInfo.PageFreeListAllocatorBuffer = GlobalDistanceFieldData.PageFreeListAllocatorBuffer;
