[HWRT] Cache CastsHiddenShadow() and AffectsIndirectLightingWhileHidden() in FPrimitiveSceneInfo.
- reduces cache misses in GatherRelevantPrimitives(...).

Performance - Level with 20K primitives (single threaded): 2.54ms -> 1.62ms

diff --git a/Engine/Source/Runtime/Renderer/Private/PrimitiveSceneInfo.cpp b/Engine/Source/Runtime/Renderer/Private/PrimitiveSceneInfo.cpp
--- a/Engine/Source/Runtime/Renderer/Private/PrimitiveSceneInfo.cpp
+++ b/Engine/Source/Runtime/Renderer/Private/PrimitiveSceneInfo.cpp
@@ -312,6 +312,8 @@ FPrimitiveSceneInfo::FPrimitiveSceneInfo(const FPrimitiveSceneInfoAdapter& InAda
 #if RHI_RAYTRACING
 	bDrawInGame(Proxy->IsDrawnInGame()),
 	bRayTracingFarField(Proxy->IsRayTracingFarField()),
+	bCastHiddenShadow(InAdapter.SceneProxy->CastsHiddenShadow()),
+	bAffectIndirectLightingWhileHidden(InAdapter.SceneProxy->AffectsIndirectLightingWhileHidden()),
 	bIsVisibleInSceneCaptures(!InAdapter.SceneProxy->IsHiddenInSceneCapture()),
 	bIsVisibleInSceneCapturesOnly(InAdapter.SceneProxy->IsVisibleInSceneCaptureOnly()),
 	bIsRayTracingRelevant(InAdapter.SceneProxy->IsRayTracingRelevant()),
diff --git a/Engine/Source/Runtime/Renderer/Private/RayTracing/RayTracing.cpp b/Engine/Source/Runtime/Renderer/Private/RayTracing/RayTracing.cpp
--- a/Engine/Source/Runtime/Renderer/Private/RayTracing/RayTracing.cpp
+++ b/Engine/Source/Runtime/Renderer/Private/RayTracing/RayTracing.cpp
@@ -1017,8 +1017,7 @@ namespace RayTracing
 
 				const FPrimitiveSceneInfo* SceneInfo = Scene.Primitives[PrimitiveIndex];
 
-				checkf(SceneInfo->Proxy != nullptr, TEXT("SceneInfo does not have a valid Proxy object. If this occurs, this object should probably have been filtered out before being added to Scene.Primitives"));
-				const bool bRetainWhileHidden = SceneInfo->Proxy->CastsHiddenShadow() || SceneInfo->Proxy->AffectsIndirectLightingWhileHidden();
+				const bool bRetainWhileHidden = SceneInfo->bCastHiddenShadow || SceneInfo->bAffectIndirectLightingWhileHidden;
 
 				// #dxr_todo: ray tracing in scene captures should re-use the persistent RT scene. (UE-112448)
 				if (View.bIsSceneCapture && !SceneInfo->bIsVisibleInSceneCaptures && !bRetainWhileHidden)
diff --git a/Engine/Source/Runtime/Renderer/Public/PrimitiveSceneInfo.h b/Engine/Source/Runtime/Renderer/Public/PrimitiveSceneInfo.h
--- a/Engine/Source/Runtime/Renderer/Public/PrimitiveSceneInfo.h
+++ b/Engine/Source/Runtime/Renderer/Public/PrimitiveSceneInfo.h
@@ -680,6 +680,8 @@ public:
 #if RHI_RAYTRACING
 	bool bDrawInGame : 1;
 	bool bRayTracingFarField : 1;
+	bool bCastHiddenShadow : 1;
+	bool bAffectIndirectLightingWhileHidden : 1;
 	bool bIsVisibleInSceneCaptures : 1;
 	bool bIsVisibleInSceneCapturesOnly : 1;
 	bool bIsRayTracingRelevant : 1;
