Fix "GPU/MobileSceneRender" CSV stat when the new GPU profiler is enabled

diff --git a/Engine/Source/Runtime/Renderer/Private/MobileShadingRenderer.cpp b/Engine/Source/Runtime/Renderer/Private/MobileShadingRenderer.cpp
--- a/Engine/Source/Runtime/Renderer/Private/MobileShadingRenderer.cpp
+++ b/Engine/Source/Runtime/Renderer/Private/MobileShadingRenderer.cpp
@@ -1032,7 +1032,7 @@ void FMobileSceneRenderer::Render(FRDGBuilder& GraphBuilder, const FSceneRenderU
 	const ERendererOutput RendererOutput = GetRendererOutput();
 	const bool bRendererOutputFinalSceneColor = (RendererOutput != ERendererOutput::DepthPrepassOnly);
 
-	RDG_RHI_EVENT_SCOPE(GraphBuilder, MobileSceneRender);
+	RDG_RHI_EVENT_SCOPE_STAT(GraphBuilder, MobileSceneRender, MobileSceneRender);
 	RDG_RHI_GPU_STAT_SCOPE(GraphBuilder, MobileSceneRender);
 
 	IVisibilityTaskData* VisibilityTaskData = OnRenderBegin(GraphBuilder, SceneUpdateInputs);
