[HWRT] Avoid some cache misses accessing RTLODData in GatherRelevantStaticPrimitives(...)
- use TInlineAllocator<1> since the hot path (cached static instances) only has one entry.
- Cache SBTAllocation->GetSegmentCount() to avoid dereferencing SBTAllocation.

Performance - Level with 20K primitives (single threaded): 2.75ms -> 1.82ms
Memory - 20K * sizeof(FRayTracingLODData) = ~1MB (a lot of the primitives have already had a FRayTracingLODData, so actual overhead is lower).

diff --git a/Engine/Source/Runtime/Renderer/Private/PrimitiveSceneInfo.cpp b/Engine/Source/Runtime/Renderer/Private/PrimitiveSceneInfo.cpp
--- a/Engine/Source/Runtime/Renderer/Private/PrimitiveSceneInfo.cpp
+++ b/Engine/Source/Runtime/Renderer/Private/PrimitiveSceneInfo.cpp
@@ -436,6 +436,7 @@ void FPrimitiveSceneInfo::AllocateRayTracingSBT()
 		if (RayTracingGeometry && SegmentCount > 0)
 		{
 			LODData.SBTAllocation = Scene->RayTracingSBT.AllocateStaticRange(SegmentCount, RayTracingGeometry, LODData.CachedMeshCommandFlags);
+			LODData.SBTAllocationSegmentCount = LODData.SBTAllocation ? LODData.SBTAllocation->GetSegmentCount() : 0;
 		}
 	}
 }
diff --git a/Engine/Source/Runtime/Renderer/Private/RayTracing/RayTracing.cpp b/Engine/Source/Runtime/Renderer/Private/RayTracing/RayTracing.cpp
--- a/Engine/Source/Runtime/Renderer/Private/RayTracing/RayTracing.cpp
+++ b/Engine/Source/Runtime/Renderer/Private/RayTracing/RayTracing.cpp
@@ -1269,7 +1269,8 @@ namespace RayTracing
 
 					// Setup the shader bindings write offset which is computed here for all cached instances so the writing can happen in parallel to the visible bindings array
 					RelevantPrimitive->RelativeVisibleShaderBindingsOffset = Context.NumCachedStaticVisibleShaderBindings;
-					Context.NumCachedStaticVisibleShaderBindings += RTLODData.SBTAllocation->GetSegmentCount();
+					Context.NumCachedStaticVisibleShaderBindings += RTLODData.SBTAllocationSegmentCount;
+					checkSlow(RTLODData.SBTAllocationSegmentCount == RTLODData.SBTAllocation->GetSegmentCount());
 				}
 				// - DirtyCachedRayTracingPrimitives are only processed after StaticPrimitiveIndices is filled
 				// so we can end up with primitives that should be skipped here
diff --git a/Engine/Source/Runtime/Renderer/Public/PrimitiveSceneInfo.h b/Engine/Source/Runtime/Renderer/Public/PrimitiveSceneInfo.h
--- a/Engine/Source/Runtime/Renderer/Public/PrimitiveSceneInfo.h
+++ b/Engine/Source/Runtime/Renderer/Public/PrimitiveSceneInfo.h
@@ -700,6 +700,7 @@ public:
 		FRayTracingCachedMeshCommandFlags CachedMeshCommandFlags;
 		TArray<int32, TInlineAllocator<2>> CachedMeshCommandIndices;
 		FRayTracingSBTAllocation* SBTAllocation = nullptr;
+		int32 SBTAllocationSegmentCount = -1; // cache SBTAllocation->GetSegmentCount() to avoid cache misses in GatherRelevantStaticPrimitives(...)
 	};
 
 	uint32 GetRayTracingLODDataNum() const
@@ -727,7 +728,8 @@ public:
 
 private:
 
-	TArray<FRayTracingLODData> RayTracingLODData;
+	// for cached static instances (hot path), we only cache data for one LOD, so inline allocator avoids cache misses.
+	TArray<FRayTracingLODData, TInlineAllocator<1>> RayTracingLODData;
 
 	// Allocate the RayTracing SBT ranges for all the setup LODs
 	void AllocateRayTracingSBT();
