Allow an optional category to be specificed when calling r.MeshDrawCommands.DumpStats which will get appended to the stats filename.

diff --git a/Engine/Source/Runtime/Renderer/Private/MeshDrawCommandStats.cpp b/Engine/Source/Runtime/Renderer/Private/MeshDrawCommandStats.cpp
--- a/Engine/Source/Runtime/Renderer/Private/MeshDrawCommandStats.cpp
+++ b/Engine/Source/Runtime/Renderer/Private/MeshDrawCommandStats.cpp
@@ -577,7 +577,8 @@ void FMeshDrawCommandStatsManager::Update()
 
 void FMeshDrawCommandStatsManager::DumpStats(FFrameData* FrameData)
 {
-	const FString Filename = FString::Printf(TEXT("%sMeshDrawCommandStats-%s.csv"), *FPaths::ProfilingDir(), *FDateTime::Now().ToString());
+	const FString Category = OptionalCategory.IsEmpty() ? TEXT("") : FString::Printf(TEXT("%s-"), *OptionalCategory);
+	const FString Filename = FString::Printf(TEXT("%sMeshDrawCommandStats-%s%s.csv"), *FPaths::ProfilingDir(), *Category, *FDateTime::Now().ToString());
 	FArchive* CSVFile = IFileManager::Get().CreateFileWriter(*Filename, FILEWRITE_AllowRead);
 	if (CSVFile == nullptr)
 	{
@@ -670,12 +671,14 @@ void FMeshDrawCommandStatsManager::DumpStats(FFrameData* FrameData)
 
 static FAutoConsoleCommand GDumpMeshDrawCommandStatsCmd(
 	TEXT("r.MeshDrawCommands.DumpStats"),
-	TEXT("Dumps all of the Mesh Draw Command stats for a single frame to a csv file in the saved profile directory.\n"),
-	FConsoleCommandDelegate::CreateStatic([]()
+	TEXT("Dumps all of the Mesh Draw Command stats for a single frame to a csv file in the saved profile directory.\nOptionally pass in a category string to be appened to the filename"),
+	FConsoleCommandWithArgsDelegate::CreateStatic([](const TArray<FString>& Args)
 {
+	const FString& OptionalCategory = Args.Num() == 1 ? Args[0] : FString();	
+		
 	if (FMeshDrawCommandStatsManager* Instance = FMeshDrawCommandStatsManager::Get())
 	{
-		Instance->RequestDumpStats();
+		Instance->RequestDumpStats(OptionalCategory);
 	}
 }));
 
diff --git a/Engine/Source/Runtime/Renderer/Private/MeshDrawCommandStats.h b/Engine/Source/Runtime/Renderer/Private/MeshDrawCommandStats.h
--- a/Engine/Source/Runtime/Renderer/Private/MeshDrawCommandStats.h
+++ b/Engine/Source/Runtime/Renderer/Private/MeshDrawCommandStats.h
@@ -81,7 +81,11 @@ public:
 	void QueueCustomDrawIndirectArgsReadback(FRHICommandListImmediate& CommandList);
 	
 	bool CollectStats() const { return bCollectStats; }
-	void RequestDumpStats() { bRequestDumpStats = true; }	
+	void RequestDumpStats(const FString& InOptionalCategory)
+	{
+		bRequestDumpStats = true;
+		OptionalCategory = InOptionalCategory;
+	}	
 	void Update();
 	
 private:
@@ -207,7 +211,8 @@ private:
 	int32 CurrentFrameNumber = 0;			//< Ever incrementing frame number
 	bool bRequestDumpStats = false;			//< Dump stats requested via command?
 	bool bCollectStats = false;				//< Collect stats during next frame?
-	
+
+	FString OptionalCategory;				//< Optional category to append to stats filename. Empty string if no category.
 	FCriticalSection FrameDataCS;			
 	TArray<FFrameData*> Frames;						//< All active frames (contains the frame for which we are collecting stats now and all frames waiting for GPU readback)
 	FStats Stats;									//< Last updated frame stats
