Add CPU distance based streaming for virtual textures.
Virtual textures can opt into having mip levels streamed using the existing regular texture streaming logic.
A budget is configured using r.Streaming.PoolSizeForVirtualTextures.
It is assumed that this budget is low enough that there will be room in the virtual texture physical pools to accomodate it.
The system works by locking mip tiles in the virtual texture physical pools.
For testing I added r.vt.EnableFeedback.Produce. This can be set to 0 so that gpu feedback doesn't produce pages, and we rely on CPU streaming to do that instead. (The feedback still makes the page map happen).

diff --git a/Engine/Source/Runtime/Engine/Classes/Engine/Texture.h b/Engine/Source/Runtime/Engine/Classes/Engine/Texture.h
--- a/Engine/Source/Runtime/Engine/Classes/Engine/Texture.h
+++ b/Engine/Source/Runtime/Engine/Classes/Engine/Texture.h
@@ -1478,6 +1478,14 @@ PRAGMA_ENABLE_DEPRECATION_WARNINGS
 	UPROPERTY(EditAnywhere, BlueprintReadWrite, Category = Texture, AdvancedDisplay, meta=(EditCondition = "bUseVirtualTextureStreamingPriority"))
 	EVTProducerPriority VirtualTextureStreamingPriority;
 
+	/** 
+	 * The number of mips that a virtual texture can stream using the traditional CPU texture streaming system. 
+	 * In most cases this should be 0 because the virtual texture GPU feedback is sufficient and will minimize memory usage.
+	 * But CPU prefetching can be used in some cases to avoid late appearance of mips.
+	 */
+	UPROPERTY(EditAnywhere, Category = LevelOfDetail, AdvancedDisplay, meta=(EditCondition = "VirtualTextureStreaming"))
+	uint8 VirtualTexturePrefetchMips = 0;
+
 	/** If the platform supports it, tile the texture when cooking, or keep it linear and tile it when it's actually submitted to the GPU. */
 	UPROPERTY(EditAnywhere, BlueprintReadWrite, Category = Compression, AdvancedDisplay)
 	TEnumAsByte<enum TextureCookPlatformTilingSettings> CookPlatformTilingSettings;
@@ -1557,8 +1565,8 @@ PRAGMA_ENABLE_DEPRECATION_WARNINGS
 	UPROPERTY(EditAnywhere, BlueprintReadOnly, Category = Texture, AssetRegistrySearchable, AdvancedDisplay)
 	uint8 VirtualTextureStreaming : 1;
 
-	/** Use a virtual texture streaming priority set on this texture instead of the setting from the texture LOD group. */
-	UPROPERTY(EditAnywhere, Category = Texture, AdvancedDisplay, meta=(EditCondition = "VirtualTextureStreaming", InlineEditConditionToggle))
+	/** Override the virtual texture streaming priority set on this texture by the texture LOD group. */
+	UPROPERTY(EditAnywhere, Category = LevelOfDetail, AdvancedDisplay, meta=(EditCondition = "VirtualTextureStreaming", InlineEditConditionToggle))
 	uint32 bUseVirtualTextureStreamingPriority : 1;
 
 #if WITH_EDITOR
diff --git a/Engine/Source/Runtime/Engine/Private/Streaming/AsyncTextureStreaming.cpp b/Engine/Source/Runtime/Engine/Private/Streaming/AsyncTextureStreaming.cpp
--- a/Engine/Source/Runtime/Engine/Private/Streaming/AsyncTextureStreaming.cpp
+++ b/Engine/Source/Runtime/Engine/Private/Streaming/AsyncTextureStreaming.cpp
@@ -590,6 +590,11 @@ void FRenderAssetStreamingMipCalcTask::UpdateBudgetedMips_Async()
 	{
 		if (IsAborted()) break;
 
+		if (StreamingRenderAsset.IsVirtualTexture())
+		{
+			continue;
+		}
+
 		const int64 AssetMemBudgeted = StreamingRenderAsset.UpdateRetentionPriority_Async(Settings.bPrioritizeMeshLODRetention);
 		const int32 AssetMemUsed = StreamingRenderAsset.GetSize(StreamingRenderAsset.ResidentMips);
 		MemoryUsed += AssetMemUsed;
@@ -687,6 +692,11 @@ void FRenderAssetStreamingMipCalcTask::UpdateBudgetedMips_Async()
 		{
 			if (IsAborted()) break;
 
+			if (StreamingRenderAsset.IsVirtualTexture())
+			{
+				continue;
+			}
+
 			if (StreamingRenderAsset.BudgetMipBias > 0
 				&& (bResetMipBias
 					|| (Settings.bFullyLoadMeshes && StreamingRenderAsset.IsMesh())
@@ -717,8 +727,14 @@ void FRenderAssetStreamingMipCalcTask::UpdateBudgetedMips_Async()
 		{
 			FStreamingRenderAsset& StreamingRenderAsset = StreamingRenderAssets[AssetIndex];
 			// Only consider non deleted textures/meshes (can change any time).
-			if (!StreamingRenderAsset.RenderAsset) continue;
-
+			if (!StreamingRenderAsset.RenderAsset) 
+			{
+				continue;
+			}
+			if (StreamingRenderAsset.IsVirtualTexture())
+			{
+				continue;
+			}
 			// Ignore textures/meshes for which we are not allowed to reduce resolution.
 			if (!StreamingRenderAsset.IsMaxResolutionAffectedByGlobalBias() || (Settings.bFullyLoadMeshes && StreamingRenderAsset.IsMesh()))
 			{
@@ -787,8 +803,14 @@ void FRenderAssetStreamingMipCalcTask::UpdateBudgetedMips_Async()
 		{
 			FStreamingRenderAsset& StreamingRenderAsset = StreamingRenderAssets[AssetIndex];
 			// Only consider non deleted textures/meshes (can change any time).
-			if (!StreamingRenderAsset.RenderAsset) continue;
-
+			if (!StreamingRenderAsset.RenderAsset) 
+			{
+				continue;
+			}
+			if (StreamingRenderAsset.IsVirtualTexture())
+			{
+				continue;
+			}
 			// Only consider textures/meshes that won't bust budget nor generate new I/O requests
 			if (StreamingRenderAsset.BudgetedMips < StreamingRenderAsset.ResidentMips)
 			{
@@ -838,6 +860,52 @@ void FRenderAssetStreamingMipCalcTask::UpdateBudgetedMips_Async()
 #endif
 }
 
+void FRenderAssetStreamingMipCalcTask::UpdateBudgetedVirtualTextureMips_Async()
+{
+	TArray<FStreamingRenderAsset>& StreamingRenderAssets = StreamingManager.AsyncUnsafeStreamingRenderAssets;
+	const FRenderAssetStreamingSettings& Settings = StreamingManager.Settings;
+	int64 TargetMemoryBudget = FMath::Max(Settings.VirtualTexturePoolSize * 1024 * 1024, 0);
+	const bool bDisableStreaming = TargetMemoryBudget == 0;
+
+	int64 MemoryBudgeted = 0;
+	TArray<int32> AssetIndices;
+
+	// Iterate the virtual texture streaming assets, set their current budget, and collect candidates for budget reduction.
+	for (int32 AssetIndex = 0; AssetIndex < StreamingRenderAssets.Num() && !IsAborted(); ++AssetIndex)
+	{
+		FStreamingRenderAsset& StreamingRenderAsset = StreamingRenderAssets[AssetIndex];
+		if (StreamingRenderAsset.RenderAsset && StreamingRenderAsset.IsVirtualTexture())
+		{
+			if (bDisableStreaming)
+			{
+				StreamingRenderAsset.BudgetedMips = StreamingRenderAsset.MinAllowedMips;
+			}
+			else
+			{
+				MemoryBudgeted += StreamingRenderAsset.UpdateRetentionPriority_Async(Settings.bPrioritizeMeshLODRetention);;
+		
+				if (StreamingRenderAsset.IsMaxResolutionAffectedByGlobalBias())
+				{
+					const int32 MinAllowedMips = FMath::Max(StreamingRenderAsset.MinAllowedMips, StreamingRenderAsset.NumForcedMips);
+					if (StreamingRenderAsset.BudgetedMips > MinAllowedMips)
+					{
+						AssetIndices.Add(AssetIndex);
+					}
+				}
+			}
+		}
+	}
+
+	// Try to drop mips when overbudget.
+	if (MemoryBudgeted > TargetMemoryBudget && !IsAborted())
+	{
+		// Sort to prioritize mip retention.
+		AssetIndices.Sort(FCompareRenderAssetByRetentionPriority(StreamingRenderAssets));
+		// Drop mips to reach target budget.
+		TryDropMips(AssetIndices, MemoryBudgeted, TargetMemoryBudget);
+	}
+}
+
 void FRenderAssetStreamingMipCalcTask::UpdateLoadAndCancelationRequests_Async()
 {
 	TArray<FStreamingRenderAsset>& StreamingRenderAssets = StreamingManager.AsyncUnsafeStreamingRenderAssets;
@@ -989,6 +1057,7 @@ void FRenderAssetStreamingMipCalcTask::DoWork()
 
 	// According to budget, make relevant sacrifices and keep possible unwanted mips
 	UpdateBudgetedMips_Async();
+	UpdateBudgetedVirtualTextureMips_Async();
 
 	// Update load requests.
 	UpdateLoadAndCancelationRequests_Async();
diff --git a/Engine/Source/Runtime/Engine/Private/Streaming/AsyncTextureStreaming.h b/Engine/Source/Runtime/Engine/Private/Streaming/AsyncTextureStreaming.h
--- a/Engine/Source/Runtime/Engine/Private/Streaming/AsyncTextureStreaming.h
+++ b/Engine/Source/Runtime/Engine/Private/Streaming/AsyncTextureStreaming.h
@@ -198,6 +198,8 @@ private:
 	void TryKeepMips(TArray<int32>& PrioritizedRenderAssets, int64& MemoryBudgeted, const int64 InMemoryBudget);
 
 	void UpdateBudgetedMips_Async();
+	
+	void UpdateBudgetedVirtualTextureMips_Async();
 
 	void UpdateLoadAndCancelationRequests_Async();
 
diff --git a/Engine/Source/Runtime/Engine/Private/Streaming/StreamingManagerTexture.cpp b/Engine/Source/Runtime/Engine/Private/Streaming/StreamingManagerTexture.cpp
--- a/Engine/Source/Runtime/Engine/Private/Streaming/StreamingManagerTexture.cpp
+++ b/Engine/Source/Runtime/Engine/Private/Streaming/StreamingManagerTexture.cpp
@@ -2930,7 +2930,7 @@ bool FRenderAssetStreamingManager::HandleInvestigateRenderAssetCommand(const TCH
 					UE_LOG(LogContentStreaming, Log, TEXT("  Force all mips:  No mip-maps"));
 				}
 				
-				if (Texture && Texture->GetRunningPlatformData() && *Texture->GetRunningPlatformData())
+				if (Texture && !Texture->IsCurrentlyVirtualTextured() && Texture->GetRunningPlatformData() && *Texture->GetRunningPlatformData())
 				{
 					const TIndirectArray<struct FTexture2DMipMap>& TextureMips = (*Texture->GetRunningPlatformData())->Mips;
 					UTexture2D* Texture2D = Cast<UTexture2D>(Texture);
diff --git a/Engine/Source/Runtime/Engine/Private/Streaming/StreamingTexture.cpp b/Engine/Source/Runtime/Engine/Private/Streaming/StreamingTexture.cpp
--- a/Engine/Source/Runtime/Engine/Private/Streaming/StreamingTexture.cpp
+++ b/Engine/Source/Runtime/Engine/Private/Streaming/StreamingTexture.cpp
@@ -84,6 +84,7 @@ void FStreamingRenderAsset::UpdateStaticData(const FRenderAssetStreamingSettings
 			BoostFactor = GetExtraBoost(TextureLODGroup, Settings);
 			bLoadWithHigherPriority = Settings.HighPriorityLoad_Texture[LODGroup];
 			bIsTerrainTexture = (TextureLODGroup == TEXTUREGROUP_Terrain_Heightmap || TextureLODGroup == TEXTUREGROUP_Terrain_Weightmap);
+			bIsVirtualTexture = ResourceState.bSupportsVirtualStreaming;
 		}
 		else
 		{
@@ -94,6 +95,7 @@ void FStreamingRenderAsset::UpdateStaticData(const FRenderAssetStreamingSettings
 			BoostFactor = 1.f;
 			bLoadWithHigherPriority = false;
 			bIsTerrainTexture = false;
+			bIsVirtualTexture = false;
 			if (RenderAssetType == EStreamableRenderAssetType::StaticMesh)
 			{
 				const UStaticMesh* StaticMesh = CastChecked<UStaticMesh>(RenderAsset);
diff --git a/Engine/Source/Runtime/Engine/Private/Streaming/StreamingTexture.h b/Engine/Source/Runtime/Engine/Private/Streaming/StreamingTexture.h
--- a/Engine/Source/Runtime/Engine/Private/Streaming/StreamingTexture.h
+++ b/Engine/Source/Runtime/Engine/Private/Streaming/StreamingTexture.h
@@ -147,6 +147,11 @@ struct FStreamingRenderAsset
 		return RenderAssetType == EStreamableRenderAssetType::Texture;
 	}
 
+	bool IsVirtualTexture() const
+	{
+		return RenderAssetType == EStreamableRenderAssetType::Texture && bIsVirtualTexture;
+	}
+
 	bool IsMesh() const
 	{
 		// FRenderAssetStreamingManager only handles textures and meshes currently
@@ -257,6 +262,8 @@ struct FStreamingRenderAsset
 	/** (2) Cached state on disk of the optional mips for this streaming texture */
 	EOptionalMipsState	OptionalMipsState;
 
+	/** (1) Whether the texture uses virtual texture streaming. */
+	uint8			bIsVirtualTexture : 1;
 	/** (1) Whether the texture is loaded with higher priority. */
 	uint8			bLoadWithHigherPriority : 1;
 	/** (1) Whether the texture should be forcibly fully loaded. */
diff --git a/Engine/Source/Runtime/Engine/Private/Streaming/TextureStreamingHelpers.cpp b/Engine/Source/Runtime/Engine/Private/Streaming/TextureStreamingHelpers.cpp
--- a/Engine/Source/Runtime/Engine/Private/Streaming/TextureStreamingHelpers.cpp
+++ b/Engine/Source/Runtime/Engine/Private/Streaming/TextureStreamingHelpers.cpp
@@ -123,6 +123,13 @@ static TAutoConsoleVariable<int32> CVarStreamingPoolSizeForMeshes(
 	TEXT("< 0: Mesh and texture share the same pool, otherwise the size of pool dedicated to meshes."),
 	ECVF_Scalability | ECVF_ExcludeFromPreview);
 
+static TAutoConsoleVariable<int32> CVarStreamingPoolSizeForVirtualTextures(
+	TEXT("r.Streaming.PoolSizeForVirtualTextures"),
+	0,
+	TEXT("The maximum size in MB of the virtual texture pool to use for streaming using distance based streaming.")
+	TEXT("0 (default) means that the virtual texture streaming is only GPU feedback driven."),
+	ECVF_Scalability | ECVF_ExcludeFromPreview);
+
 TAutoConsoleVariable<int32> CVarStreamingMaxTempMemoryAllowed(
 	TEXT("r.Streaming.MaxTempMemoryAllowed"),
 	50,
@@ -324,6 +331,7 @@ void FRenderAssetStreamingSettings::Update()
 	GlobalMipBias = !GIsEditor ? FMath::FloorToInt(FMath::Max<float>(0.f, CVarStreamingMipBias.GetValueOnAnyThread())) : 0;
 	PoolSize = CVarStreamingPoolSize.GetValueOnAnyThread();
 	MeshPoolSize = CVarStreamingPoolSizeForMeshes.GetValueOnAnyThread();
+	VirtualTexturePoolSize = CVarStreamingPoolSizeForVirtualTextures.GetValueOnAnyThread();
 	bUsePerTextureBias = CVarStreamingUsePerTextureBias.GetValueOnAnyThread() != 0;
 	bUseNewMetrics = CVarStreamingUseNewMetrics.GetValueOnAnyThread() != 0;
 	bLimitPoolSizeToVRAM = !GIsEditor && CVarStreamingLimitPoolSizeToVRAM.GetValueOnAnyThread() != 0;
diff --git a/Engine/Source/Runtime/Engine/Private/Streaming/TextureStreamingHelpers.h b/Engine/Source/Runtime/Engine/Private/Streaming/TextureStreamingHelpers.h
--- a/Engine/Source/Runtime/Engine/Private/Streaming/TextureStreamingHelpers.h
+++ b/Engine/Source/Runtime/Engine/Private/Streaming/TextureStreamingHelpers.h
@@ -101,6 +101,7 @@ struct FRenderAssetStreamingSettings
 	int32 GlobalMipBias;
 	int32 PoolSize;
 	int32 MeshPoolSize;
+	int32 VirtualTexturePoolSize;
 	bool bLimitPoolSizeToVRAM;
 	bool bUseNewMetrics;
 	bool bFullyLoadUsedTextures;
diff --git a/Engine/Source/Runtime/Engine/Private/Texture.cpp b/Engine/Source/Runtime/Engine/Private/Texture.cpp
--- a/Engine/Source/Runtime/Engine/Private/Texture.cpp
+++ b/Engine/Source/Runtime/Engine/Private/Texture.cpp
@@ -413,6 +413,26 @@ void UTexture::UpdateResourceWithParams(EUpdateResourceFlags InFlags)
 				}
 			}
 
+			if (FVirtualTexture2DResource* VirtualTextureResource = NewResource->GetVirtualTexture2DResource())
+			{
+				const uint8 NonStreamingMips = FMath::FloorLog2(VirtualTextureResource->GetTileSize()) + 1;
+				const uint8 NumMips = VirtualTextureResource->GetNumMips();
+				const uint8 MaxStreamingMip = FMath::Min(VirtualTexturePrefetchMips, NumMips);
+				
+				if (MaxStreamingMip > NonStreamingMips)
+				{
+					CachedSRRState.Data = 0;
+					CachedSRRState.bSupportsStreaming = 1;
+					CachedSRRState.bSupportsVirtualStreaming = 1;
+					CachedSRRState.NumNonStreamingLODs = NonStreamingMips;
+					CachedSRRState.NumResidentLODs = NonStreamingMips;
+					CachedSRRState.NumRequestedLODs = NonStreamingMips;
+					CachedSRRState.NumNonOptionalLODs = MaxStreamingMip;
+					CachedSRRState.MaxNumLODs = MaxStreamingMip;
+					CachedSRRState.bHasPendingInitHint = true;
+				}
+			}
+
 			PrivateResource = NewResource;
 
 #if RHI_ENABLE_RESOURCE_INFO
@@ -1855,7 +1875,8 @@ bool UTexture::IsPossibleToStream() const
 bool UTexture::IsCandidateForTextureStreamingOnPlatformDuringCook(const ITargetPlatform* InTargetPlatform) const
 {
 	const bool bIsVirtualTextureStreaming = IsVirtualTexturingEnabled(InTargetPlatform) ? VirtualTextureStreaming : false;
-	const bool bIsCandidateForTextureStreaming = InTargetPlatform->SupportsFeature(ETargetPlatformFeatures::TextureStreaming) && !bIsVirtualTextureStreaming;
+	const bool bIsVirtualTextureStreamingWithNoPrefetching = bIsVirtualTextureStreaming && VirtualTexturePrefetchMips == 0;
+	const bool bIsCandidateForTextureStreaming = InTargetPlatform->SupportsFeature(ETargetPlatformFeatures::TextureStreaming) && !bIsVirtualTextureStreamingWithNoPrefetching;
 
 	if (bIsCandidateForTextureStreaming &&
 		IsPossibleToStream())
diff --git a/Engine/Source/Runtime/Engine/Private/Texture2D.cpp b/Engine/Source/Runtime/Engine/Private/Texture2D.cpp
--- a/Engine/Source/Runtime/Engine/Private/Texture2D.cpp
+++ b/Engine/Source/Runtime/Engine/Private/Texture2D.cpp
@@ -1895,11 +1895,36 @@ FIntPoint FVirtualTexture2DResource::GetPhysicalTextureSize(uint32 LayerIndex) c
 	return FIntPoint(0, 0);
 }
 
+static void StreamVirtualTextureMips(FVirtualTexture2DResource* VirtualTexture2DResource, FStreamableRenderResourceState& CachedSRRState, int32 NewMipCount)
+{
+	CachedSRRState.NumResidentLODs = CachedSRRState.NumRequestedLODs;
+	CachedSRRState.bHasPendingLODTransitionHint = false;
+
+	ENQUEUE_RENDER_COMMAND(RequestVTTiles)([VirtualTexture2DResource, NewMipCount](FRHICommandListImmediate& RHICmdList)
+	{
+		const int32 MipLevel = VirtualTexture2DResource->GetNumMips() - NewMipCount;
+		GetRendererModule().LockVirtualTextureTiles(VirtualTexture2DResource->GetProducerHandle(), MipLevel);
+	});
+}
+
 bool UTexture2D::StreamIn(int32 NewMipCount, bool bHighPrio)
 {
 	check(IsInGameThread());
 
-	FTexture2DResource* Texture2DResource = GetResource() ? GetResource()->GetTexture2DResource() : nullptr;
+	FTextureResource* TextureResource = GetResource();
+
+	FVirtualTexture2DResource* VirtualTexture2DResource = TextureResource ? TextureResource->GetVirtualTexture2DResource() : nullptr;
+	if (VirtualTexture2DResource != nullptr)
+	{
+		const bool bOK = CachedSRRState.StreamIn(NewMipCount);
+		if (bOK)
+		{
+			StreamVirtualTextureMips(VirtualTexture2DResource, CachedSRRState, NewMipCount);
+		}
+		return bOK;
+	}
+
+	FTexture2DResource* Texture2DResource = TextureResource ? TextureResource->GetTexture2DResource() : nullptr;
 	if (!HasPendingInitOrStreaming() && CachedSRRState.StreamIn(NewMipCount) && ensure(Texture2DResource))
 	{
 		FTextureMipDataProvider* CustomMipDataProvider = nullptr;
@@ -2022,7 +2047,20 @@ bool UTexture2D::StreamOut(int32 NewMipCount)
 {
 	check(IsInGameThread());
 
-	FTexture2DResource* Texture2DResource = GetResource() ? GetResource()->GetTexture2DResource() : nullptr;
+	FTextureResource* TextureResource = GetResource();
+	
+	FVirtualTexture2DResource* VirtualTexture2DResource = TextureResource ? TextureResource->GetVirtualTexture2DResource() : nullptr;
+	if (VirtualTexture2DResource != nullptr)
+	{
+		const bool bOK = CachedSRRState.StreamOut(NewMipCount);
+		if (bOK)
+		{
+			StreamVirtualTextureMips(VirtualTexture2DResource, CachedSRRState, NewMipCount);
+		}
+		return bOK;
+	}
+
+	FTexture2DResource* Texture2DResource = TextureResource ? TextureResource->GetTexture2DResource() : nullptr;
 	if (!HasPendingInitOrStreaming() && CachedSRRState.StreamOut(NewMipCount) && ensure(Texture2DResource))
 	{
 		if (Texture2DResource->bUsePartiallyResidentMips)
diff --git a/Engine/Source/Runtime/Engine/Public/Streaming/StreamableRenderResourceState.h b/Engine/Source/Runtime/Engine/Public/Streaming/StreamableRenderResourceState.h
--- a/Engine/Source/Runtime/Engine/Public/Streaming/StreamableRenderResourceState.h
+++ b/Engine/Source/Runtime/Engine/Public/Streaming/StreamableRenderResourceState.h
@@ -44,6 +44,9 @@ struct FStreamableRenderResourceState
 			/* Whether the resource LODs can be streamed or not. Even when false, the structure can contain valid information if IsValid(). */
 			uint8 bSupportsStreaming : 1;
 
+			/** Whether the resource uses virtual texture streaming. Virtual textures need to use a specific memory category. */
+			uint8 bSupportsVirtualStreaming : 1;
+
 			/* Whether the render resources have pending InitRHI(). Used in HasPendingInitOrStreaming() to avoid dereferencing the different resource pointers unless required. Once false, expected to stay false. */
 			mutable uint8 bHasPendingInitHint : 1;
 
diff --git a/Engine/Source/Runtime/RenderCore/Public/RendererInterface.h b/Engine/Source/Runtime/RenderCore/Public/RendererInterface.h
--- a/Engine/Source/Runtime/RenderCore/Public/RendererInterface.h
+++ b/Engine/Source/Runtime/RenderCore/Public/RendererInterface.h
@@ -834,6 +834,14 @@ public:
 
 	/** Ensure that any tiles requested by 'RequestVirtualTextureTiles' are loaded, must be called from render thread */
 	virtual void LoadPendingVirtualTextureTiles(FRHICommandListImmediate& RHICmdList, ERHIFeatureLevel::Type FeatureLevel) = 0;
+	
+	/** 
+	 * Lock all tiles associated with a producer up to and including a given mip level.
+	 * The mip level is stored on the producer, so that if we call twice then the second call will either lock or unlock tiles to reach the new level.
+	 * Locked tiles should remain resident once loaded, though there is no guarantee over how quickly they will be loaded.
+	 * Also if the virtual texture pools are already oversubscribed then we will not be able to satisfy the loading.
+	 */
+	virtual void LockVirtualTextureTiles(FVirtualTextureProducerHandle ProducerHandle, int32 InMipLevel) = 0;
 
 	/** Allocate a buffer and record all virtual texture page requests until the next call to either SetVirtualTextureRequestRecordBuffer or GetVirtualTextureRequestRecordBuffer. */
 	virtual void SetVirtualTextureRequestRecordBuffer(uint64 Handle) = 0;
diff --git a/Engine/Source/Runtime/Renderer/Private/RendererModule.h b/Engine/Source/Runtime/Renderer/Private/RendererModule.h
--- a/Engine/Source/Runtime/Renderer/Private/RendererModule.h
+++ b/Engine/Source/Runtime/Renderer/Private/RendererModule.h
@@ -110,11 +110,12 @@ public:
 	virtual void RequestVirtualTextureTiles(const FVector2D& InScreenSpaceSize, int32 InMipLevel) override;
 	virtual void RequestVirtualTextureTiles(const FMaterialRenderProxy* InMaterialRenderProxy, const FVector2D& InScreenSpaceSize, ERHIFeatureLevel::Type InFeatureLevel) override;
 	virtual void RequestVirtualTextureTiles(IAllocatedVirtualTexture* AllocatedVT, const FVector2D& InScreenSpaceSize, const FVector2D& InViewportPosition, const FVector2D& InViewportSize, const FVector2D& InUV0, const FVector2D& InUV1, int32 InMipLevel) override;
-	
+
 	UE_DEPRECATED(5.4, "Use RequestVirtualTextureTiles() overloads that takes similar parameters. Make sure not to negate the InViewportPosition.")
 	virtual void RequestVirtualTextureTilesForRegion(IAllocatedVirtualTexture* AllocatedVT, const FVector2D& InScreenSpaceSize, const FVector2D& InViewportPosition, const FVector2D& InViewportSize, const FVector2D& InUV0, const FVector2D& InUV1, int32 InMipLevel) override;
 	
 	virtual void LoadPendingVirtualTextureTiles(FRHICommandListImmediate& RHICmdList, ERHIFeatureLevel::Type FeatureLevel) override;
+	virtual void LockVirtualTextureTiles(FVirtualTextureProducerHandle ProducerHandle, int32 InMipLevel) override;
 	virtual void SetVirtualTextureRequestRecordBuffer(uint64 Handle) override;
 	virtual uint64 GetVirtualTextureRequestRecordBuffer(TSet<uint64>& OutPageRequests) override;
 	virtual void FlushVirtualTextureCache() override;
diff --git a/Engine/Source/Runtime/Renderer/Private/SceneRendering.cpp b/Engine/Source/Runtime/Renderer/Private/SceneRendering.cpp
--- a/Engine/Source/Runtime/Renderer/Private/SceneRendering.cpp
+++ b/Engine/Source/Runtime/Renderer/Private/SceneRendering.cpp
@@ -5473,6 +5473,11 @@ void FRendererModule::LoadPendingVirtualTextureTiles(FRHICommandListImmediate& R
 	GraphBuilder.Execute();
 }
 
+void FRendererModule::LockVirtualTextureTiles(FVirtualTextureProducerHandle ProducerHandle, int32 InMipLevel)
+{
+	FVirtualTextureSystem::Get().SetMipLevelToLock(ProducerHandle, InMipLevel);
+}
+
 void FRendererModule::SetVirtualTextureRequestRecordBuffer(uint64 Handle)
 {
 #if WITH_EDITOR
diff --git a/Engine/Source/Runtime/Renderer/Private/VT/VirtualTextureProducer.h b/Engine/Source/Runtime/Renderer/Private/VT/VirtualTextureProducer.h
--- a/Engine/Source/Runtime/Renderer/Private/VT/VirtualTextureProducer.h
+++ b/Engine/Source/Runtime/Renderer/Private/VT/VirtualTextureProducer.h
@@ -33,12 +33,18 @@ public:
 	inline uint32 GetNumPhysicalGroups() const { return Description.NumPhysicalGroups; }
 	inline FVirtualTexturePhysicalSpace* GetPhysicalSpaceForPhysicalGroup(uint32 GroupIndex) const { check(GroupIndex < Description.NumPhysicalGroups); return PhysicalGroups[GroupIndex].PhysicalSpace; }
 
+	inline void SetLockedMipCount(uint8 InLockedMipCount) { LockedMipCount = InLockedMipCount; }
+	inline uint8 GetLockedMipCount() const { return LockedMipCount; }
+
 private:
 	friend class FVirtualTextureProducerCollection;
 	~FVirtualTextureProducer();
 
 	FVTProducerDescription Description;
 	
+	// Number of mips that we have requested to keep locked.
+	uint8 LockedMipCount = 0;
+
 	// A physical group is a physical space which multiple texture layers map to
 	// Texture layers that share a physical group will always be sampled via the same page table channel (they will have the same physical UV coordinates)
 	// More than one physical group can share the same FVirtualTexturePhysicalSpace
diff --git a/Engine/Source/Runtime/Renderer/Private/VT/VirtualTextureSystem.cpp b/Engine/Source/Runtime/Renderer/Private/VT/VirtualTextureSystem.cpp
--- a/Engine/Source/Runtime/Renderer/Private/VT/VirtualTextureSystem.cpp
+++ b/Engine/Source/Runtime/Renderer/Private/VT/VirtualTextureSystem.cpp
@@ -87,6 +87,12 @@ static TAutoConsoleVariable<int32> CVarVTEnableFeedback(
 	TEXT("Enable processing of the GPU generated feedback buffer."),
 	ECVF_RenderThreadSafe
 );
+static TAutoConsoleVariable<int32> CVarVTEnableFeedbackProduce(
+	TEXT("r.VT.EnableFeedback.Produce"),
+	1,
+	TEXT("Enable page loading triggered by GPU virtual texture feedback."),
+	ECVF_RenderThreadSafe
+);
 static TAutoConsoleVariable<int32> CVarVTParallelFeedbackTasks(
 	TEXT("r.VT.ParallelFeedbackTasks"),
 	0,
@@ -183,6 +189,7 @@ TAutoConsoleVariable<int32> CVarVTSortTileRequestsByPriority(
 FVirtualTextureUpdateSettings::FVirtualTextureUpdateSettings()
 {
 	bEnableFeedback = CVarVTEnableFeedback.GetValueOnRenderThread() != 0;
+	bEnableFeedbackProduce = CVarVTEnableFeedbackProduce.GetValueOnRenderThread() != 0;
 	bEnablePlayback = CVarVTEnablePlayback.GetValueOnRenderThread() != 0;
 	bForceContinuousUpdate = CVarVTForceContinuousUpdate.GetValueOnRenderThread() != 0;
 	bParallelFeedbackTasks = CVarVTParallelFeedbackTasks.GetValueOnRenderThread() != 0;
@@ -311,6 +318,7 @@ struct FGatherRequestsParameters
 	uint32 NumPages = 0u;
 	uint32 FrameRequested;
 	bool bForceContinuousUpdate = false;
+	bool bEnableLoadRequests = false;
 };
 
 class FGatherRequestsTask
@@ -1311,6 +1319,48 @@ void FVirtualTextureSystem::RequestTilesInternal(const IAllocatedVirtualTexture*
 	}
 }
 
+void FVirtualTextureSystem::SetMipLevelToLock(FVirtualTextureProducerHandle ProducerHandle, int32 InMipLevel)
+{
+	check(!bUpdating);
+	UE::TScopeLock Lock(Mutex);
+
+	FVirtualTextureProducer& Producer = Producers.GetProducer(ProducerHandle);
+	const int32 MaxLevel = Producer.GetMaxLevel();
+	const int32 MaxLoadLevel = MaxLevel - FMath::FloorLog2(Producer.GetDescription().TileSize);
+	const int32 NumLockMips = FMath::Max(MaxLoadLevel - InMipLevel + 1, 0);
+
+	// Compare with current setting on the producer and Lock or Unlock tiles to reach the new setting.
+	const int32 NumLockMipsPrevious = Producer.GetLockedMipCount();
+	Producer.SetLockedMipCount(NumLockMips);
+
+	const bool bLock = NumLockMips > NumLockMipsPrevious;
+	const int32 FirstMip = bLock ? NumLockMips : NumLockMipsPrevious;
+	const int32 LastMip = bLock ? NumLockMipsPrevious : NumLockMips;
+
+	for (int32 MipIndex = FirstMip; MipIndex > LastMip; --MipIndex)
+	{
+		const int32 vLevel = MaxLoadLevel - MipIndex;
+		const int32 MipWidthInTiles = 1 << MipIndex;
+		const int32 MipHeightInTiles = 1 << MipIndex;
+		for (int32 PosY = 0; PosY < MipHeightInTiles; PosY++)
+		{
+			for (int32 PosX = 0; PosX < MipWidthInTiles; PosX++)
+			{
+				const uint32 Address = FMath::MortonCode2(PosX) | (FMath::MortonCode2(PosY) << 1);
+				const FVirtualTextureLocalTile Tile(ProducerHandle, Address, vLevel);
+				if (bLock)
+				{
+					LockTile(Tile);
+				}
+				else
+				{
+					UnlockTile(Tile, &Producer);
+				}
+			}
+		}
+	}
+}
+
 static int32 WrapTilePosition(int32 Position, int32 Size)
 {
 	const int32 Result = Position % Size;
@@ -1425,6 +1475,7 @@ void FVirtualTextureSystem::GatherRequests(FUniqueRequestList* MergedRequestList
 				Params.System = this;
 				Params.FrameRequested = FrameRequested;
 				Params.bForceContinuousUpdate = Settings.bForceContinuousUpdate;
+				Params.bEnableLoadRequests = Settings.bEnableFeedbackProduce;
 				Params.UniquePageList = UniquePageList;
 				Params.PageUpdateFlushCount = PageUpdateFlushCount;
 				Params.PageUpdateBuffers = Allocator.CreateArray<FPageUpdateBuffer>(PhysicalSpaces.Num());
@@ -1855,25 +1906,28 @@ void FVirtualTextureSystem::GatherRequestsTask(const FGatherRequestsParameters&
 						}
 					}
 
-					const bool bStreamingRequest = Producer->GetVirtualTexture()->IsPageStreamed(PrefetchLocal_vLevel, PrefetchLocal_vAddress);
-					const FVirtualTextureLocalTile Tile(ProducerHandle, PrefetchLocal_vAddress, PrefetchLocal_vLevel);
-					const uint16 LoadRequestIndex = RequestList->AddLoadRequest(
-						FVirtualTextureLocalTileRequest(Tile, ProducerPriority, EVTInvalidatePriority::Normal),
-						ProducerPhysicalGroupMaskToPrefetch, PageCount, bStreamingRequest);
-					if (LoadRequestIndex != 0xffff)
+					if (Parameters.bEnableLoadRequests)
 					{
-						const uint32 PrefetchMapping_vLevel = PrefetchLocal_vLevel + ProducerMipBias;
-						ensure(PrefetchMapping_vLevel <= MaxLevel);
-						const uint32 Prefetch_vAddress = vAddress & (0xffffffff << (PrefetchMapping_vLevel * vDimensions));
-						for (uint32 LoadLayerIndex = 0u; LoadLayerIndex < NumPageTableLayersToLoad; ++LoadLayerIndex)
+						const bool bStreamingRequest = Producer->GetVirtualTexture()->IsPageStreamed(PrefetchLocal_vLevel, PrefetchLocal_vAddress);
+						const FVirtualTextureLocalTile Tile(ProducerHandle, PrefetchLocal_vAddress, PrefetchLocal_vLevel);
+						const uint16 LoadRequestIndex = RequestList->AddLoadRequest(
+							FVirtualTextureLocalTileRequest(Tile, ProducerPriority, EVTInvalidatePriority::Normal),
+							ProducerPhysicalGroupMaskToPrefetch, PageCount, bStreamingRequest);
+						if (LoadRequestIndex != 0xffff)
 						{
-							const uint32 LayerIndex = PageTableLayersToLoad[LoadLayerIndex];
-							if (AllocatedVT->GetProducerIndexForPageTableLayer(LayerIndex) == ProducerIndex)
+							const uint32 PrefetchMapping_vLevel = PrefetchLocal_vLevel + ProducerMipBias;
+							ensure(PrefetchMapping_vLevel <= MaxLevel);
+							const uint32 Prefetch_vAddress = vAddress & (0xffffffff << (PrefetchMapping_vLevel * vDimensions));
+							for (uint32 LoadLayerIndex = 0u; LoadLayerIndex < NumPageTableLayersToLoad; ++LoadLayerIndex)
 							{
-								const uint32 ProducerPhysicalGroupIndex = AllocatedVT->GetProducerPhysicalGroupIndexForPageTableLayer(LayerIndex);
-								if (ProducerPhysicalGroupMaskToPrefetch & (1u << ProducerPhysicalGroupIndex))
+								const uint32 LayerIndex = PageTableLayersToLoad[LoadLayerIndex];
+								if (AllocatedVT->GetProducerIndexForPageTableLayer(LayerIndex) == ProducerIndex)
 								{
-									RequestList->AddMappingRequest(LoadRequestIndex, ProducerPhysicalGroupIndex, ID, LayerIndex, MaxLevel, Prefetch_vAddress, PrefetchMapping_vLevel, PrefetchMapping_vLevel);
+									const uint32 ProducerPhysicalGroupIndex = AllocatedVT->GetProducerPhysicalGroupIndexForPageTableLayer(LayerIndex);
+									if (ProducerPhysicalGroupMaskToPrefetch & (1u << ProducerPhysicalGroupIndex))
+									{
+										RequestList->AddMappingRequest(LoadRequestIndex, ProducerPhysicalGroupIndex, ID, LayerIndex, MaxLevel, Prefetch_vAddress, PrefetchMapping_vLevel, PrefetchMapping_vLevel);
+									}
 								}
 							}
 						}
@@ -1881,7 +1935,7 @@ void FVirtualTextureSystem::GatherRequestsTask(const FGatherRequestsParameters&
 				}
 			}
 
-			if (GroupMaskToLoad != 0u)
+			if (GroupMaskToLoad != 0u && Parameters.bEnableLoadRequests)
 			{
 				const bool bStreamingRequest = Producer->GetVirtualTexture()->IsPageStreamed(Local_vLevel, Local_vAddress);
 				const FVirtualTextureLocalTile Tile(ProducerHandle, Local_vAddress, Local_vLevel);
diff --git a/Engine/Source/Runtime/Renderer/Private/VT/VirtualTextureSystem.h b/Engine/Source/Runtime/Renderer/Private/VT/VirtualTextureSystem.h
--- a/Engine/Source/Runtime/Renderer/Private/VT/VirtualTextureSystem.h
+++ b/Engine/Source/Runtime/Renderer/Private/VT/VirtualTextureSystem.h
@@ -67,6 +67,7 @@ struct FVirtualTextureUpdateSettings
 	bool bEnableAsyncTasks = true;
 	bool bEnablePageRequests = true;
 	bool bEnableFeedback;
+	bool bEnableFeedbackProduce;
 	bool bEnablePlayback;
 	bool bForceContinuousUpdate;
 	bool bParallelFeedbackTasks;
@@ -156,6 +157,8 @@ public:
 
 	void LoadPendingTiles(FRDGBuilder& GraphBuilder, ERHIFeatureLevel::Type FeatureLevel);
 	
+	void SetMipLevelToLock(FVirtualTextureProducerHandle ProducerHandle, int32 InMipLevel);
+
 #if WITH_EDITOR
 	void SetVirtualTextureRequestRecordBuffer(uint64 Handle);
 	uint64 GetVirtualTextureRequestRecordBuffer(TSet<uint64>& OutPageRequests);
