Lens Flare: Backing out all CLs related to lens flare intensity separation from bloom intensity due to issues with existing live content.

Leaving engine version bump submitted to avoid any serialization issue.

diff --git a/Engine/Source/Runtime/Engine/Private/Scene.cpp b/Engine/Source/Runtime/Engine/Private/Scene.cpp
--- a/Engine/Source/Runtime/Engine/Private/Scene.cpp
+++ b/Engine/Source/Runtime/Engine/Private/Scene.cpp
@@ -498,7 +498,7 @@ FPostProcessSettings::FPostProcessSettings()
 	LocalExposureMiddleGreyBias = 0.0f;
 
 	// next value might get overwritten by r.DefaultFeature.LensFlare
-	LensFlareIntensity = 0.675f;
+	LensFlareIntensity = 1.0f;
 	LensFlareTint = FLinearColor(1.0f, 1.0f, 1.0f);
 	LensFlareBokehSize = 3.0f;
 	LensFlareThreshold = 8.0f;
@@ -664,12 +664,6 @@ bool FPostProcessSettings::Serialize(FArchive& Ar)
 	Ar.UsingCustomVersion(FReleaseObjectVersion::GUID);
 	Ar.UsingCustomVersion(FUE5ReleaseStreamObjectVersion::GUID);
 
-	if (Ar.CustomVer(FUE5ReleaseStreamObjectVersion::GUID) < FUE5ReleaseStreamObjectVersion::BloomIndependentLensFlare)
-	{
-		// For backwards compatibility, we must preserve the pre-UE 5.6 default value *before* serialization.
-		LensFlareIntensity = 1.0f;
-	}
-
 	// Don't actually serialize, just write the custom version for PostSerialize
 	return false;
 }
@@ -934,12 +928,6 @@ void FPostProcessSettings::PostSerialize(const FArchive& Ar)
 				}
 			}
 		}
-
-		if (Ar.CustomVer(FUE5ReleaseStreamObjectVersion::GUID) < FUE5ReleaseStreamObjectVersion::BloomIndependentLensFlare)
-		{
-			// Lens flare intensity is now independent from bloom, but we preserve the bloom intensity scale on settings before UE 5.6.
-			LensFlareIntensity *= FMath::Max(0.0f, BloomIntensity);
-		}
 	}
 }
 #endif
diff --git a/Engine/Source/Runtime/Renderer/Private/PostProcess/PostProcessLensFlares.cpp b/Engine/Source/Runtime/Renderer/Private/PostProcess/PostProcessLensFlares.cpp
--- a/Engine/Source/Runtime/Renderer/Private/PostProcess/PostProcessLensFlares.cpp
+++ b/Engine/Source/Runtime/Renderer/Private/PostProcess/PostProcessLensFlares.cpp
@@ -406,7 +406,7 @@ FScreenPassTexture AddLensFlaresPass(
 	LensFlareInputs.TintColorsPerFlare = Settings.LensFlareTints;
 	LensFlareInputs.TintColor = Settings.LensFlareTint;
 	LensFlareInputs.BokehSizePercent = Settings.LensFlareBokehSize;
-	LensFlareInputs.Intensity = Settings.LensFlareIntensity;
+	LensFlareInputs.Intensity = Settings.LensFlareIntensity * Settings.BloomIntensity;
 	LensFlareInputs.Threshold = Settings.LensFlareThreshold;
 
 	// If a bloom output texture isn't available, substitute the half resolution scene color instead, but disable bloom
diff --git a/Engine/Source/Runtime/Renderer/Private/PostProcess/PostProcessing.cpp b/Engine/Source/Runtime/Renderer/Private/PostProcess/PostProcessing.cpp
--- a/Engine/Source/Runtime/Renderer/Private/PostProcess/PostProcessing.cpp
+++ b/Engine/Source/Runtime/Renderer/Private/PostProcess/PostProcessing.cpp
@@ -771,7 +771,7 @@ void AddPostProcessingPasses(
 
 		PassSequence.Finalize();
 
-		const bool bLensFlareEnabled = IsLensFlaresEnabled(View) && !bVisualizeTSR;
+		const bool bLensFlareEnabled = bBloomEnabled && IsLensFlaresEnabled(View);
 		const bool bFFTBloomEnabled = bBloomEnabled && IsFFTBloomEnabled(View);
 
 		const bool bBasicEyeAdaptationEnabled = bEyeAdaptationEnabled && (AutoExposureMethod == EAutoExposureMethod::AEM_Basic);
@@ -786,7 +786,7 @@ void AddPostProcessingPasses(
 		const bool bProduceSceneColorChain = (
 			bBasicEyeAdaptationEnabled ||
 			(bBloomEnabled && !bFFTBloomEnabled) ||
-			bLensFlareEnabled ||
+			(bLensFlareEnabled && bFFTBloomEnabled) ||
 			bLocalExposureBlurredLum);
 		extern int32 GSSRHalfResSceneColor;
 		const bool bNeedBeforeBloomHalfRes    = (!bProcessQuarterResolution && !bProcessEighthResolution) || (bFFTBloomEnabled && FFTBloomResolutionFraction > 0.25f && FFTBloomResolutionFraction <= 0.5f) || (ReflectionsMethod == EReflectionsMethod::SSR && !View.bStatePrevViewInfoIsReadOnly && GSSRHalfResSceneColor);
@@ -1281,95 +1281,99 @@ void AddPostProcessingPasses(
 
 		FScreenPassTexture Bloom;
 		FRDGBufferRef SceneColorApplyParameters = nullptr;
+		if (bBloomEnabled)
 		{
-			const FTextureDownsampleChain* LensFlareSceneDownsampleChain = &SceneDownsampleChain;
+			const FTextureDownsampleChain* LensFlareSceneDownsampleChain;
 
 			FTextureDownsampleChain BloomDownsampleChain;
 
-			if (bBloomEnabled)
+			if (bFFTBloomEnabled)
 			{
-				if (bFFTBloomEnabled)
-				{
-					float InputResolutionFraction;
-					FScreenPassTextureSlice InputSceneColor;
+				LensFlareSceneDownsampleChain = &SceneDownsampleChain;
 
-					if (FFTBloomResolutionFraction <= 0.125f)
-					{
-						InputSceneColor = EighthResSceneColor;
-						InputResolutionFraction = 0.125f;
-					}
-					else if (FFTBloomResolutionFraction <= 0.25f)
-					{
-						InputSceneColor = QuarterResSceneColor;
-						InputResolutionFraction = 0.25f;
-					}
-					else if (FFTBloomResolutionFraction <= 0.5f)
-					{
-						InputSceneColor = HalfResSceneColor;
-						InputResolutionFraction = 0.5f;
-					}
-					else
-					{
-						InputSceneColor = SceneColorSlice;
-						InputResolutionFraction = 1.0f;
-					}
+				float InputResolutionFraction;
+				FScreenPassTextureSlice InputSceneColor;
 
-					FFFTBloomOutput Outputs = AddFFTBloomPass(
-						GraphBuilder,
-						View,
-						InputSceneColor,
-						InputResolutionFraction,
-						EyeAdaptationParameters,
-						EyeAdaptationBuffer,
-						LocalExposureParameters,
-						CVarBloomApplyLocalExposure.GetValueOnRenderThread() ? LocalExposureBilateralGridTexture : nullptr,
-						LocalExposureBlurredLogLumTexture);
-
-					Bloom = Outputs.BloomTexture;
-					SceneColorApplyParameters = Outputs.SceneColorApplyParameters;
+				if (FFTBloomResolutionFraction <= 0.125f)
+				{
+					InputSceneColor = EighthResSceneColor;
+					InputResolutionFraction = 0.125f;
+				}
+				else if (FFTBloomResolutionFraction <= 0.25f)
+				{
+					InputSceneColor = QuarterResSceneColor;
+					InputResolutionFraction = 0.25f;
+				}
+				else if (FFTBloomResolutionFraction <= 0.5f)
+				{
+					InputSceneColor = HalfResSceneColor;
+					InputResolutionFraction = 0.5f;
 				}
 				else
 				{
-					const bool bApplyLocalExposureToBloom =
-						CVarBloomApplyLocalExposure.GetValueOnRenderThread()
-						&& View.FinalPostProcessSettings.LocalExposureMethod == ELocalExposureMethod::Bilateral
-						&& LocalExposureBilateralGridTexture != nullptr;
+					InputSceneColor = SceneColorSlice;
+					InputResolutionFraction = 1.0f;
+				}
 
-					const bool bBloomSetupRequiredEnabled = View.FinalPostProcessSettings.BloomThreshold > -1.0f || bApplyLocalExposureToBloom;
+				FFFTBloomOutput Outputs = AddFFTBloomPass(
+					GraphBuilder, 
+					View,
+					InputSceneColor,
+					InputResolutionFraction,
+					EyeAdaptationParameters,
+					EyeAdaptationBuffer,
+					LocalExposureParameters,
+					CVarBloomApplyLocalExposure.GetValueOnRenderThread() ? LocalExposureBilateralGridTexture : nullptr,
+					LocalExposureBlurredLogLumTexture);
 
-					// Reuse the main scene downsample chain if setup isn't required for gaussian bloom.
-					if (!SceneDownsampleChain.IsInitialized() || bBloomSetupRequiredEnabled)
-					{
-						FScreenPassTextureSlice DownsampleInput = bProcessEighthResolution ? EighthResSceneColor : (bProcessQuarterResolution ? QuarterResSceneColor : HalfResSceneColor);
+				Bloom = Outputs.BloomTexture;
+				SceneColorApplyParameters = Outputs.SceneColorApplyParameters;
+			}
+			else
+			{
+				const bool bApplyLocalExposureToBloom = 
+					CVarBloomApplyLocalExposure.GetValueOnRenderThread()
+					&& View.FinalPostProcessSettings.LocalExposureMethod == ELocalExposureMethod::Bilateral
+					&& LocalExposureBilateralGridTexture != nullptr;
 
-						if (bBloomSetupRequiredEnabled)
-						{
-							const float BloomThreshold = View.FinalPostProcessSettings.BloomThreshold;
+				const bool bBloomSetupRequiredEnabled = View.FinalPostProcessSettings.BloomThreshold > -1.0f || bApplyLocalExposureToBloom;
 
-							FBloomSetupInputs SetupPassInputs;
-							SetupPassInputs.SceneColor = DownsampleInput;
-							SetupPassInputs.EyeAdaptationBuffer = EyeAdaptationBuffer;
-							SetupPassInputs.EyeAdaptationParameters = &EyeAdaptationParameters;
-							SetupPassInputs.Threshold = BloomThreshold;
+				// Reuse the main scene downsample chain if setup isn't required for gaussian bloom.
+				if (SceneDownsampleChain.IsInitialized() && !bBloomSetupRequiredEnabled)
+				{
+					LensFlareSceneDownsampleChain = &SceneDownsampleChain;
+				}
+				else
+				{
+					FScreenPassTextureSlice DownsampleInput = bProcessEighthResolution ? EighthResSceneColor : (bProcessQuarterResolution ? QuarterResSceneColor : HalfResSceneColor);
 
-							if (bApplyLocalExposureToBloom)
-							{
-								SetupPassInputs.LocalExposureParameters = &LocalExposureParameters;
-								SetupPassInputs.LocalExposureTexture = LocalExposureBilateralGridTexture;
-								SetupPassInputs.BlurredLogLuminanceTexture = LocalExposureBlurredLogLumTexture;
-							}
+					if (bBloomSetupRequiredEnabled)
+					{
+						const float BloomThreshold = View.FinalPostProcessSettings.BloomThreshold;
 
-							DownsampleInput = FScreenPassTextureSlice::CreateFromScreenPassTexture(GraphBuilder, AddBloomSetupPass(GraphBuilder, View, SetupPassInputs));
-						}
+						FBloomSetupInputs SetupPassInputs;
+						SetupPassInputs.SceneColor = DownsampleInput;
+						SetupPassInputs.EyeAdaptationBuffer = EyeAdaptationBuffer;
+						SetupPassInputs.EyeAdaptationParameters = &EyeAdaptationParameters;
+						SetupPassInputs.Threshold = BloomThreshold;
 
-						const bool bLogLumaInAlpha = false;
-						BloomDownsampleChain.Init(GraphBuilder, View, EyeAdaptationParameters, DownsampleInput, DownsampleChainQuality, (uint32)EBloomQuality::MAX, bLogLumaInAlpha, nullptr, bProcessEighthResolution ? 3 : (bProcessQuarterResolution ? 2 : 1));
+						if (bApplyLocalExposureToBloom)
+						{
+							SetupPassInputs.LocalExposureParameters = &LocalExposureParameters;
+							SetupPassInputs.LocalExposureTexture = LocalExposureBilateralGridTexture;
+							SetupPassInputs.BlurredLogLuminanceTexture = LocalExposureBlurredLogLumTexture;
+						}
 
-						LensFlareSceneDownsampleChain = &BloomDownsampleChain;
+						DownsampleInput = FScreenPassTextureSlice::CreateFromScreenPassTexture(GraphBuilder, AddBloomSetupPass(GraphBuilder, View, SetupPassInputs));
 					}
 
-					Bloom = AddGaussianBloomPasses(GraphBuilder, View, LensFlareSceneDownsampleChain);
+					const bool bLogLumaInAlpha = false;
+					BloomDownsampleChain.Init(GraphBuilder, View, EyeAdaptationParameters, DownsampleInput, DownsampleChainQuality, (uint32)EBloomQuality::MAX, bLogLumaInAlpha, nullptr, bProcessEighthResolution ? 3 : (bProcessQuarterResolution ? 2 : 1));
+
+					LensFlareSceneDownsampleChain = &BloomDownsampleChain;
 				}
+
+				Bloom = AddGaussianBloomPasses(GraphBuilder, View, LensFlareSceneDownsampleChain);
 			}
 
 			if (bLensFlareEnabled)
@@ -2592,10 +2596,8 @@ void AddMobilePostProcessingPasses(FRDGBuilder& GraphBuilder, FScene* Scene, con
 	if (bIsPostProcessingEnabled)
 	{
 		bool bUseSun = View.MobileLightShaft.IsSet();
-		
-		const bool bBloomEnabled = View.FinalPostProcessSettings.BloomIntensity > 0.0f;
-		const bool bLensFlareEnabled = IsLensFlaresEnabled(View);
-		const bool bUseBloom = bBloomEnabled || bLensFlareEnabled;
+			
+		bool bUseBloom = View.FinalPostProcessSettings.BloomIntensity > 0.0f;
 
 		bool bUseBasicEyeAdaptation = bUseEyeAdaptation && (AutoExposureMethod == EAutoExposureMethod::AEM_Basic);
 		bool bUseHistogramEyeAdaptation = bUseEyeAdaptation && (AutoExposureMethod == EAutoExposureMethod::AEM_Histogram) &&
@@ -2842,75 +2844,69 @@ void AddMobilePostProcessingPasses(FRDGBuilder& GraphBuilder, FScene* Scene, con
 				return AddMobileBloomDownPass(GraphBuilder, View, BloomDownInputs);
 			};
 
-			constexpr float BloomDownScale = 0.66f * 4.0f;
-			constexpr int32 MaxPasses = 6;
+			const float BloomDownScale = 0.66f * 4.0f;
+			const int MaxPasses = 6;
 			const EBloomQuality BloomQuality = GetBloomQuality();
-			
-			static_assert(static_cast<int32>(ELensFlareQuality::MAX) == 4);
-			const int NumDownsamplePasses = (BloomQuality == EBloomQuality::Q1 || !bBloomEnabled) ? 4 : BloomQuality == EBloomQuality::Q2 ? 5 : 6;
-
+			const int NumDownsamplePasses = BloomQuality == EBloomQuality::Q1 ? 4 : BloomQuality == EBloomQuality::Q2 ? 5 : 6;
 			FScreenPassTexture PostProcessDownsample_Bloom[MaxPasses];
 
 			for (int32 i = 0; i < NumDownsamplePasses; ++i)
 			{
 				PostProcessDownsample_Bloom[i] = AddBloomDownPass(i == 0 ? BloomSetupOutputs.Bloom : PostProcessDownsample_Bloom[i - 1], BloomDownScale);
 			}
-			
-			if (bBloomEnabled)
-			{
-				const FFinalPostProcessSettings& Settings = View.FinalPostProcessSettings;
 
-				auto AddBloomUpPass = [&GraphBuilder, &View](FScreenPassTexture& BloomUpSourceA, FScreenPassTexture& BloomUpSourceB, float BloomSourceScale, const FVector4f& TintA, const FVector4f& TintB)
-					{
-						FMobileBloomUpInputs BloomUpInputs;
-						BloomUpInputs.BloomUpSourceA = BloomUpSourceA;
-						BloomUpInputs.BloomUpSourceB = BloomUpSourceB;
-						BloomUpInputs.ScaleAB = FVector2D(BloomSourceScale, BloomSourceScale);
-						BloomUpInputs.TintA = TintA;
-						BloomUpInputs.TintB = TintB;
+			const FFinalPostProcessSettings& Settings = View.FinalPostProcessSettings;
 
-						return AddMobileBloomUpPass(GraphBuilder, View, BloomUpInputs);
-					};
+			auto AddBloomUpPass = [&GraphBuilder, &View](FScreenPassTexture& BloomUpSourceA, FScreenPassTexture& BloomUpSourceB, float BloomSourceScale, const FVector4f& TintA, const FVector4f& TintB)
+			{
+				FMobileBloomUpInputs BloomUpInputs;
+				BloomUpInputs.BloomUpSourceA = BloomUpSourceA;
+				BloomUpInputs.BloomUpSourceB = BloomUpSourceB;
+				BloomUpInputs.ScaleAB = FVector2D(BloomSourceScale, BloomSourceScale);
+				BloomUpInputs.TintA = TintA;
+				BloomUpInputs.TintB = TintB;
 
-				const float BloomUpScale = 0.66f * 2.0f;
+				return AddMobileBloomUpPass(GraphBuilder, View, BloomUpInputs);
+			};
 
-				const FLinearColor Tint[]{ Settings.Bloom1Tint, Settings.Bloom2Tint, Settings.Bloom3Tint, Settings.Bloom4Tint, Settings.Bloom5Tint, Settings.Bloom6Tint };
+			const float BloomUpScale = 0.66f * 2.0f;
 
-				// Upsample by 2
-				{
-					const int IndexA = NumDownsamplePasses - 2;
-					const int IndexB = NumDownsamplePasses - 1;
-					FVector4f TintA{ Tint[IndexA], 0.0f };
-					FVector4f TintB{ Tint[IndexB], 0.0f };
-					TintA *= Settings.BloomIntensity;
-					TintB *= Settings.BloomIntensity;
-
-					BloomUpOutputs = AddBloomUpPass(PostProcessDownsample_Bloom[IndexA], PostProcessDownsample_Bloom[IndexB], BloomUpScale, TintA, TintB);
-				}
+			const FLinearColor Tint[]{ Settings.Bloom1Tint, Settings.Bloom2Tint, Settings.Bloom3Tint, Settings.Bloom4Tint, Settings.Bloom5Tint, Settings.Bloom6Tint };
 
-				for (int Index = NumDownsamplePasses - 3; Index > 0; --Index)
-				{
-					// Upsample by 2
-					FVector4f TintA{ Tint[Index], 0.0f };
-					TintA *= Settings.BloomIntensity;
-					FVector4f TintB = FVector4f(1.0f, 1.0f, 1.0f, 0.0f);
+			// Upsample by 2
+			{
+				const int IndexA = NumDownsamplePasses - 2;
+				const int IndexB = NumDownsamplePasses - 1;
+				FVector4f TintA{ Tint[IndexA], 0.0f };
+				FVector4f TintB{ Tint[IndexB], 0.0f };
+				TintA *= Settings.BloomIntensity;
+				TintB *= Settings.BloomIntensity;
 
-					BloomUpOutputs = AddBloomUpPass(PostProcessDownsample_Bloom[Index], BloomUpOutputs, BloomUpScale, TintA, TintB);
-				}
+				BloomUpOutputs = AddBloomUpPass(PostProcessDownsample_Bloom[IndexA], PostProcessDownsample_Bloom[IndexB], BloomUpScale, TintA, TintB);
+			}
 
+			for (int Index = NumDownsamplePasses - 3; Index > 0; --Index)
+			{
 				// Upsample by 2
-				{
-					FVector4f TintA = FVector4f(Settings.Bloom2Tint.R, Settings.Bloom2Tint.G, Settings.Bloom2Tint.B, 0.0f);
-					TintA *= Settings.BloomIntensity;
-					// Scaling Bloom2 by extra factor to match filter area difference between PC default and mobile.
-					TintA *= 0.5;
-					FVector4f TintB = FVector4f(1.0f, 1.0f, 1.0f, 0.0f);
+				FVector4f TintA{ Tint[Index], 0.0f };
+				TintA *= Settings.BloomIntensity;
+				FVector4f TintB = FVector4f(1.0f, 1.0f, 1.0f, 0.0f);
 
-					BloomUpOutputs = AddBloomUpPass(PostProcessDownsample_Bloom[0], BloomUpOutputs, BloomUpScale, TintA, TintB);
-				}
+				BloomUpOutputs = AddBloomUpPass(PostProcessDownsample_Bloom[Index], BloomUpOutputs, BloomUpScale, TintA, TintB);
 			}
 
-			if (bLensFlareEnabled)
+			// Upsample by 2
+			{
+				FVector4f TintA = FVector4f(Settings.Bloom2Tint.R, Settings.Bloom2Tint.G, Settings.Bloom2Tint.B, 0.0f);
+				TintA *= Settings.BloomIntensity;
+				// Scaling Bloom2 by extra factor to match filter area difference between PC default and mobile.
+				TintA *= 0.5;
+				FVector4f TintB = FVector4f(1.0f, 1.0f, 1.0f, 0.0f);
+
+				BloomUpOutputs = AddBloomUpPass(PostProcessDownsample_Bloom[0], BloomUpOutputs, BloomUpScale, TintA, TintB);
+			}
+
+			if (IsLensFlaresEnabled(View))
 			{
 				const ELensFlareQuality LensFlareQuality = GetLensFlareQuality();
 				const uint32 LensFlareDownsampleStageIndex = static_cast<uint32>(ELensFlareQuality::MAX) - static_cast<uint32>(LensFlareQuality) - 1;
