Heterogeneous Volumes: Fixing stereo-related issues

1. Iterating through primary view volumetric mesh batches when rendering from a stereo view.
2. Adding downsampled ViewRectMin to allow for correct stereo rendering.

diff --git a/Engine/Shaders/Private/HeterogeneousVolumes/HeterogeneousVolumesComposite.usf b/Engine/Shaders/Private/HeterogeneousVolumes/HeterogeneousVolumesComposite.usf
--- a/Engine/Shaders/Private/HeterogeneousVolumes/HeterogeneousVolumesComposite.usf
+++ b/Engine/Shaders/Private/HeterogeneousVolumes/HeterogeneousVolumesComposite.usf
@@ -40,6 +40,7 @@ RWTexture2D<float4> RWUpsampledTexture;
 SamplerState TextureSampler;
 uint2 UpsampledResolution;
 uint2 UpsampledViewRect;
+uint2 UpsampledViewRectMin;
 
 [numthreads(THREADGROUP_SIZE_2D, THREADGROUP_SIZE_2D, 1)]
 void HeterogeneousVolumesUpsampleCS(
@@ -51,11 +52,11 @@ void HeterogeneousVolumesUpsampleCS(
 	{
 		return;
 	}
-	int2 UpsampledPixelCoord = DispatchThreadId.xy; // TODO: UpsampledViewRectMin
+	int2 UpsampledPixelCoord = DispatchThreadId.xy + UpsampledViewRectMin;
 
 #if UPSAMPLE_MODE == UPSAMPLE_MODE_BILINEAR
 	float2 UV = (UpsampledPixelCoord + 0.5 * DownsampleFactor) / float2(UpsampledResolution);
-	UV = clamp(UV, 0, (UpsampledViewRect - 1) / float2(UpsampledResolution));
+	UV = clamp(UV, 0, (UpsampledViewRect + UpsampledViewRectMin - 1) / float2(UpsampledResolution));
 	float4 Result = Texture2DSample(HeterogeneousVolumeRadiance, TextureSampler, UV);
 #else
 	int2 DownsampledPixelCoord = UpsampledPixelCoord / DownsampleFactor;
@@ -79,7 +80,7 @@ void HeterogeneousVolumesFilterCS(
 		return;
 	}
 
-	int2 PixelCoord = DispatchThreadId.xy; // TODO: UpsampledViewRectMin
+	int2 PixelCoord = DispatchThreadId.xy + UpsampledViewRectMin;
 	float4 Radiance = HeterogeneousVolumeRadiance[PixelCoord];
 	float HeterogeneousVolumeTransmittance = Radiance.a;
 
@@ -98,7 +99,7 @@ void HeterogeneousVolumesFilterCS(
 			{
 				int2 Offset = int2(X, Y) - FilterHalfWidth;
 				int2 SampleCoord = PixelCoord + Offset;
-				SampleCoord = clamp(SampleCoord, 0, int2(UpsampledViewRect) - 1);
+				SampleCoord = clamp(SampleCoord, 0, int2(UpsampledViewRect + UpsampledViewRectMin) - 1);
 				float FilterWeight = max(FilterRadius2 - dot(Offset, Offset), 0.0);
 
 				float4 Value = HeterogeneousVolumeRadiance[SampleCoord];
@@ -135,7 +136,7 @@ void HeterogeneousVolumesFilterCS(
 			{
 				int2 OffsetCoord = int2(X, Y);
 				int2 SampleCoord = PixelCoord + OffsetCoord - FilterHalfWidth;
-				SampleCoord = clamp(SampleCoord, 0, int2(UpsampledViewRect) - 1);
+				SampleCoord = clamp(SampleCoord, 0, int2(UpsampledViewRect + UpsampledViewRectMin) - 1);
 				float FilterWeight = Gaussian3x3Weights[OffsetCoord.y * 3 + OffsetCoord.x];
 
 				Radiance += HeterogeneousVolumeRadiance[SampleCoord] * FilterWeight;
@@ -166,7 +167,7 @@ void HeterogeneousVolumesFilterCS(
 			{
 				int2 OffsetCoord = int2(X, Y);
 				int2 SampleCoord = PixelCoord + OffsetCoord - FilterHalfWidth;
-				SampleCoord = clamp(SampleCoord, 0, int2(UpsampledViewRect) - 1);
+				SampleCoord = clamp(SampleCoord, 0, int2(UpsampledViewRect + UpsampledViewRectMin) - 1);
 				float FilterWeight = Gaussian5x5Weights[OffsetCoord.y * 5 + OffsetCoord.x];
 
 				Radiance += HeterogeneousVolumeRadiance[SampleCoord] * FilterWeight;
@@ -181,6 +182,16 @@ void HeterogeneousVolumesFilterCS(
 	RWFilteredTexture[PixelCoord] = Radiance;
 }
 
+int2 GetDownsampledResolution(int2 Resolution, int Factor)
+{
+	return (Resolution + Factor - 1) / Factor;
+}
+
+int2 GetScaledViewRect()
+{
+	return GetDownsampledResolution(View.ViewSizeAndInvSize.xy, DownsampleFactor);
+}
+
 [numthreads(THREADGROUP_SIZE_2D, THREADGROUP_SIZE_2D, 1)]
 void HeterogeneousVolumesCompositeCS(
 	uint2 GroupThreadId : SV_GroupThreadID,
@@ -188,11 +199,11 @@ void HeterogeneousVolumesCompositeCS(
 )
 {
 	ResolvedView = ResolveView();
-	if (any(DispatchThreadId.xy >= ResolvedView.ViewSizeAndInvSize.xy))
+	if (any(DispatchThreadId >= View.ViewSizeAndInvSize.xy))
 	{
 		return;
 	}
-	uint2 ViewPixelCoord = DispatchThreadId.xy; // TODO: ViewRectMin
+	uint2 ViewPixelCoord = DispatchThreadId.xy + View.ViewRectMin.xy;
 	uint2 PixelCoord = ViewPixelCoord / DownsampleFactor;
 
 	// Cache Sample values
diff --git a/Engine/Shaders/Private/HeterogeneousVolumes/HeterogeneousVolumesLiveShadingPipeline.usf b/Engine/Shaders/Private/HeterogeneousVolumes/HeterogeneousVolumesLiveShadingPipeline.usf
--- a/Engine/Shaders/Private/HeterogeneousVolumes/HeterogeneousVolumesLiveShadingPipeline.usf
+++ b/Engine/Shaders/Private/HeterogeneousVolumes/HeterogeneousVolumesLiveShadingPipeline.usf
@@ -637,7 +637,9 @@ void RenderSingleScatteringWithLiveShading(
 	}
 
 	uint LinearIndex = DispatchThreadId.y * GroupCount.x * THREADGROUP_SIZE_2D + DispatchThreadId.x;
-	uint2 PixelCoord = DispatchThreadId.xy;
+	uint2 ViewPixelCoord = DispatchThreadId.xy * DownsampleFactor + View.ViewRectMin.xy;
+	uint2 PixelCoord = DispatchThreadId.xy + View.ViewRectMin.xy / DownsampleFactor;
+
 	//float Epsilon = 1.0e-5;
 	//if (RWLightingTexture[PixelCoord].a < Epsilon)
 	//if (RWLightingTexture[PixelCoord].a <= 0.0)
@@ -651,7 +653,6 @@ void RenderSingleScatteringWithLiveShading(
 	FBeerShadowMapData BeerShadowMapData = InitializeBeerShadowMapData(PrevTransmittance);
 
 	// Extract depth
-	uint2 ViewPixelCoord = PixelCoord * DownsampleFactor + View.ViewRectMin.xy / DownsampleFactor;
 	float DeviceZ = SceneDepthTexture.Load(int3(ViewPixelCoord, 0)).r;
 #if HAS_INVERTED_Z_BUFFER
 	DeviceZ = max(0.000000000001, DeviceZ);
diff --git a/Engine/Source/Runtime/Renderer/Private/HeterogeneousVolumes/HeterogeneousVolumes.cpp b/Engine/Source/Runtime/Renderer/Private/HeterogeneousVolumes/HeterogeneousVolumes.cpp
--- a/Engine/Source/Runtime/Renderer/Private/HeterogeneousVolumes/HeterogeneousVolumes.cpp
+++ b/Engine/Source/Runtime/Renderer/Private/HeterogeneousVolumes/HeterogeneousVolumes.cpp
@@ -393,12 +393,15 @@ bool ShouldRenderHeterogeneousVolumesForAnyView(
 	const TArrayView<FViewInfo>& Views
 )
 {
-	bool Result = false;
-	for (FViewInfo& View : Views)
+	for (int32 ViewIndex = 0; ViewIndex < Views.Num(); ++ViewIndex)
 	{
-		Result |= ShouldRenderHeterogeneousVolumesForView(View);
+		if (ShouldRenderHeterogeneousVolumesForView(Views, ViewIndex))
+		{
+			return true;
+		}
 	}
-	return Result;
+
+	return false;
 }
 
 bool ShouldRenderHeterogeneousVolumesForView(
@@ -406,22 +409,79 @@ bool ShouldRenderHeterogeneousVolumesForView(
 )
 {
 	return IsHeterogeneousVolumesEnabled()
+		&& IStereoRendering::IsAPrimaryView(View)
 		&& !View.HeterogeneousVolumesMeshBatches.IsEmpty()
 		&& View.Family
 		&& !View.bIsReflectionCapture;
 }
 
+int32 GetPrimaryViewIndex(const TArrayView<FViewInfo>& Views)
+{
+	int32 PrimaryViewIndex = Views.Num() - 1;
+	while (PrimaryViewIndex >= 0)
+	{
+		if (IStereoRendering::IsAPrimaryView(Views[PrimaryViewIndex]))
+		{
+			break;
+		}
+		PrimaryViewIndex--;
+	}
+
+	return PrimaryViewIndex;
+}
+
+FViewInfo* GetPrimaryViewForView(
+	const TArrayView<FViewInfo>& Views,
+	int32 ViewIndex
+)
+{
+	check(ViewIndex < Views.Num())
+	FViewInfo* PrimaryView = &Views[ViewIndex];
+
+	// For stereo views, mesh batches will only be defined on the primary
+	if (PrimaryView->bIsInstancedStereoEnabled && IStereoRendering::IsStereoEyeView(*PrimaryView) && !IStereoRendering::IsAPrimaryView(*PrimaryView))
+	{
+		int32 PrimaryViewIndex = GetPrimaryViewIndex(Views);
+		check(PrimaryViewIndex >= 0);
+		PrimaryView = &Views[PrimaryViewIndex];
+	}
+
+	check(IStereoRendering::IsAPrimaryView(*PrimaryView))
+	return PrimaryView;
+}
+
+bool ShouldRenderHeterogeneousVolumesForView(
+	const TArrayView<FViewInfo>& Views,
+	int32 ViewIndex
+)
+{
+	FViewInfo* PrimaryView = GetPrimaryViewForView(Views, ViewIndex);
+	return ShouldRenderHeterogeneousVolumesForView(*PrimaryView);
+}
+
+TArray<FVolumetricMeshBatch, SceneRenderingAllocator> GetHeterogeneousVolumesMeshBatches(
+	const TArrayView<FViewInfo>& Views,
+	int32 ViewIndex
+)
+{
+	FViewInfo* PrimaryView = GetPrimaryViewForView(Views, ViewIndex);
+	return PrimaryView->HeterogeneousVolumesMeshBatches;
+}
+
 bool ShouldRenderHeterogeneousVolumesAsHoldoutForView(
 	const FViewInfo& View
 )
 {
+	check(IStereoRendering::IsAPrimaryView(View));
+
 	// This query returns true if any volume is marked as a holdout; otherwise, the query returns false
 	if (ShouldRenderHeterogeneousVolumesForView(View))
 	{
-		for (int32 MeshBatchIndex = 0; MeshBatchIndex < View.HeterogeneousVolumesMeshBatches.Num(); ++MeshBatchIndex)
+		const TArray<FVolumetricMeshBatch, SceneRenderingAllocator>& MeshBatches = View.HeterogeneousVolumesMeshBatches;
+		for (int32 MeshBatchIndex = 0; MeshBatchIndex < MeshBatches.Num(); ++MeshBatchIndex)
 		{
-			const FMeshBatch* Mesh = View.HeterogeneousVolumesMeshBatches[MeshBatchIndex].Mesh;
-			const FPrimitiveSceneProxy* PrimitiveSceneProxy = View.HeterogeneousVolumesMeshBatches[MeshBatchIndex].Proxy;
+			const FMeshBatch* Mesh = MeshBatches[MeshBatchIndex].Mesh;
+			const FPrimitiveSceneProxy* PrimitiveSceneProxy = MeshBatches[MeshBatchIndex].Proxy;
 			if (ShouldRenderMeshBatchWithHeterogeneousVolumes(Mesh, PrimitiveSceneProxy, View.GetFeatureLevel()))
 			{
 				for (int32 VolumeIndex = 0; VolumeIndex < Mesh->Elements.Num(); ++VolumeIndex)
@@ -1024,7 +1084,7 @@ void FDeferredShadingSceneRenderer::RenderHeterogeneousVolumes(
 	for (int32 ViewIndex = 0; ViewIndex < Views.Num(); ViewIndex++)
 	{
 		FViewInfo& View = Views[ViewIndex];
-		if (ShouldRenderHeterogeneousVolumesForView(View))
+		if (ShouldRenderHeterogeneousVolumesForView(Views, ViewIndex))
 		{
 			if (HeterogeneousVolumes::GetDebugMode() != 0)
 			{
@@ -1054,10 +1114,12 @@ void FDeferredShadingSceneRenderer::RenderHeterogeneousVolumes(
 				};
 
 				TArray<VolumeMesh> VolumeMeshes;
-				for (int32 MeshBatchIndex = 0; MeshBatchIndex < View.HeterogeneousVolumesMeshBatches.Num(); ++MeshBatchIndex)
+
+				const TArray<FVolumetricMeshBatch, SceneRenderingAllocator>& MeshBatches = GetHeterogeneousVolumesMeshBatches(Views, ViewIndex);
+				for (int32 MeshBatchIndex = 0; MeshBatchIndex < MeshBatches.Num(); ++MeshBatchIndex)
 				{
-					const FMeshBatch* Mesh = View.HeterogeneousVolumesMeshBatches[MeshBatchIndex].Mesh;
-					const FPrimitiveSceneProxy* PrimitiveSceneProxy = View.HeterogeneousVolumesMeshBatches[MeshBatchIndex].Proxy;
+					const FMeshBatch* Mesh = MeshBatches[MeshBatchIndex].Mesh;
+					const FPrimitiveSceneProxy* PrimitiveSceneProxy = MeshBatches[MeshBatchIndex].Proxy;
 					if (!ShouldRenderMeshBatchWithHeterogeneousVolumes(Mesh, PrimitiveSceneProxy, View.GetFeatureLevel()))
 					{
 						continue;
@@ -1283,6 +1345,7 @@ class FHeterogeneousVolumesUpsampleCS : public FGlobalShader
 		SHADER_PARAMETER_SAMPLER(SamplerState, TextureSampler)
 		SHADER_PARAMETER(FIntPoint, UpsampledResolution)
 		SHADER_PARAMETER(FIntPoint, UpsampledViewRect)
+		SHADER_PARAMETER(FIntPoint, UpsampledViewRectMin)
 		SHADER_PARAMETER(int, DownsampleFactor)
 
 		// Output
@@ -1344,6 +1407,7 @@ class FHeterogeneousVolumesFilterCS : public FGlobalShader
 		SHADER_PARAMETER_SAMPLER(SamplerState, TextureSampler)
 		SHADER_PARAMETER(FIntPoint, UpsampledResolution)
 		SHADER_PARAMETER(FIntPoint, UpsampledViewRect)
+		SHADER_PARAMETER(FIntPoint, UpsampledViewRectMin)
 		SHADER_PARAMETER(int, DownsampleFactor)
 		SHADER_PARAMETER(int, FilterWidth)
 
@@ -1460,7 +1524,7 @@ void FDeferredShadingSceneRenderer::CompositeHeterogeneousVolumes(
 	{
 		FViewInfo& View = Views[ViewIndex];
 
-		if (ShouldRenderHeterogeneousVolumesForView(View))
+		if (ShouldRenderHeterogeneousVolumesForView(Views, ViewIndex))
 		{
 			FRDGTextureDesc UpsampleDesc = View.HeterogeneousVolumeRadiance->Desc;
 			int32 DownsampleFactor = HeterogeneousVolumes::GetDownsampleFactor();
@@ -1471,7 +1535,9 @@ void FDeferredShadingSceneRenderer::CompositeHeterogeneousVolumes(
 				int32 DownsampleIterationFactor = (DownsampleFactor % 3 == 0) ? 3 : 2;
 
 				FIntPoint DownsampledViewRect = FIntPoint::DivideAndRoundUp(View.ViewRect.Size(), DownsampleFactor);
+				FIntPoint DownsampledViewRectMin = FIntPoint::DivideAndRoundUp(View.ViewRect.Min, DownsampleFactor);
 				FIntPoint UpsampledViewRect = DownsampledViewRect * DownsampleIterationFactor;
+				FIntPoint UpsampledViewRectMin = DownsampledViewRectMin * DownsampleIterationFactor;
 				FIntVector GroupCount = FComputeShaderUtils::GetGroupCount(UpsampledViewRect, FHeterogeneousVolumesUpsampleCS::GetThreadGroupSize2D());
 				UpsampleDesc.Extent *= DownsampleIterationFactor;
 
@@ -1490,6 +1556,7 @@ void FDeferredShadingSceneRenderer::CompositeHeterogeneousVolumes(
 						PassParameters->TextureSampler = TStaticSamplerState<SF_Bilinear, AM_Clamp, AM_Clamp, AM_Clamp>::GetRHI();
 						PassParameters->UpsampledResolution = UpsampleDesc.Extent;
 						PassParameters->UpsampledViewRect = UpsampledViewRect;
+						PassParameters->UpsampledViewRectMin = UpsampledViewRectMin;
 						PassParameters->DownsampleFactor = DownsampleIterationFactor;
 
 						// Output
@@ -1524,6 +1591,7 @@ void FDeferredShadingSceneRenderer::CompositeHeterogeneousVolumes(
 						PassParameters->TextureSampler = TStaticSamplerState<SF_Bilinear, AM_Clamp, AM_Clamp, AM_Clamp>::GetRHI();
 						PassParameters->UpsampledResolution = UpsampleDesc.Extent;
 						PassParameters->UpsampledViewRect = UpsampledViewRect;
+						PassParameters->UpsampledViewRectMin = UpsampledViewRectMin;
 						PassParameters->DownsampleFactor = 1;
 						PassParameters->FilterWidth = HeterogeneousVolumes::GetFilterWidth();
 
diff --git a/Engine/Source/Runtime/Renderer/Private/HeterogeneousVolumes/HeterogeneousVolumes.h b/Engine/Source/Runtime/Renderer/Private/HeterogeneousVolumes/HeterogeneousVolumes.h
--- a/Engine/Source/Runtime/Renderer/Private/HeterogeneousVolumes/HeterogeneousVolumes.h
+++ b/Engine/Source/Runtime/Renderer/Private/HeterogeneousVolumes/HeterogeneousVolumes.h
@@ -35,6 +35,7 @@ bool ShouldRenderHeterogeneousVolumes(const FScene* Scene);
 bool ShouldRenderHeterogeneousVolumesForAnyView(const TArrayView<FViewInfo>& Views);
 bool ShouldRenderHeterogeneousVolumesForView(const FViewInfo& View);
 bool ShouldRenderHeterogeneousVolumesAsHoldoutForView(const FViewInfo& View);
+bool ShouldRenderHeterogeneousVolumesForView(const TArrayView<FViewInfo>& Views, int32 ViewIndex);
 bool DoesMaterialShaderSupportHeterogeneousVolumes(const FMaterialShaderParameters& Parameters);
 bool DoesMaterialShaderSupportHeterogeneousVolumes(const FMaterial& Material);
 bool ShouldRenderMeshBatchWithHeterogeneousVolumes(
