Fix virtual texture locking request for UDIM tiles.
We were only locking for the first block, but should lock each one.
#rnx

diff --git a/Engine/Source/Runtime/Renderer/Private/VT/VirtualTextureSystem.cpp b/Engine/Source/Runtime/Renderer/Private/VT/VirtualTextureSystem.cpp
--- a/Engine/Source/Runtime/Renderer/Private/VT/VirtualTextureSystem.cpp
+++ b/Engine/Source/Runtime/Renderer/Private/VT/VirtualTextureSystem.cpp
@@ -1340,8 +1340,8 @@ void FVirtualTextureSystem::SetMipLevelToLock(FVirtualTextureProducerHandle Prod
 	for (int32 MipIndex = FirstMip; MipIndex > LastMip; --MipIndex)
 	{
 		const int32 vLevel = MaxLoadLevel - MipIndex;
-		const int32 MipWidthInTiles = 1 << MipIndex;
-		const int32 MipHeightInTiles = 1 << MipIndex;
+		const int32 MipWidthInTiles = Producer.GetDescription().WidthInBlocks << MipIndex;
+		const int32 MipHeightInTiles = Producer.GetDescription().HeightInBlocks << MipIndex;
 		for (int32 PosY = 0; PosY < MipHeightInTiles; PosY++)
 		{
 			for (int32 PosX = 0; PosX < MipWidthInTiles; PosX++)
