Fix for possible invalid reference of RelevantDecals.List passed to a task that can happen if RelevantDecalsMap.Emplace reallocate the container

diff --git a/Engine/Source/Runtime/Renderer/Private/DecalRenderingShared.cpp b/Engine/Source/Runtime/Renderer/Private/DecalRenderingShared.cpp
--- a/Engine/Source/Runtime/Renderer/Private/DecalRenderingShared.cpp
+++ b/Engine/Source/Runtime/Renderer/Private/DecalRenderingShared.cpp
@@ -75,6 +75,8 @@ FDecalVisibilityViewPacket::FDecalVisibilityViewPacket(const FDecalVisibilityTas
 	: TaskData(InTaskData)
 	, View(InView)
 {
+	RelevantDecalsMap.Reserve((int32)EDecalRenderStage::Num);
+
 	VisibleDecals.Task = LaunchSceneRenderTask(TEXT("BuildVisibleDecalList"), [&OutputList = VisibleDecals.List, &InputList = Scene.Decals, &View = View]
 	{
 		OutputList = DecalRendering::BuildVisibleDecalList(InputList, View);
