SkyLight
- Real time capture now render at some +8EV by default (range becomes -[16;16] EV to [-8;24] EV). This fixes issues with clipped sky pixel values when using large physically based value for sun directional light illuminance.
- A constant was already there to play with this. The read sky cube map values are scaled back up at fetch time byusing SkyLightColor.
- Fix skypass being clamped to half*0.5.
- Updated SkyLight Illuminance visualisation pass.
- Adapted PathTracer

diff --git a/Engine/Shaders/Private/ReflectionEnvironmentShaders.usf b/Engine/Shaders/Private/ReflectionEnvironmentShaders.usf
--- a/Engine/Shaders/Private/ReflectionEnvironmentShaders.usf
+++ b/Engine/Shaders/Private/ReflectionEnvironmentShaders.usf
@@ -982,14 +982,6 @@ float3 GetIlluminance()
 	return IlluminanceAcc.rgb;
 }
 
-float GetSimpleAutoExposure()
-{
-	float3 IlluminanceAcc = GetIlluminance();
-	float MaxIlluminance = max3(IlluminanceAcc.r, IlluminanceAcc.g, IlluminanceAcc.b);
-	float SimpleAutoExposure = 1.0f / (MaxIlluminance <= 0.0 ? 0.01f : MaxIlluminance);
-	return 1.0;// SimpleAutoExposure; Disabled for now, using the view exposure
-}
-
 float4 SampleCubeAsLatLon(TextureCube CubeTexture, SamplerState CubeSampler, float2 UV)
 {
 	float HAngle = UV.x * 2.0 * PI;
@@ -998,7 +990,6 @@ float4 SampleCubeAsLatLon(TextureCube CubeTexture, SamplerState CubeSampler, flo
 	float3 DirFromLatLong = float3(cos(HAngle) * CosVAngle, sin(HAngle) * CosVAngle, -sin(VAngle));
 
 	float4 DebugVis = TextureCubeSampleLevel(CubeTexture, CubeSampler, DirFromLatLong, 0);
-	DebugVis.rgb *= GetSimpleAutoExposure();
 	return DebugVis;
 }
 
@@ -1022,7 +1013,6 @@ float4 SampleCubeAsSphere(TextureCube CubeTexture, SamplerState CubeSampler, flo
 		}
 
 		DebugVis = TextureCubeSampleLevel(CubeTexture, CubeSampler, mul(N, (float3x3)ResolvedView.ViewToTranslatedWorld), 0);
-		DebugVis.rgb *= GetSimpleAutoExposure();
 	}
 
 	return DebugVis;
@@ -1117,6 +1107,7 @@ void VisualizeIlluminanceMeterPS(
 	{
 		float2 UV = (SVPos.xy - float2(SkyMapPreviewOrigin)) / float2(SkyMapPreviewSize);
 		float4 DebugVis = SampleCubeAsLatLon(SkyLightCubeTexture, SkyLightCubeSampler, UV);
+		DebugVis.rgb *= ResolvedView.SkyLightColor.rgb;
 		OutColor = float4(DebugVis.rgb * ResolvedView.PreExposure, 0.0);
 	}
 	else if (all(uint2(SVPos.xy) >= SphereIlluMapPreviewOrigin) && all(uint2(SVPos.xy) <= (SphereIlluMapPreviewOrigin + SphereIlluMapPreviewSize)))
@@ -1130,6 +1121,7 @@ void VisualizeIlluminanceMeterPS(
 	{
 		float2 UV = (SVPos.xy - float2(SphereSkyMapPreviewOrigin)) / float2(SphereSkyMapPreviewSize);
 		float4 DebugVis = SampleCubeAsSphere(SkyLightCubeTexture, SkyLightCubeSampler, UV);
+		DebugVis.rgb *= ResolvedView.SkyLightColor.rgb;
 		OutColor = float4(DebugVis.rgb * ResolvedView.PreExposure, 0.0);
 	}
 	else
diff --git a/Engine/Source/Runtime/Renderer/Private/PathTracing.cpp b/Engine/Source/Runtime/Renderer/Private/PathTracing.cpp
--- a/Engine/Source/Runtime/Renderer/Private/PathTracing.cpp
+++ b/Engine/Source/Runtime/Renderer/Private/PathTracing.cpp
@@ -2718,6 +2718,14 @@ bool PrepareSkyTexture(FRDGBuilder& GraphBuilder, FScene* Scene, const FViewInfo
 	// should we remember the skylight prep for the next frame?
 	const bool IsSkylightCachingEnabled = CVarPathTracingSkylightCaching.GetValueOnAnyThread() != 0;
 	FLinearColor SkyColor = Scene->SkyLight->GetEffectiveLightColor();
+	if (Scene->SkyLight->bRealTimeCaptureEnabled)
+	{
+		// Apply same logic as for RealTimeReflectionCapturePreExposure View parameter.
+		const float RealTimeReflectionCapturePreExposureInv = 1.0f / View.CachedViewUniformShaderParameters->RealTimeReflectionCapturePreExposure;
+		SkyColor.R *= RealTimeReflectionCapturePreExposureInv;
+		SkyColor.G *= RealTimeReflectionCapturePreExposureInv;
+		SkyColor.B *= RealTimeReflectionCapturePreExposureInv;
+	}
 	const bool bSkylightColorChanged = SkyColor != Scene->PathTracingSkylightColor;
 	if (!IsSkylightCachingEnabled || bSkylightColorChanged)
 	{
diff --git a/Engine/Source/Runtime/Renderer/Private/ReflectionEnvironmentRealTimeCapture.cpp b/Engine/Source/Runtime/Renderer/Private/ReflectionEnvironmentRealTimeCapture.cpp
--- a/Engine/Source/Runtime/Renderer/Private/ReflectionEnvironmentRealTimeCapture.cpp
+++ b/Engine/Source/Runtime/Renderer/Private/ReflectionEnvironmentRealTimeCapture.cpp
@@ -570,6 +570,9 @@ void FScene::AllocateAndCaptureFrameSkyEnvMap(
 				CubeView.CachedViewUniformShaderParameters->bPrimitiveAlphaHoldoutEnabled = 0;
 				// Notify the fact that we render a reflection, e.g. use special exposure.
 				CubeView.CachedViewUniformShaderParameters->RealTimeReflectionCapture = 1.0f;
+				// Clamp emissive to max10bits for real time capture
+				const static float Max10BitsFloat = 64512.0f;
+				CubeView.CachedViewUniformShaderParameters->MaterialMaxEmissiveValue = Max10BitsFloat;
 			};
 
 			auto SetupViewSkyAtmosphereParametersAndResources = [&](FViewInfo& OutView)
diff --git a/Engine/Source/Runtime/Renderer/Private/SceneRendering.cpp b/Engine/Source/Runtime/Renderer/Private/SceneRendering.cpp
--- a/Engine/Source/Runtime/Renderer/Private/SceneRendering.cpp
+++ b/Engine/Source/Runtime/Renderer/Private/SceneRendering.cpp
@@ -473,6 +473,13 @@ static TAutoConsoleVariable<int32> CVarCrossGPUTransferOption(
 	TEXT(" 2: optimized pull transfer (destination GPU runs copy, with transfers delayed to last view's render); (default)"),
 	ECVF_Default);
 
+static TAutoConsoleVariable<float> CVarSkyLightRealTimeReflectionCapturePreExposure(
+	TEXT("r.SkyLight.RealTimeReflectionCapture.PreExposure"),
+	8.0f,
+	TEXT("Fixed pre-exposure value for real time sky light capture in EV. Default 8 means [-8;24] EV representable range, which should cover physically based lighting range."),
+	ECVF_Scalability | ECVF_RenderThreadSafe
+);
+
 FOcclusionSubmittedFenceState FSceneRenderer::OcclusionSubmittedFence[FOcclusionQueryHelpers::MaxBufferedOcclusionFrames];
 
 // cleanup OcclusionSubmittedFence to avoid undefined order of destruction that can destroy it after its allocator
@@ -1867,7 +1874,8 @@ void FViewInfo::SetupUniformBufferParameters(
 
 	ViewUniformShaderParameters.RenderingReflectionCaptureMask = bIsReflectionCapture ? 1.0f : 0.0f;
 	ViewUniformShaderParameters.RealTimeReflectionCapture = 0.0f;
-	ViewUniformShaderParameters.RealTimeReflectionCapturePreExposure = 1.0f; // This must be 1 for now. If changed, we need to update the SkyLight AverageExposure and take it into account when sampling sky specular and diffuse irradiance.
+	const float RealTimeReflectionCapturePreExposureInv = FMath::Exp2(CVarSkyLightRealTimeReflectionCapturePreExposure.GetValueOnAnyThread());
+	ViewUniformShaderParameters.RealTimeReflectionCapturePreExposure = 1.0f / RealTimeReflectionCapturePreExposureInv;
 
 	ViewUniformShaderParameters.bPrimitiveAlphaHoldoutEnabled = IsPrimitiveAlphaHoldoutEnabled(*this);
 
@@ -1882,11 +1890,17 @@ void FViewInfo::SetupUniformBufferParameters(
 	{
 		FSkyLightSceneProxy* SkyLight = Scene->SkyLight;
 
+		float SkyPreExposureInv = 1.0f;
+		if (SkyLight->bRealTimeCaptureEnabled)
+		{
+			SkyPreExposureInv = RealTimeReflectionCapturePreExposureInv;
+		}
+
 		// Setup the sky color mulitpler, and use it to nullify the sky contribution in case SkyLighting is disabled.
 		// Note: we cannot simply select the base pass shader permutation skylight=0 because we would need to trigger bScenesPrimitivesNeedStaticMeshElementUpdate.
 		// However, this would need to be done per view (showflag is per view) and this is not possible today as it is selected within the scene. 
 		// So we simply nullify the sky light diffuse contribution. Reflection are handled by the indirect lighting render pass.
-		ViewUniformShaderParameters.SkyLightColor = Family->EngineShowFlags.SkyLighting ? SkyLight->GetEffectiveLightColor() : FLinearColor::Black;
+		ViewUniformShaderParameters.SkyLightColor = Family->EngineShowFlags.SkyLighting ? SkyLight->GetEffectiveLightColor() * SkyPreExposureInv : FLinearColor::Black;
 
 		bool bApplyPrecomputedBentNormalShadowing = 
 			SkyLight->bCastShadows 
