Path Tracer: Simplify logic around skylight color in the path tracer

Rather than duplicate logic from other parts of the renderer, take the (exposed) value already calculated in the view uniform structure.

#rnx

diff --git a/Engine/Source/Runtime/Renderer/Private/PathTracing.cpp b/Engine/Source/Runtime/Renderer/Private/PathTracing.cpp
--- a/Engine/Source/Runtime/Renderer/Private/PathTracing.cpp
+++ b/Engine/Source/Runtime/Renderer/Private/PathTracing.cpp
@@ -2717,15 +2717,7 @@ bool PrepareSkyTexture(FRDGBuilder& GraphBuilder, FScene* Scene, const FViewInfo
 
 	// should we remember the skylight prep for the next frame?
 	const bool IsSkylightCachingEnabled = CVarPathTracingSkylightCaching.GetValueOnAnyThread() != 0;
-	FLinearColor SkyColor = Scene->SkyLight->GetEffectiveLightColor();
-	if (Scene->SkyLight->bRealTimeCaptureEnabled)
-	{
-		// Apply same logic as for RealTimeReflectionCapturePreExposure View parameter.
-		const float RealTimeReflectionCapturePreExposureInv = 1.0f / View.CachedViewUniformShaderParameters->RealTimeReflectionCapturePreExposure;
-		SkyColor.R *= RealTimeReflectionCapturePreExposureInv;
-		SkyColor.G *= RealTimeReflectionCapturePreExposureInv;
-		SkyColor.B *= RealTimeReflectionCapturePreExposureInv;
-	}
+	const FLinearColor SkyColor = View.CachedViewUniformShaderParameters->SkyLightColor;
 	const bool bSkylightColorChanged = SkyColor != Scene->PathTracingSkylightColor;
 	if (!IsSkylightCachingEnabled || bSkylightColorChanged)
 	{
@@ -2967,7 +2959,7 @@ void SetLightParameters(
 			{
 				// NOTE: this color is already baked into the skylight texture so that importance sampling takes it into account, we pass it in here so that camera rays can factor it out
 				// This is only for the realtime capture case, because otherwise (specified cube map case) we want the displayed texture and lighting to match
-				DestLight.Color = FVector3f(Scene->SkyLight->GetEffectiveLightColor());
+				DestLight.Color = FVector3f(View.CachedViewUniformShaderParameters->SkyLightColor);
 			}
 		}
 	}
@@ -3588,12 +3580,13 @@ void FDeferredShadingSceneRenderer::RenderPathTracing(
 	}
 
 	// if the skylight has changed colors, reset both the path tracer and the importance tables
-	if (Scene->SkyLight && Scene->SkyLight->GetEffectiveLightColor() != Scene->PathTracingSkylightColor)
+	const FLinearColor SkyColor = View.CachedViewUniformShaderParameters->SkyLightColor;
+	if (Scene->SkyLight && SkyColor != Scene->PathTracingSkylightColor)
 	{
 		Scene->PathTracingSkylightTexture.SafeRelease();
 		Scene->PathTracingSkylightPdf.SafeRelease();
 		// reset last color here as well in case we don't reach PrepareSkyLightTexture
-		Scene->PathTracingSkylightColor = Scene->SkyLight->GetEffectiveLightColor();
+		Scene->PathTracingSkylightColor = SkyColor;
 		if (!View.bIsOfflineRender)
 		{
 			// reset accumulation, unless this is an offline render, in which case it is ok for the color to evolve
