LFV- real time skyloight capture setup dummy LFV resources. No LFV in real time capture.
This is to avoid dependency betweem async comptue and graphics pipe due to lfv ressources being on the View UB.

#tests PC editor dx12 and mobile vulkan preview
#rnx

diff --git a/Engine/Source/Runtime/Renderer/Private/BasePassRendering.cpp b/Engine/Source/Runtime/Renderer/Private/BasePassRendering.cpp
--- a/Engine/Source/Runtime/Renderer/Private/BasePassRendering.cpp
+++ b/Engine/Source/Runtime/Renderer/Private/BasePassRendering.cpp
@@ -30,6 +30,7 @@
 #include "PostProcess/SceneRenderTargets.h"
 #include "ShaderCompiler.h"
 #include "GBufferInfo.h"
+#include "LocalFogVolumeRendering.h"
 
 #include "BasePassRendering.inl"
 
@@ -800,7 +801,18 @@ void SetupSharedBasePassParameters(
 		SharedParameters.FogISR = SharedParameters.Fog;
 	}
 
-	SharedParameters.LFV = View.LocalFogVolumeViewData.UniformParametersStruct;
+	if (bForRealtimeSkyCapture)
+	{
+		// LFV are not allowed in real time capture since they are local effects.
+		// Furthermore, this avoid having the sky capture pass on async compute being dependent on LFV culling pass resources,
+		// because those resources are in the View UB so they most be ready to read.
+		// So we bind dummy parameters and resources for real time capture.
+		SetDummyLocalFogVolumeUniformParametersStruct(GraphBuilder, SharedParameters.LFV);
+	}
+	else
+	{
+		SharedParameters.LFV = View.LocalFogVolumeViewData.UniformParametersStruct;
+	}
 
 	SharedParameters.LightFunctionAtlas = *LightFunctionAtlas::GetGlobalParametersStruct(GraphBuilder, View);
 
diff --git a/Engine/Source/Runtime/Renderer/Private/LocalFogVolumeRendering.cpp b/Engine/Source/Runtime/Renderer/Private/LocalFogVolumeRendering.cpp
--- a/Engine/Source/Runtime/Renderer/Private/LocalFogVolumeRendering.cpp
+++ b/Engine/Source/Runtime/Renderer/Private/LocalFogVolumeRendering.cpp
@@ -165,6 +165,26 @@ DECLARE_GPU_STAT(LocalFogVolumeVolumes);
 static const uint32 SizeOfUintVec4 = sizeof(FUintVector4);
 static const uint32 UintVec4CountInLocalFogVolumeGPUInstanceData = sizeof(FLocalFogVolumeGPUInstanceData) / SizeOfUintVec4;
 
+void SetDummyLocalFogVolumeUniformParametersStruct(FRDGBuilder& GraphBuilder, FLocalFogVolumeUniformParameters& UniformParametersStruct)
+{
+	UniformParametersStruct.LocalFogVolumeCommon.LocalFogVolumeTileDataTextureResolution = FUintVector2(1, 1);
+	UniformParametersStruct.LocalFogVolumeCommon.LocalFogVolumeInstanceCount = 0;
+	UniformParametersStruct.LocalFogVolumeCommon.LocalFogVolumeTilePixelSize = GetLocalFogVolumeTilePixelSize();
+	UniformParametersStruct.LocalFogVolumeCommon.LocalFogVolumeMaxDensityIntoVolumetricFog = GetLocalFogVolumeMaxDensityIntoVolumetricFog();
+	UniformParametersStruct.LocalFogVolumeCommon.ShouldRenderLocalFogVolumeInVolumetricFog = 0;
+	UniformParametersStruct.LocalFogVolumeCommon.DirectionalLightColor = FVector3f::Zero();
+	UniformParametersStruct.LocalFogVolumeCommon.DirectionalLightDirection = FVector3f::Zero();
+	UniformParametersStruct.LocalFogVolumeCommon.GlobalStartDistance = 0.0f;
+	UniformParametersStruct.LocalFogVolumeCommon.HalfResTextureSizeAndInvSize = FVector4f(1.0f, 1.0f, 1.0f, 1.0f);
+
+	UniformParametersStruct.LocalFogVolumeTileDataTexture = GraphBuilder.CreateSRV(GSystemTextures.GetZeroUIntArrayDummy(GraphBuilder));
+
+	static FLocalFogVolumeGPUInstanceData DummyData;	// Static data so ERDGInitialDataFlags::NoCopy can be used. 
+	FRDGBufferRef GPUInstanceDataBuffer = CreateVertexBuffer(GraphBuilder, TEXT("DUMMYLocalFogVolumeGPUInstanceDataBuffer"),
+		FRDGBufferDesc::CreateBufferDesc(SizeOfUintVec4, UintVec4CountInLocalFogVolumeGPUInstanceData), &DummyData, sizeof(FLocalFogVolumeGPUInstanceData) * 1, ERDGInitialDataFlags::NoCopy);
+	UniformParametersStruct.LocalFogVolumeCommon.LocalFogVolumeInstances = GraphBuilder.CreateSRV(GPUInstanceDataBuffer, PF_A32B32G32R32F);
+}
+
 void SetDummyLocalFogVolumeForView(FRDGBuilder& GraphBuilder, FViewInfo& View)
 {
 	// The size of the structure must be a multiple of FUintVector4.
diff --git a/Engine/Source/Runtime/Renderer/Private/LocalFogVolumeRendering.h b/Engine/Source/Runtime/Renderer/Private/LocalFogVolumeRendering.h
--- a/Engine/Source/Runtime/Renderer/Private/LocalFogVolumeRendering.h
+++ b/Engine/Source/Runtime/Renderer/Private/LocalFogVolumeRendering.h
@@ -141,6 +141,7 @@ void GetLocalFogVolumeViewSortingData(const FScene* Scene, const FViewInfo& View
 void SetDummyLocalFogVolumeForViews(FRDGBuilder& GraphBuilder, TArray<FViewInfo>& Views);
 void SetDummyLocalFogVolumeForViews(FRDGBuilder& GraphBuilder, TArray<FViewInfo*>& Views);
 void SetDummyLocalFogVolumeForView(FRDGBuilder& GraphBuilder, FViewInfo& View);
+void SetDummyLocalFogVolumeUniformParametersStruct(FRDGBuilder& GraphBuilder, FLocalFogVolumeUniformParameters& UniformParametersStruct);
 
 void InitLocalFogVolumesForViews(
 	const FScene* Scene,
diff --git a/Engine/Source/Runtime/Renderer/Private/MobileBasePassRendering.cpp b/Engine/Source/Runtime/Renderer/Private/MobileBasePassRendering.cpp
--- a/Engine/Source/Runtime/Renderer/Private/MobileBasePassRendering.cpp
+++ b/Engine/Source/Runtime/Renderer/Private/MobileBasePassRendering.cpp
@@ -22,6 +22,7 @@
 #include "LocalLightSceneProxy.h"
 #include "ReflectionEnvironment.h"
 #include "RenderCore.h"
+#include "LocalFogVolumeRendering.h"
 
 // Changing this causes a full shader recompile
 static TAutoConsoleVariable<int32> CVarMobileDisableVertexFog(
@@ -347,7 +348,15 @@ void SetupMobileBasePassUniformParameters(
 	// Substrate
 	Substrate::BindSubstrateMobileForwardPasslUniformParameters(GraphBuilder, View, BasePassParameters.Substrate);
 
-	BasePassParameters.LFV = View.LocalFogVolumeViewData.UniformParametersStruct;
+	if (bForRealtimeSkyCapture)
+	{
+		// LFV are not allowed in real time capture since they are local effects.
+		SetDummyLocalFogVolumeUniformParametersStruct(GraphBuilder, BasePassParameters.LFV);
+	}
+	else
+	{
+		BasePassParameters.LFV = View.LocalFogVolumeViewData.UniformParametersStruct;
+	}
 
 	// We need to compose the half resolution LFV texture when rendering mesh with Sky materials. So that the Fog passes remains cheap and we can keep the stencil test on the fog pass.
 	if (BasePass > EMobileBasePass::DepthPrePass && !bForRealtimeSkyCapture)
