Fix memory write out of bounds in SetDefaultTextureData2D / SetDefaultTextureData2DArray due to incorrect byte count reported by FormatData() helper function for single component formats

diff --git a/Engine/Source/Runtime/Renderer/Private/SystemTextures.cpp b/Engine/Source/Runtime/Renderer/Private/SystemTextures.cpp
--- a/Engine/Source/Runtime/Renderer/Private/SystemTextures.cpp
+++ b/Engine/Source/Runtime/Renderer/Private/SystemTextures.cpp
@@ -1257,7 +1257,7 @@ void FormatData(const TInType& In, uint8* Out, uint32& OutByteCount)
 {
 	TOutType* OutTyped = (TOutType*)Out;
 	OutTyped[0] = ConvertInputFormat<typename TFormatConversionTraits<TInType>::Type, TOutType, InputFormatType>(In[0]);
-	OutByteCount = 4;
+	OutByteCount = sizeof(TOutType);
 }
 
 template <typename TInType>
