Fixing stereo view crash

diff --git a/Engine/Source/Runtime/Renderer/Private/HeterogeneousVolumes/HeterogeneousVolumes.cpp b/Engine/Source/Runtime/Renderer/Private/HeterogeneousVolumes/HeterogeneousVolumes.cpp
--- a/Engine/Source/Runtime/Renderer/Private/HeterogeneousVolumes/HeterogeneousVolumes.cpp
+++ b/Engine/Source/Runtime/Renderer/Private/HeterogeneousVolumes/HeterogeneousVolumes.cpp
@@ -446,8 +446,7 @@ FViewInfo* GetPrimaryViewForView(
 		PrimaryView = &Views[PrimaryViewIndex];
 	}
 
-	check(IStereoRendering::IsAPrimaryView(*PrimaryView))
-	return PrimaryView;
+	return IStereoRendering::IsAPrimaryView(*PrimaryView) ? PrimaryView : nullptr;
 }
 
 bool ShouldRenderHeterogeneousVolumesForView(
@@ -456,7 +455,7 @@ bool ShouldRenderHeterogeneousVolumesForView(
 )
 {
 	FViewInfo* PrimaryView = GetPrimaryViewForView(Views, ViewIndex);
-	return ShouldRenderHeterogeneousVolumesForView(*PrimaryView);
+	return PrimaryView ? ShouldRenderHeterogeneousVolumesForView(*PrimaryView) : false;
 }
 
 TArray<FVolumetricMeshBatch, SceneRenderingAllocator> GetHeterogeneousVolumesMeshBatches(
@@ -465,7 +464,7 @@ TArray<FVolumetricMeshBatch, SceneRenderingAllocator> GetHeterogeneousVolumesMes
 )
 {
 	FViewInfo* PrimaryView = GetPrimaryViewForView(Views, ViewIndex);
-	return PrimaryView->HeterogeneousVolumesMeshBatches;
+	return PrimaryView ? PrimaryView->HeterogeneousVolumesMeshBatches : TArray<FVolumetricMeshBatch, SceneRenderingAllocator>();
 }
 
 bool ShouldRenderHeterogeneousVolumesAsHoldoutForView(
