Fix a couple of virtual texture CPU streaming bugs.
Made CalcTextureMemorySize() more accurate for virtual textures. Takes layering and tiled nature of VTs into account.
Fixed an off by one error in the root page level.
Changed streamer to not take the root pages into account in its budgets because these are already managed by the main virtual texture system. The CPU streaming drives the mip range above the root pages up until the per texture prefetch mip setting.
Added some additional logging of locked page info to r.vt.ListPhysicalPools and r.vt.DumpPoolUsage.
#rnx

diff --git a/Engine/Source/Runtime/Engine/Private/Streaming/AsyncTextureStreaming.cpp b/Engine/Source/Runtime/Engine/Private/Streaming/AsyncTextureStreaming.cpp
--- a/Engine/Source/Runtime/Engine/Private/Streaming/AsyncTextureStreaming.cpp
+++ b/Engine/Source/Runtime/Engine/Private/Streaming/AsyncTextureStreaming.cpp
@@ -882,15 +882,14 @@ void FRenderAssetStreamingMipCalcTask::UpdateBudgetedVirtualTextureMips_Async()
 			}
 			else
 			{
-				MemoryBudgeted += StreamingRenderAsset.UpdateRetentionPriority_Async(Settings.bPrioritizeMeshLODRetention);;
-		
-				if (StreamingRenderAsset.IsMaxResolutionAffectedByGlobalBias())
+				MemoryBudgeted += StreamingRenderAsset.UpdateRetentionPriority_Async(Settings.bPrioritizeMeshLODRetention);
+				// Don't consider the virtual texture root pages as part of our budget.
+				MemoryBudgeted -= StreamingRenderAsset.GetSize(StreamingRenderAsset.MinAllowedMips);
+
+				const int32 MinAllowedMips = FMath::Max(StreamingRenderAsset.MinAllowedMips, StreamingRenderAsset.NumForcedMips);
+				if (StreamingRenderAsset.BudgetedMips > MinAllowedMips)
 				{
-					const int32 MinAllowedMips = FMath::Max(StreamingRenderAsset.MinAllowedMips, StreamingRenderAsset.NumForcedMips);
-					if (StreamingRenderAsset.BudgetedMips > MinAllowedMips)
-					{
-						AssetIndices.Add(AssetIndex);
-					}
+					AssetIndices.Add(AssetIndex);
 				}
 			}
 		}
@@ -1153,7 +1152,13 @@ void FRenderAssetStreamingMipCalcTask::UpdateStats_Async()
 		{
 			++Stats.NumVirtualTextures;
 			Stats.RequiredVirtualTexturePool += RequiredSize;
-			Stats.UsedVirtualTexturePool += FMath::Min<int64>(RequiredSize, BudgetedSize);
+			Stats.UsedVirtualTexturePool += ResidentSize;
+			// For virtual textures we don't include the root pages in the budget since they are managed completely 
+			// by the virtual texture system. We can't even assume that root pages are always locked because the virtual 
+			// texture system usually only locks them when a material using the texture allocates a VT.
+			const int64 RootPageSize = StreamingRenderAsset.GetSize(StreamingRenderAsset.MinAllowedMips);
+			Stats.RequiredVirtualTexturePool -= RootPageSize;
+			Stats.UsedVirtualTexturePool -= RootPageSize;
 			continue;
 		}
 
diff --git a/Engine/Source/Runtime/Engine/Private/Texture.cpp b/Engine/Source/Runtime/Engine/Private/Texture.cpp
--- a/Engine/Source/Runtime/Engine/Private/Texture.cpp
+++ b/Engine/Source/Runtime/Engine/Private/Texture.cpp
@@ -417,8 +417,8 @@ void UTexture::UpdateResourceWithParams(EUpdateResourceFlags InFlags)
 
 			if (FVirtualTexture2DResource* VirtualTextureResource = NewResource->GetVirtualTexture2DResource())
 			{
-				const uint8 NonStreamingMips = FMath::FloorLog2(VirtualTextureResource->GetTileSize()) + 1;
-				const uint8 NumMips = VirtualTextureResource->GetNumMips();
+				const uint8 NonStreamingMips = FMath::FloorLog2(VirtualTextureResource->GetTileSize());
+				const uint8 NumMips = FMath::Min(VirtualTextureResource->GetNumMips(), FStreamableRenderResourceState::MAX_LOD_COUNT);
 				const uint8 MaxStreamingMip = FMath::Min(VirtualTexturePrefetchMips, NumMips);
 				
 				if (MaxStreamingMip > NonStreamingMips)
diff --git a/Engine/Source/Runtime/Engine/Private/Texture2D.cpp b/Engine/Source/Runtime/Engine/Private/Texture2D.cpp
--- a/Engine/Source/Runtime/Engine/Private/Texture2D.cpp
+++ b/Engine/Source/Runtime/Engine/Private/Texture2D.cpp
@@ -833,41 +833,72 @@ int32 UTexture2D::CalcTextureMemorySize( int32 MipCount ) const
 #endif
 
 	int32 Size = 0;
-	if (GetPlatformData())
+	if (FTexturePlatformData const* PlatformDataPtr = GetPlatformData())
 	{
-		static TConsoleVariableData<int32>* CVarReducedMode = IConsoleManager::Get().FindTConsoleVariableDataInt(TEXT("r.PRT.ReducedMemory"));
-		check(CVarReducedMode);
-
-		ETextureCreateFlags TexCreateFlags = (SRGB ? TexCreate_SRGB : TexCreate_None) | (bNoTiling ? TexCreate_NoTiling : TexCreate_None) | (bNotOfflineProcessed ? TexCreate_None : TexCreate_OfflineProcessed) | TexCreate_Streamable;
-		const bool bCanUsePartiallyResidentMips = CanCreateWithPartiallyResidentMips(TexCreateFlags);
+		if (FVirtualTextureBuiltData const* VTData = PlatformDataPtr->VTData)
+		{
+			const uint32 TileSize = VTData->GetPhysicalTileSize();
+			const uint32 LayerCount = VTData->GetNumLayers();
+			
+			uint32 TileMemorySize = 0;				
+			for (uint32 LayerIndex = 0; LayerIndex < LayerCount; ++LayerIndex)
+			{
+				const EPixelFormat Format = VTData->LayerTypes[LayerIndex];
+				TileMemorySize += CalculateImageBytes(TileSize, TileSize, 1, Format);
+			}
 
-		const int32 SizeX = GetSizeX();
-		const int32 SizeY = GetSizeY();
-		const int32 NumMips = GetNumMips();
-		const int32 FirstMip = FMath::Max(0, NumMips - MipCount);
-		const EPixelFormat Format = GetPixelFormat();
+			// Get number of tiles in all mips up to and including this mip level.
+			// Stop iteration when we get down to one tile since this tile is shared by the mip tail.
+			const int32 MipLevel = FMath::Max<int32>(VTData->GetNumMips() - MipCount, 0);
+			uint32 TileCountX = FMath::Max(VTData->GetWidthInTiles() >> MipLevel, 1u);
+			uint32 TileCountY = FMath::Max(VTData->GetHeightInTiles() >> MipLevel, 1u);
+			uint32 TileCount = TileCountX * TileCountY;
 
-		// Must be consistent with the logic in FTexture2DResource::InitRHI
-		if (IsStreamable() && bCanUsePartiallyResidentMips && (!CVarReducedMode->GetValueOnAnyThread() || MipCount > UTexture2D::GetMinTextureResidentMipCount()))
-		{
-			const FRHITextureDesc Desc =
-				FRHITextureCreateDesc::Create2D(TEXT("Temp"), SizeX, SizeY, Format)
-				.SetNumMips(NumMips)
-				.SetFlags(TexCreateFlags | TexCreate_Virtual);
+			while (TileCountX > 1 && TileCountY > 1)
+			{
+				TileCountX = FMath::DivideAndRoundUp(TileCountX, 2u);
+				TileCountY = FMath::DivideAndRoundUp(TileCountY, 2u);
+				TileCount += TileCountX * TileCountY;
+			}
 
-			Size = RHICalcTexturePlatformSize(Desc, FirstMip).Size;
+			Size = TileCount * TileMemorySize;
 		}
 		else
 		{
-			const FIntPoint MipExtents = CalcMipMapExtent(SizeX, SizeY, Format, FirstMip);
+			static TConsoleVariableData<int32>* CVarReducedMode = IConsoleManager::Get().FindTConsoleVariableDataInt(TEXT("r.PRT.ReducedMemory"));
+			check(CVarReducedMode);
+
+			ETextureCreateFlags TexCreateFlags = (SRGB ? TexCreate_SRGB : TexCreate_None) | (bNoTiling ? TexCreate_NoTiling : TexCreate_None) | (bNotOfflineProcessed ? TexCreate_None : TexCreate_OfflineProcessed) | TexCreate_Streamable;
+			const bool bCanUsePartiallyResidentMips = CanCreateWithPartiallyResidentMips(TexCreateFlags);
 
-			const FRHITextureDesc Desc =
-				FRHITextureCreateDesc::Create2D(TEXT("Temp"), MipExtents, Format)
-				.SetNumMips(FMath::Max(1, MipCount))
-				.SetFlags(TexCreateFlags)
-				.SetExtData(GetPlatformData()->GetExtData());
+			const int32 SizeX = GetSizeX();
+			const int32 SizeY = GetSizeY();
+			const int32 NumMips = GetNumMips();
+			const int32 FirstMip = FMath::Max(0, NumMips - MipCount);
+			const EPixelFormat Format = GetPixelFormat();
 
-			Size = RHICalcTexturePlatformSize(Desc).Size;
+			// Must be consistent with the logic in FTexture2DResource::InitRHI
+			if (IsStreamable() && bCanUsePartiallyResidentMips && (!CVarReducedMode->GetValueOnAnyThread() || MipCount > UTexture2D::GetMinTextureResidentMipCount()))
+			{
+				const FRHITextureDesc Desc =
+					FRHITextureCreateDesc::Create2D(TEXT("Temp"), SizeX, SizeY, Format)
+					.SetNumMips(NumMips)
+					.SetFlags(TexCreateFlags | TexCreate_Virtual);
+
+				Size = RHICalcTexturePlatformSize(Desc, FirstMip).Size;
+			}
+			else
+			{
+				const FIntPoint MipExtents = CalcMipMapExtent(SizeX, SizeY, Format, FirstMip);
+
+				const FRHITextureDesc Desc =
+					FRHITextureCreateDesc::Create2D(TEXT("Temp"), MipExtents, Format)
+					.SetNumMips(FMath::Max(1, MipCount))
+					.SetFlags(TexCreateFlags)
+					.SetExtData(GetPlatformData()->GetExtData());
+
+				Size = RHICalcTexturePlatformSize(Desc).Size;
+			}
 		}
 	}
 	return Size;
diff --git a/Engine/Source/Runtime/Renderer/Private/VT/TexturePagePool.cpp b/Engine/Source/Runtime/Renderer/Private/VT/TexturePagePool.cpp
--- a/Engine/Source/Runtime/Renderer/Private/VT/TexturePagePool.cpp
+++ b/Engine/Source/Runtime/Renderer/Private/VT/TexturePagePool.cpp
@@ -369,14 +369,15 @@ uint32 FTexturePagePool::GetNumVisiblePages(uint32 Frame) const
 	return Count;
 }
 
-void FTexturePagePool::CollectProducerCounts(TMap<uint32, uint32>& OutProducerCountMap) const
+void FTexturePagePool::CollectProducerCounts(TMap<uint32, FUintVector2>& OutProducerCountMap) const
 {
 	for (uint32 i = NumReservedPages; i < NumPages; ++i)
 	{
 		const uint32 PackedProducerHandle = Pages[i].PackedProducerHandle;
 		if (PackedProducerHandle != 0u)
 		{
-			OutProducerCountMap.FindOrAdd(PackedProducerHandle) += 1;
+			const uint32 IsLocked = FreeHeap.IsPresent(i) ? 0u : 1u;
+			OutProducerCountMap.FindOrAdd(PackedProducerHandle) += FUintVector2(1u, IsLocked);
 		}
 	}
 }
diff --git a/Engine/Source/Runtime/Renderer/Private/VT/TexturePagePool.h b/Engine/Source/Runtime/Renderer/Private/VT/TexturePagePool.h
--- a/Engine/Source/Runtime/Renderer/Private/VT/TexturePagePool.h
+++ b/Engine/Source/Runtime/Renderer/Private/VT/TexturePagePool.h
@@ -132,9 +132,9 @@ public:
 	uint32		GetNumVisiblePages(uint32 Frame) const;
 
 	/**
-	 * Returns a map between producer handle and number of tiles that the producer uses in the pool.
+	 * Returns a map between producer handle and a uint2 containing number of tiles and number of locked tils that the producer uses in the pool.
 	 */
-	void		CollectProducerCounts(TMap<uint32, uint32>& OutProducerCountMap) const;
+	void		CollectProducerCounts(TMap<uint32, FUintVector2>& OutProducerCountMap) const;
 
 	/**
 	* Map the physical address to a specific virtual address.
diff --git a/Engine/Source/Runtime/Renderer/Private/VT/VirtualTextureSystem.cpp b/Engine/Source/Runtime/Renderer/Private/VT/VirtualTextureSystem.cpp
--- a/Engine/Source/Runtime/Renderer/Private/VT/VirtualTextureSystem.cpp
+++ b/Engine/Source/Runtime/Renderer/Private/VT/VirtualTextureSystem.cpp
@@ -508,6 +508,7 @@ void FVirtualTextureSystem::DumpFromConsole()
 void FVirtualTextureSystem::ListPhysicalPoolsFromConsole()
 {
 	uint64 TotalPhysicalMemory = 0u;
+	uint64 TotalLockedMemory = 0u;
 	for(int32 i = 0; i < PhysicalSpaces.Num(); ++i)
 	{
 		if (PhysicalSpaces[i] && PhysicalSpaces[i]->IsInitialized())
@@ -516,25 +517,27 @@ void FVirtualTextureSystem::ListPhysicalPoolsFromConsole()
 			const FVTPhysicalSpaceDescription& Desc = PhysicalSpace.GetDescription();
 			const FTexturePagePool& PagePool = PhysicalSpace.GetPagePool();
 			const uint32 TotalSizeInBytes = PhysicalSpace.GetSizeInBytes();
+			const uint32 NumTiles = PhysicalSpace.GetNumTiles();
 
 			UE_LOG(LogConsoleResponse, Display, TEXT("PhysicalPool: [%i] %s (%ix%i):"), i, *PhysicalSpace.GetFormatString(), Desc.TileSize, Desc.TileSize);
 			
-			const int32 AllocatedTiles = PagePool.GetNumAllocatedPages();
-			const float AllocatedLoad = (float)AllocatedTiles / (float)PhysicalSpace.GetNumTiles();
-			const float AllocatedMemory = AllocatedLoad * TotalSizeInBytes / 1024.0f / 1024.0f;
+			const uint32 AllocatedTiles = PagePool.GetNumAllocatedPages();
+			const uint32 AllocatedMemoryInBytes = AllocatedTiles * (TotalSizeInBytes / NumTiles);
+			const float AllocatedMemory = (float)AllocatedMemoryInBytes / 1024.0f / 1024.0f;
 
-			const int32 LockedTiles = PagePool.GetNumLockedPages();
-			const float LockedLoad = (float)LockedTiles / (float)PhysicalSpace.GetNumTiles();
-			const float LockedMemory = LockedLoad * TotalSizeInBytes / 1024.0f / 1024.0f;
+			const uint32 LockedTiles = PagePool.GetNumLockedPages();
+			const uint32 LockedMemoryInBytes = LockedTiles * (TotalSizeInBytes / NumTiles);
+			const float LockedMemory = (float)LockedMemoryInBytes / 1024.0f / 1024.0f;
 
 			UE_LOG(LogConsoleResponse, Display, TEXT("  SizeInMegabyte= %f"), (float)TotalSizeInBytes / 1024.0f / 1024.0f);
 			UE_LOG(LogConsoleResponse, Display, TEXT("  Dimensions= %ix%i"), PhysicalSpace.GetTextureSize(), PhysicalSpace.GetTextureSize());
-			UE_LOG(LogConsoleResponse, Display, TEXT("  Tiles= %i"), PhysicalSpace.GetNumTiles());
+			UE_LOG(LogConsoleResponse, Display, TEXT("  Tiles= %i"), NumTiles);
 			UE_LOG(LogConsoleResponse, Display, TEXT("  Tiles Allocated= %i (%fMB)"), AllocatedTiles, AllocatedMemory);
 			UE_LOG(LogConsoleResponse, Display, TEXT("  Tiles Locked= %i (%fMB)"), LockedTiles, LockedMemory);
 			UE_LOG(LogConsoleResponse, Display, TEXT("  Tiles Mapped= %i"), PagePool.GetNumMappedPages());
 
 			TotalPhysicalMemory += TotalSizeInBytes;
+			TotalLockedMemory += LockedMemoryInBytes;
 		}
 	}
 
@@ -577,6 +580,7 @@ void FVirtualTextureSystem::ListPhysicalPoolsFromConsole()
 
 	UE_LOG(LogConsoleResponse, Display, TEXT("TotalPageTableMemory: %fMB"), (double)TotalPageTableMemory / 1024.0 / 1024.0);
 	UE_LOG(LogConsoleResponse, Display, TEXT("TotalPhysicalMemory: %fMB"), (double)TotalPhysicalMemory / 1024.0 / 1024.0);
+	UE_LOG(LogConsoleResponse, Display, TEXT("TotalLockedMemory: %fMB"), (double)TotalLockedMemory / 1024.0 / 1024.0);
 }
 
 void FVirtualTextureSystem::DumpPoolUsageFromConsole()
@@ -591,31 +595,32 @@ void FVirtualTextureSystem::DumpPoolUsageFromConsole()
 
 			UE_LOG(LogConsoleResponse, Display, TEXT("PhysicalPool: [%i] %s (%ix%i):"), i, *PhysicalSpace.GetFormatString(), Desc.TileSize, Desc.TileSize);
 
-			TMap<uint32, uint32> ProducerCountMap;
+			TMap<uint32, FUintVector2> ProducerCountMap;
 			PagePool.CollectProducerCounts(ProducerCountMap);
 
-			TSet< TPair<uint32, uint32> > SortedProducerCounts;
-			for (TPair<uint32, uint32> ProducerCount : ProducerCountMap)
+			TSet<TPair<uint32, FUintVector2>> SortedProducerCounts;
+			for (TPair<uint32, FUintVector2> ProducerCount : ProducerCountMap)
 			{
-				// Filter out producers that only have locked page mapped.
+				// Filter out producers that only have locked pages mapped.
 				// Keep all producers for non compressed formats (assuming that we want to avoid these as much as possible).
 				// In future we can add other filters here (count, age etc), or we can dump more info and process off line.
-				if (ProducerCount.Value > 1 || GPixelFormats[PhysicalSpace.GetFormat(0)].BlockSizeX == 1)
+				if (ProducerCount.Value.X > 1 || GPixelFormats[PhysicalSpace.GetFormat(0)].BlockSizeX == 1)
 				{
 					SortedProducerCounts.Add(ProducerCount);
 				}
 			}
-			SortedProducerCounts.Sort([](TPair<uint32, uint32> const& LHS, TPair<uint32, uint32> const& RHS) { return LHS.Value > RHS.Value; });
+			SortedProducerCounts.Sort([](TPair<uint32, FUintVector2> const& LHS, TPair<uint32, FUintVector2> const& RHS) { return LHS.Value.X > RHS.Value.X; });
 
-			for (TPair<uint32, uint32> ProducerCount : SortedProducerCounts)
+			for (TPair<uint32, FUintVector2> ProducerCount : SortedProducerCounts)
 			{
 				const uint32 PackedProducerHandle = ProducerCount.Key;
-				const uint32 Count = ProducerCount.Value;
+				const uint32 Count = ProducerCount.Value.X;
+				const uint32 LockedCount = ProducerCount.Value.Y;
 
 				FVirtualTextureProducer* Producer = Producers.FindProducer(FVirtualTextureProducerHandle(PackedProducerHandle));
 				if (Producer != nullptr)
 				{
-					UE_LOG(LogConsoleResponse, Display, TEXT("   %s %d"), *Producer->GetName().ToString(), Count);
+					UE_LOG(LogConsoleResponse, Display, TEXT("   %s %d (%d)"), *Producer->GetName().ToString(), Count, LockedCount);
 				}
 			}
 		}
@@ -1325,26 +1330,28 @@ void FVirtualTextureSystem::SetMipLevelToLock(FVirtualTextureProducerHandle Prod
 	UE::TScopeLock Lock(Mutex);
 
 	FVirtualTextureProducer& Producer = Producers.GetProducer(ProducerHandle);
-	const int32 MaxLevel = Producer.GetMaxLevel();
-	const int32 MaxLoadLevel = MaxLevel - FMath::FloorLog2(Producer.GetDescription().TileSize);
-	const int32 NumLockMips = FMath::Max(MaxLoadLevel - InMipLevel + 1, 0);
+	const uint32 TileCountX = Producer.GetWidthInTiles();
+	const uint32 TileCountY = Producer.GetHeightInTiles();
+	const uint32 NumLevels = Producer.GetMaxLevel() + 1;
+	const uint32 NumNonStreamingLevels = FMath::FloorLog2(Producer.GetDescription().TileSize);
+	const uint32 NumLockMips = NumLevels - FMath::Min(NumLevels, InMipLevel + NumNonStreamingLevels);
 
 	// Compare with current setting on the producer and Lock or Unlock tiles to reach the new setting.
-	const int32 NumLockMipsPrevious = Producer.GetLockedMipCount();
+	const uint32 NumLockMipsPrevious = Producer.GetLockedMipCount();
 	Producer.SetLockedMipCount(NumLockMips);
 
 	const bool bLock = NumLockMips > NumLockMipsPrevious;
-	const int32 FirstMip = bLock ? NumLockMips : NumLockMipsPrevious;
-	const int32 LastMip = bLock ? NumLockMipsPrevious : NumLockMips;
+	const uint32 FirstMip = NumNonStreamingLevels + (bLock ? NumLockMipsPrevious : NumLockMips);
+	const uint32 LastMip = NumNonStreamingLevels + (bLock ? NumLockMips : NumLockMipsPrevious);
 
-	for (int32 MipIndex = FirstMip; MipIndex > LastMip; --MipIndex)
+	for (uint32 MipIndex = FirstMip; MipIndex < LastMip; ++MipIndex)
 	{
-		const int32 vLevel = MaxLoadLevel - MipIndex;
-		const int32 MipWidthInTiles = 1 << MipIndex;
-		const int32 MipHeightInTiles = 1 << MipIndex;
-		for (int32 PosY = 0; PosY < MipHeightInTiles; PosY++)
+		const uint32 vLevel = NumLevels - MipIndex - 1;
+		const uint32 MipWidthInTiles = FMath::DivideAndRoundUp(TileCountX, 1u << vLevel);
+		const uint32 MipHeightInTiles = FMath::DivideAndRoundUp(TileCountY, 1u << vLevel);
+		for (uint32 PosY = 0; PosY < MipHeightInTiles; PosY++)
 		{
-			for (int32 PosX = 0; PosX < MipWidthInTiles; PosX++)
+			for (uint32 PosX = 0; PosX < MipWidthInTiles; PosX++)
 			{
 				const uint32 Address = FMath::MortonCode2(PosX) | (FMath::MortonCode2(PosY) << 1);
 				const FVirtualTextureLocalTile Tile(ProducerHandle, Address, vLevel);
