Increase default ShaderPrint buffer size for character/line/triangle.

diff --git a/Engine/Source/Runtime/Renderer/Private/ShaderPrint.cpp b/Engine/Source/Runtime/Renderer/Private/ShaderPrint.cpp
--- a/Engine/Source/Runtime/Renderer/Private/ShaderPrint.cpp
+++ b/Engine/Source/Runtime/Renderer/Private/ShaderPrint.cpp
@@ -50,7 +50,7 @@ namespace ShaderPrint
 
 	static TAutoConsoleVariable<int32> CVarMaxCharacterCount(
 		TEXT("r.ShaderPrint.MaxCharacters"),
-		2000,
+		2048,
 		TEXT("ShaderPrint output buffer size.\n"),
 		ECVF_Cheat | ECVF_RenderThreadSafe);
 
@@ -62,13 +62,13 @@ namespace ShaderPrint
 
 	static TAutoConsoleVariable<int32> CVarMaxLineCount(
 		TEXT("r.ShaderPrint.MaxLine"),
-		32,
+		1024,
 		TEXT("ShaderPrint max line count.\n"),
 		ECVF_Cheat | ECVF_RenderThreadSafe);
 
 	static TAutoConsoleVariable<int32> CVarMaxTriangleCount(
 		TEXT("r.ShaderPrint.MaxTriangle"),
-		32,
+		1024,
 		TEXT("ShaderPrint max triangle count.\n"),
 		ECVF_Cheat | ECVF_RenderThreadSafe);
 
