Explicitly remove SSD shaders from Mobile Platforms as its never used

diff --git a/Engine/Source/Runtime/Renderer/Private/ScreenSpaceDenoise.cpp b/Engine/Source/Runtime/Renderer/Private/ScreenSpaceDenoise.cpp
--- a/Engine/Source/Runtime/Renderer/Private/ScreenSpaceDenoise.cpp
+++ b/Engine/Source/Runtime/Renderer/Private/ScreenSpaceDenoise.cpp
@@ -761,6 +761,12 @@ static_assert(UE_ARRAY_COUNT(kDenoiserOutputResourceNames) == int32(ESignalProce
 /** Returns whether should compile pipeline for a given shader platform.*/
 static bool ShouldCompileSignalPipeline(ESignalProcessing SignalProcessing, EShaderPlatform Platform)
 {
+	// Explicitly disable SSD shaders on mobile as they will fail to compile at this moment.
+	if (IsMobilePlatform(Platform))
+	{
+		return false;
+	}
+
 	if (SignalProcessing == ESignalProcessing::ScreenSpaceDiffuseIndirect)
 	{
 		return FDataDrivenShaderPlatformInfo::GetCompileSignalProcessingPipeline(Platform) || FDataDrivenShaderPlatformInfo::GetSupportsSSDIndirect(Platform);
