Path Tracer: Fix intensity of directly viewed skylights

diff --git a/Engine/Source/Runtime/Renderer/Private/PathTracing.cpp b/Engine/Source/Runtime/Renderer/Private/PathTracing.cpp
--- a/Engine/Source/Runtime/Renderer/Private/PathTracing.cpp
+++ b/Engine/Source/Runtime/Renderer/Private/PathTracing.cpp
@@ -2968,7 +2968,8 @@ void SetLightParameters(
 			{
 				// NOTE: this color is already baked into the skylight texture so that importance sampling takes it into account, we pass it in here so that camera rays can factor it out
 				// This is only for the realtime capture case, because otherwise (specified cube map case) we want the displayed texture and lighting to match
-				DestLight.Color = FVector3f(View.CachedViewUniformShaderParameters->SkyLightColor);
+				// NOTE: We want to use the effective light color, not the exposed skylight color SkyLightColor from the View UniformBuffer so that the overall intensity is correct (see jira UE-280734)
+				DestLight.Color = FVector3f(Scene->SkyLight->GetEffectiveLightColor());
 			}
 		}
 	}
