[Drawcalls] Prototype projected bounding sphere screen size culling during FrustumCull. The mechanism is very simple, taking the r^2 of the projected bounding sphere onto the screen and testing whether it is larger than the cutoff size. If it is, it's visible. If not, it should be marked as not visible.

#tests Local builds and preflight tests. Tested standalone and by itself with upscaler tests. Ran tests at 000.0, 001.0, 002.5, and 100.0.
[RN] Prototype projected bounding sphere screen size culling added during FrustumCull. Controlled via cvar "r.Visibility.FrustumCull.MaxScreenBoundsSpaceRadius".

diff --git a/Engine/Source/Runtime/Renderer/Private/SceneVisibility.cpp b/Engine/Source/Runtime/Renderer/Private/SceneVisibility.cpp
--- a/Engine/Source/Runtime/Renderer/Private/SceneVisibility.cpp
+++ b/Engine/Source/Runtime/Renderer/Private/SceneVisibility.cpp
@@ -616,6 +616,40 @@ inline bool ShouldCullForRayTracing(const FScene& Scene, FViewInfo& View, int32
 };
 #endif //RHI_RAYTRACING
 
+///////////////////////////////////////////////////////////////////////////////
+
+// (brad.blankenship): This is an initial prototype of the work for having objects cull themselves by
+// their projected screen size. It's very simple right now, only taking into account the r^2 vs. the cutoff. 
+// as part of the future work for the prototype we are looking at enabling a fade between visible and not visible
+// to reduce popin, as well as seeing what it would take to piggyback off of other features such as closeness to
+// screen center, view relevance, etc.
+float GMinScreenBoundsRadius = 0.0f;
+static FAutoConsoleVariableRef CVarFrustumCullMaxScreenSpaceBoundsRadius(
+	TEXT("r.MinScreenBoundsRadius"),
+	GMinScreenBoundsRadius,
+	TEXT("Specifies a minimum value for the screen space radius of the projected bounding sphere for visibility tests. ")
+	TEXT("Value must be from 0.0 (disabled) to 1.0 (100% of the longest dimension of the screen). "),
+	ECVF_Scalability | ECVF_RenderThreadSafe
+);
+
+inline bool IsScreenSizeCulled(FViewInfo& View, const FPrimitiveBounds& Bounds)
+{
+	// Never cull if we do not specify the value OR
+	// if the radius is larger than the screen itself.
+	if (GMinScreenBoundsRadius <= 0.0f ||
+		GMinScreenBoundsRadius > 1.0f)
+	{
+		return false;
+	}
+
+	// Assuming the cvar is 0 to 1.0 (mapping as 0%^2 to 100%^2), calculate the bounds of the sphere.
+	const float MinScreenSpaceBoundsRadiusSquared = GMinScreenBoundsRadius * GMinScreenBoundsRadius;
+	const float ComputedScreenRadiusSquared = ComputeBoundsScreenRadiusSquared(Bounds.BoxSphereBounds.Origin, Bounds.BoxSphereBounds.SphereRadius, View);
+	return ComputedScreenRadiusSquared < MinScreenSpaceBoundsRadiusSquared;
+}
+
+///////////////////////////////////////////////////////////////////////////////
+
 inline void ComputeDistances(const FPrimitiveBounds& Bounds, const FVector& ViewOriginForDistanceCulling, float& OutClosestDistSq, float& OutFurthestDistSq)
 {
 	float Dist = (Bounds.BoxSphereBounds.Origin - ViewOriginForDistanceCulling).Length();
@@ -829,6 +863,12 @@ static int32 FrustumCull(const FScene& Scene, FViewInfo& View, FFrustumCullingFl
 					}
 				}
 
+				// Screen size culling next
+				if (bIsVisible)
+				{
+					bIsVisible = !IsScreenSizeCulled(View, Bounds);
+				}
+
 				// Distance cull if frustum cull passed
 				bShouldDistanceCull = bShouldDistanceCull && bIsVisible;
 				if (bShouldDistanceCull)
