[HWRT] Prevent crash in FScene::UpdateCachedRayTracingState(...).
- also added an ensure to get more data about the cause.

diff --git a/Engine/Source/Runtime/Renderer/Private/RendererScene.cpp b/Engine/Source/Runtime/Renderer/Private/RendererScene.cpp
--- a/Engine/Source/Runtime/Renderer/Private/RendererScene.cpp
+++ b/Engine/Source/Runtime/Renderer/Private/RendererScene.cpp
@@ -4109,7 +4109,11 @@ void FScene::UpdateCachedRayTracingState(FPrimitiveSceneProxy* SceneProxy)
 {
 	if (SceneProxy->GetPrimitiveSceneInfo())
 	{
-		PrimitiveRayTracingDatas[SceneProxy->GetPrimitiveSceneInfo()->GetIndex()].bCachedRaytracingDataDirty = true;
+		const int32 PrimitiveIndex = SceneProxy->GetPrimitiveSceneInfo()->GetIndex();
+		if (ensure(PrimitiveIndex != INDEX_NONE))
+		{
+			PrimitiveRayTracingDatas[PrimitiveIndex].bCachedRaytracingDataDirty = true;
+		}
 
 		// Clear the recounted pointer as well since we don't need it anymore
 		SceneProxy->GetPrimitiveSceneInfo()->CachedRayTracingInstance.GeometryRHI = nullptr;
