Fix incorrect state for the bloom input texture when used by the lens flare VS.

#rnx

diff --git a/Engine/Source/Runtime/Renderer/Private/PostProcess/PostProcessLensFlares.cpp b/Engine/Source/Runtime/Renderer/Private/PostProcess/PostProcessLensFlares.cpp
--- a/Engine/Source/Runtime/Renderer/Private/PostProcess/PostProcessLensFlares.cpp
+++ b/Engine/Source/Runtime/Renderer/Private/PostProcess/PostProcessLensFlares.cpp
@@ -27,7 +27,7 @@ TAutoConsoleVariable<int32> CVarLensFlareQuality(
 
 // The RDG inputs shared by all lens flare passes.
 BEGIN_SHADER_PARAMETER_STRUCT(FLensFlarePassParameters, )
-	SHADER_PARAMETER_RDG_TEXTURE_NON_PIXEL_SRV(Texture2D, InputTexture)
+	SHADER_PARAMETER_RDG_TEXTURE_SRV(Texture2D, InputTexture)
 	SHADER_PARAMETER_STRUCT_INCLUDE(FViewShaderParameters, View)
 	RENDER_TARGET_BINDING_SLOTS()
 END_SHADER_PARAMETER_STRUCT()
@@ -59,7 +59,7 @@ public:
 	SHADER_USE_PARAMETER_STRUCT(FLensFlareBlurVS, FLensFlareShader);
 
 	BEGIN_SHADER_PARAMETER_STRUCT(FParameters, )
-		SHADER_PARAMETER_STRUCT_INCLUDE(FLensFlarePassParameters, Pass)
+		SHADER_PARAMETER_RDG_TEXTURE_NON_PIXEL_SRV(Texture2D, InputTexture)
 		SHADER_PARAMETER_SAMPLER(SamplerState, InputSampler)
 		SHADER_PARAMETER_STRUCT(FScreenPassTextureViewportParameters, Input)
 		SHADER_PARAMETER(FIntPoint, TileCount)
@@ -188,7 +188,7 @@ FScreenPassTexture AddLensFlaresPass(
 
 		// Setup vertex shader parameters.
 		FLensFlareBlurVS::FParameters VertexParameters;
-		VertexParameters.Pass = *PassParameters;
+		VertexParameters.InputTexture = Inputs.Flare.TextureSRV;
 		VertexParameters.Input = GetScreenPassTextureViewportParameters(FlareViewport);
 		VertexParameters.InputSampler = BilinearClampSampler;
 		VertexParameters.TileCount = TileCount;
