PostProcessing: Fix alpha invert showflag on desktop under certain conditions.

The non-mobile post-processing logic disabled the invert pass when primary upscale was active, however its inversion was mobile-only.

We retain the upscale inversion that was added for HMI uses cases (avoiding the cost of a separate pass). Note that it remains not entirely robust since there's no strong guarantee that the upscale pass runs on the mobile side.

diff --git a/Engine/Shaders/Private/PostProcessUpscale.usf b/Engine/Shaders/Private/PostProcessUpscale.usf
--- a/Engine/Shaders/Private/PostProcessUpscale.usf
+++ b/Engine/Shaders/Private/PostProcessUpscale.usf
@@ -35,6 +35,7 @@ SamplerState PointSceneColorSampler;
 Texture2DArray PointSceneColorTextureArray;
 
 uint2 GridDimensions;
+uint bInvertAlpha;
 
 // Do a smoothstep(x) = 3 x^2 - 2 x^3
 float2 GetSmoothstepUV(float2 LinearUV, float2 TextureSize, float2 TextureInvSize)
@@ -287,10 +288,7 @@ void MainPS(noperspective float4 UVAndScreenPos : TEXCOORD0, FStereoPSInput Ster
 		OutColor.a = 0.0;
 	}
 
-#if SHADING_PATH_MOBILE
-	// TODO: Remove inversion & replace with more robust solution for HMI: pass dependent & can potentially conflict with OpenXR backbuffer alpha inversion.
-	OutColor.a = 1.0f - OutColor.a;
-#endif
+	OutColor.a = select(bInvertAlpha > 0, 1.0f - OutColor.a, OutColor.a);
 	
 #else
 	OutColor.a = 0; // Skip all computations related to alpha
diff --git a/Engine/Source/Runtime/Renderer/Private/PostProcess/PostProcessUpscale.cpp b/Engine/Source/Runtime/Renderer/Private/PostProcess/PostProcessUpscale.cpp
--- a/Engine/Source/Runtime/Renderer/Private/PostProcess/PostProcessUpscale.cpp
+++ b/Engine/Source/Runtime/Renderer/Private/PostProcess/PostProcessUpscale.cpp
@@ -39,6 +39,7 @@ BEGIN_SHADER_PARAMETER_STRUCT(FUpscaleParameters, )
 	SHADER_PARAMETER_RDG_TEXTURE(Texture2D, DistortingDisplacementTexture)
 	SHADER_PARAMETER_SAMPLER(SamplerState, DistortingDisplacementSampler)
 	SHADER_PARAMETER(FIntPoint, GridDimensions)
+	SHADER_PARAMETER(uint32, bInvertAlpha)
 	SHADER_PARAMETER_RDG_TEXTURE(Texture2D, SceneColorTexture)
 	SHADER_PARAMETER_SAMPLER(SamplerState, SceneColorSampler)
 	SHADER_PARAMETER_RDG_TEXTURE(Texture2D, PointSceneColorTexture)
@@ -152,6 +153,7 @@ FScreenPassTexture ISpatialUpscaler::AddDefaultUpscalePass(
 	PassParameters->DistortingDisplacementTexture = LensDistortionLUT.DistortingDisplacementTexture;
 	PassParameters->DistortingDisplacementSampler = TStaticSamplerState<SF_Bilinear, AM_Clamp, AM_Clamp, AM_Clamp>::GetRHI();
 	PassParameters->GridDimensions = LensDistortionLUT.DistortionGridDimensions;
+	PassParameters->bInvertAlpha = View.Family->EngineShowFlags.AlphaInvert;
 	PassParameters->SceneColorTexture = Inputs.SceneColor.Texture;
 	PassParameters->SceneColorSampler = TStaticSamplerState<SF_Bilinear, AM_Border, AM_Border, AM_Border>::GetRHI();
 	PassParameters->PointSceneColorTexture = Inputs.SceneColor.Texture;
diff --git a/Engine/Source/Runtime/Renderer/Private/PostProcess/PostProcessing.cpp b/Engine/Source/Runtime/Renderer/Private/PostProcess/PostProcessing.cpp
--- a/Engine/Source/Runtime/Renderer/Private/PostProcess/PostProcessing.cpp
+++ b/Engine/Source/Runtime/Renderer/Private/PostProcess/PostProcessing.cpp
@@ -2506,6 +2506,8 @@ void AddMobilePostProcessingPasses(FRDGBuilder& GraphBuilder, FScene* Scene, con
 	PassSequence.SetEnabled(EPass::VisualizeLightGrid, ShouldVisualizeLightGrid());
 	PassSequence.SetEnabled(EPass::HMDDistortion, View.Family->EngineShowFlags.StereoRendering && View.Family->EngineShowFlags.HMDDistortion);
 
+	// TODO: Add support for EPass::AlphaInvert if the AlphaInvert show flag is enabled, but EPass::PrimaryUpscale is disabled
+
 	const auto GetPostProcessMaterialInputs = [&](FScreenPassTexture InSceneColor)
 		{
 			FPostProcessMaterialInputs PostProcessMaterialInputs;
