VolumetricFog - SkyAtmosphere - added time to better track work load in gpu capture.

diff --git a/Engine/Source/Runtime/Renderer/Private/SkyAtmosphereRendering.cpp b/Engine/Source/Runtime/Renderer/Private/SkyAtmosphereRendering.cpp
--- a/Engine/Source/Runtime/Renderer/Private/SkyAtmosphereRendering.cpp
+++ b/Engine/Source/Runtime/Renderer/Private/SkyAtmosphereRendering.cpp
@@ -1386,6 +1386,7 @@ void FSceneRenderer::RenderSkyAtmosphereLookUpTables(FRDGBuilder& GraphBuilder,
 	FGlobalShaderMap* GlobalShaderMap = GetGlobalShaderMap(FeatureLevel);
 	if (CVarSkyAtmosphereTransmittanceLUT.GetValueOnRenderThread() > 0)
 	{
+		RDG_EVENT_SCOPE(GraphBuilder, "SkyAtmosphere::TransmittanceLut");
 		TShaderMapRef<FRenderTransmittanceLutCS> ComputeShader(GlobalShaderMap);
 
 		FRenderTransmittanceLutCS::FParameters * PassParameters = GraphBuilder.AllocParameters<FRenderTransmittanceLutCS::FParameters>();
@@ -1401,6 +1402,7 @@ void FSceneRenderer::RenderSkyAtmosphereLookUpTables(FRDGBuilder& GraphBuilder,
 
 	// Multi-Scattering LUT
 	{
+		RDG_EVENT_SCOPE(GraphBuilder, "SkyAtmosphere::MultiScatteringLut");
 		FRenderMultiScatteredLuminanceLutCS::FPermutationDomain PermutationVector;
 		PermutationVector.Set<FHighQualityMultiScatteringApprox>(bHighQualityMultiScattering);
 		TShaderMapRef<FRenderMultiScatteredLuminanceLutCS> ComputeShader(GlobalShaderMap, PermutationVector);
@@ -1423,6 +1425,7 @@ void FSceneRenderer::RenderSkyAtmosphereLookUpTables(FRDGBuilder& GraphBuilder,
 	// Distant Sky Light LUT
 	if(CVarSkyAtmosphereDistantSkyLightLUT.GetValueOnRenderThread() > 0)
 	{
+		RDG_EVENT_SCOPE(GraphBuilder, "SkyAtmosphere::DistantSkyLightLut");
 		FRDGBufferRef DistantSkyLightLutBuffer = GraphBuilder.RegisterExternalBuffer(SkyInfo.GetDistantSkyLightLutBuffer());
 		FRDGBufferUAVRef DistantSkyLightLutBufferUAV = GraphBuilder.CreateUAV(FRDGBufferUAVDesc(DistantSkyLightLutBuffer, PF_A32B32G32R32F));
 		FRDGBufferRef MobileDistantSkyLightLutBuffer = GraphBuilder.RegisterExternalBuffer(SkyInfo.GetMobileDistantSkyLightLutBuffer());
@@ -1522,6 +1525,7 @@ void FSceneRenderer::RenderSkyAtmosphereLookUpTables(FRDGBuilder& GraphBuilder,
 
 		if (bRealTimeReflectionCaptureSkyAtmosphereViewLutTexture)
 		{
+			RDG_EVENT_SCOPE(GraphBuilder, "SkyAtmosphere::RealTimeCaptureSkyViewLut");
 			FIntVector SkyViewLutSize = ViewForReflection.SkyAtmosphereViewLutTexture->GetDesc().GetSize();
 			FRDGTextureRef RealTimeReflectionCaptureSkyAtmosphereViewLutTexture = GraphBuilder.CreateTexture(
 				FRDGTextureDesc::Create2D(FIntPoint(SkyViewLutSize.X, SkyViewLutSize.Y), GetSkyLutTextureFormat(Scene->GetFeatureLevel()), FClearValueBinding::None,
@@ -1570,6 +1574,7 @@ void FSceneRenderer::RenderSkyAtmosphereLookUpTables(FRDGBuilder& GraphBuilder,
 
 		if (bRealTimeReflectionCapture360APLutTexture)
 		{
+			RDG_EVENT_SCOPE(GraphBuilder, "SkyAtmosphere::RealTimeCaptureCamera360VolumeLut");
 			FIntVector CameraAPLutSize = ViewForReflection.SkyAtmosphereCameraAerialPerspectiveVolume->GetDesc().GetSize();
 			EPixelFormat TextureAerialLUTFormat = (CVarSkyAtmosphereLUT32.GetValueOnAnyThread() != 0) ? PF_A32B32G32R32F : PF_FloatRGBA;
 			FRDGTextureRef RealTimeReflectionCaptureCamera360APLutTexture = GraphBuilder.CreateTexture(
@@ -1657,6 +1662,7 @@ void FSceneRenderer::RenderSkyAtmosphereLookUpTables(FRDGBuilder& GraphBuilder,
 
 		// Sky View LUT
 		{
+			RDG_EVENT_SCOPE(GraphBuilder, "SkyAtmosphere::SkyViewLut");
 			FRenderSkyViewLutCS::FPermutationDomain PermutationVector;
 			PermutationVector.Set<FSampleCloudSkyAO>(CloudShadowAOData.bShouldSampleCloudSkyAO);
 			PermutationVector.Set<FSecondAtmosphereLight>(bSecondAtmosphereLightEnabled);
@@ -1702,6 +1708,7 @@ void FSceneRenderer::RenderSkyAtmosphereLookUpTables(FRDGBuilder& GraphBuilder,
 
 		// Camera Atmosphere Volume
 		{
+			RDG_EVENT_SCOPE(GraphBuilder, "SkyAtmosphere::CameraVolumeLut");
 			FRenderCameraAerialPerspectiveVolumeCS::FPermutationDomain PermutationVector;
 			PermutationVector.Set<FSampleCloudSkyAO>(CloudShadowAOData.bShouldSampleCloudSkyAO);
 			PermutationVector.Set<FSecondAtmosphereLight>(bSecondAtmosphereLightEnabled);
diff --git a/Engine/Source/Runtime/Renderer/Private/VolumetricFog.cpp b/Engine/Source/Runtime/Renderer/Private/VolumetricFog.cpp
--- a/Engine/Source/Runtime/Renderer/Private/VolumetricFog.cpp
+++ b/Engine/Source/Runtime/Renderer/Private/VolumetricFog.cpp
@@ -523,6 +523,8 @@ static void RenderRaytracedDirectionalShadowVolume(
 
 	if (DirectionalLightSceneInfo)
 	{
+		RDG_EVENT_SCOPE(GraphBuilder, "VolumetricFog::RaytraceDirLightShadow");
+
 		int32 VolumetricFogGridPixelSize;
 		const FIntVector VolumetricFogResourceGridSize = GetVolumetricFogResourceGridSize(View, VolumetricFogGridPixelSize);
 		FRDGTextureDesc RaytracedShadowsVolumeDesc(FRDGTextureDesc::Create3D(
@@ -769,6 +771,8 @@ void FSceneRenderer::RenderLocalLightsForVolumetricFog(
 	TConstArrayView<const FLightSceneInfo*> RayTracedLightsToInject,
 	FRDGTexture*& OutLocalShadowedLightScattering)
 {
+	RDG_EVENT_SCOPE(GraphBuilder, "VolumetricFog::LocalLights");
+
 	// Setup the light function atlas
 	const bool bUseLightFunctionAtlas = LightFunctionAtlas::IsEnabled(View, ELightFunctionAtlasSystem::VolumetricFog);
 	TRDGUniformBufferRef<FLightFunctionAtlasGlobalParameters> LightFunctionAtlasGlobalParameters = LightFunctionAtlas::BindGlobalParameters(GraphBuilder, View);
@@ -1518,6 +1522,7 @@ void FSceneRenderer::ComputeVolumetricFog(FRDGBuilder& GraphBuilder,
 		// To use a depth target format, and depth tests, we will have to render depth from a PS depth output. Keeping it simple for now with all the tests happening in shader.
 		if (GVolumetricFogConservativeDepth > 0)
 		{
+			RDG_EVENT_SCOPE(GraphBuilder, "VolumetricFog::ConservativeDepth");
 			FIntPoint ConservativeDepthTextureSize = FIntPoint(VolumetricFogViewGridSize.X, VolumetricFogViewGridSize.Y);
 			ConservativeDepthTexture = GraphBuilder.CreateTexture(FRDGTextureDesc::Create2D(ConservativeDepthTextureSize, PF_R16F,
 				FClearValueBinding::Black, TexCreate_RenderTargetable | TexCreate_ShaderResource | TexCreate_UAV), TEXT("VolumetricFog.ConservativeDepthTexture"));
@@ -1562,15 +1567,16 @@ void FSceneRenderer::ComputeVolumetricFog(FRDGBuilder& GraphBuilder,
 		
 		if (LightsToInject.DirectionalLightFunction)
 		{
-		RenderLightFunctionForVolumetricFog(
-			GraphBuilder,
-			View,
-			SceneTextures,
-			VolumetricFogViewGridSize,
-			FogInfo.VolumetricFogDistance,
-				LightsToInject.DirectionalLightFunction,
-			DirectionalLightFunctionTranslatedWorldToShadow,
-				DirectionalLightFunctionTexture);
+			RDG_EVENT_SCOPE(GraphBuilder, "VolumetricFog::DirLightFunction");
+			RenderLightFunctionForVolumetricFog(
+				GraphBuilder,
+				View,
+				SceneTextures,
+				VolumetricFogViewGridSize,
+				FogInfo.VolumetricFogDistance,
+					LightsToInject.DirectionalLightFunction,
+				DirectionalLightFunctionTranslatedWorldToShadow,
+					DirectionalLightFunctionTexture);
 		}
 			
 		View.VolumetricFogResources.IntegratedLightScatteringTexture = nullptr;
@@ -1601,6 +1607,8 @@ void FSceneRenderer::ComputeVolumetricFog(FRDGBuilder& GraphBuilder,
 #endif
 
 		{
+			RDG_EVENT_SCOPE(GraphBuilder, "VolumetricFog::InitialiseVolume");
+
 			FVolumetricFogMaterialSetupCS::FParameters* PassParameters = GraphBuilder.AllocParameters<FVolumetricFogMaterialSetupCS::FParameters>();
 			PassParameters->GlobalAlbedo = FogInfo.VolumetricFogAlbedo;
 			PassParameters->GlobalEmissive = FogInfo.VolumetricFogEmissive;
@@ -1650,6 +1658,8 @@ void FSceneRenderer::ComputeVolumetricFog(FRDGBuilder& GraphBuilder,
 		IntegrationData.LightScatteringUAV = GraphBuilder.CreateUAV(FRDGTextureUAVDesc(IntegrationData.LightScattering));
 
 		{
+			RDG_EVENT_SCOPE(GraphBuilder, "VolumetricFog::LightScattering");
+
 			FVolumetricFogLightScatteringCS::FParameters* PassParameters = GraphBuilder.AllocParameters<FVolumetricFogLightScatteringCS::FParameters>();
 
 			PassParameters->View = View.ViewUniformBuffer;
@@ -1849,6 +1859,8 @@ void FSceneRenderer::ComputeVolumetricFog(FRDGBuilder& GraphBuilder,
 		FRDGTextureUAV* IntegratedLightScatteringUAV = GraphBuilder.CreateUAV(FRDGTextureUAVDesc(IntegratedLightScattering));
 
 		{
+			RDG_EVENT_SCOPE(GraphBuilder, "VolumetricFog::FinalIntegration");
+
 			FVolumetricFogFinalIntegrationCS::FParameters* PassParameters = GraphBuilder.AllocParameters<FVolumetricFogFinalIntegrationCS::FParameters>();
 			PassParameters->LightScattering = IntegrationData.LightScattering;
 			PassParameters->RWIntegratedLightScattering = IntegratedLightScatteringUAV;
diff --git a/Engine/Source/Runtime/Renderer/Private/VolumetricFogVoxelization.cpp b/Engine/Source/Runtime/Renderer/Private/VolumetricFogVoxelization.cpp
--- a/Engine/Source/Runtime/Renderer/Private/VolumetricFogVoxelization.cpp
+++ b/Engine/Source/Runtime/Renderer/Private/VolumetricFogVoxelization.cpp
@@ -701,6 +701,8 @@ void FSceneRenderer::VoxelizeFogVolumePrimitives(
 {
 	if (View.VolumetricMeshBatches.Num() > 0 && DoesPlatformSupportVolumetricFogVoxelization(View.GetShaderPlatform()))
 	{
+		RDG_EVENT_SCOPE(GraphBuilder, "VolumetricFog::VoxelizePrimitives");
+
 		const FVector2D Jitter(
 			IntegrationData.FrameJitterOffsetValues[0].X / VolumetricFogViewGridSize.X,
 			IntegrationData.FrameJitterOffsetValues[0].Y / VolumetricFogViewGridSize.Y);
