[Landscape][RayTracing] Ability To Toggle Instanced Skeletal Meshes On/Off from Ray Tracing.

diff --git a/Engine/Source/Runtime/Engine/Private/Components/InstancedSkinnedMeshComponent.cpp b/Engine/Source/Runtime/Engine/Private/Components/InstancedSkinnedMeshComponent.cpp
--- a/Engine/Source/Runtime/Engine/Private/Components/InstancedSkinnedMeshComponent.cpp
+++ b/Engine/Source/Runtime/Engine/Private/Components/InstancedSkinnedMeshComponent.cpp
@@ -183,6 +183,18 @@ public:
 		return AnimBankIds;
 	}
 
+	virtual ERayTracingPrimitiveFlags GetCachedRayTracingInstance(FRayTracingInstance& RayTracingInstance) override
+	{
+		static const auto RayTracingInstancedSkeletalMeshesCVar = IConsoleManager::Get().FindTConsoleVariableDataInt(TEXT("r.RayTracing.Geometry.InstancedSkeletalMeshes"));
+
+		if (RayTracingInstancedSkeletalMeshesCVar && RayTracingInstancedSkeletalMeshesCVar->GetValueOnRenderThread() <= 0)
+		{
+			return ERayTracingPrimitiveFlags::Exclude;
+		}
+
+		return Nanite::FSkinnedSceneProxy::GetCachedRayTracingInstance(RayTracingInstance);
+	}
+
 	virtual bool GetInstanceDrawDistanceMinMax(FVector2f& OutCullRange) const override
 	{
 		if (InstanceEndCullDistance > 0)
diff --git a/Engine/Source/Runtime/Engine/Private/SkeletalMeshSceneProxy.cpp b/Engine/Source/Runtime/Engine/Private/SkeletalMeshSceneProxy.cpp
--- a/Engine/Source/Runtime/Engine/Private/SkeletalMeshSceneProxy.cpp
+++ b/Engine/Source/Runtime/Engine/Private/SkeletalMeshSceneProxy.cpp
@@ -38,6 +38,11 @@ static TAutoConsoleVariable<int32> CVarRayTracingSkeletalMeshes(
 	1,
 	TEXT("Include skeletal meshes in ray tracing effects (default = 1 (skeletal meshes enabled in ray tracing))"));
 
+static TAutoConsoleVariable<int32> CVarRayTracingInstancedSkeletalMeshes(
+	TEXT("r.RayTracing.Geometry.InstancedSkeletalMeshes"),
+	1,
+	TEXT("Include instanced skeletal meshes in ray tracing effects (default = 1 (instanced skeletal meshes enabled in ray tracing))"));
+
 static TAutoConsoleVariable<int32> CVarRayTracingSupportSkeletalMeshes(
 	TEXT("r.RayTracing.Geometry.SupportSkeletalMeshes"),
 	1,
diff --git a/Engine/Source/Runtime/Renderer/Private/RayTracing/RayTracing.cpp b/Engine/Source/Runtime/Renderer/Private/RayTracing/RayTracing.cpp
--- a/Engine/Source/Runtime/Renderer/Private/RayTracing/RayTracing.cpp
+++ b/Engine/Source/Runtime/Renderer/Private/RayTracing/RayTracing.cpp
@@ -103,24 +103,28 @@ static void RefreshRayTracingInstancesSinkFunction()
 	static const auto RayTracingHISMCVar = IConsoleManager::Get().FindTConsoleVariableDataInt(TEXT("r.RayTracing.Geometry.HierarchicalInstancedStaticMesh"));
 	static const auto RayTracingNaniteProxiesCVar = IConsoleManager::Get().FindTConsoleVariableDataInt(TEXT("r.RayTracing.Geometry.NaniteProxies"));
 	static const auto RayTracingSkeletalMeshesCVar = IConsoleManager::Get().FindTConsoleVariableDataInt(TEXT("r.RayTracing.Geometry.SkeletalMeshes"));
+	static const auto RayTracingISKMCVar = IConsoleManager::Get().FindTConsoleVariableDataInt(TEXT("r.RayTracing.Geometry.InstancedSkeletalMeshes"));
 	static const auto RayTracingLandscapeGrassCVar = IConsoleManager::Get().FindTConsoleVariableDataInt(TEXT("r.RayTracing.Geometry.LandscapeGrass"));
 
 	static int32 CachedRayTracingStaticMeshes = RayTracingStaticMeshesCVar->GetValueOnGameThread();
 	static int32 CachedRayTracingHISM = RayTracingHISMCVar->GetValueOnGameThread();
 	static int32 CachedRayTracingNaniteProxies = RayTracingNaniteProxiesCVar->GetValueOnGameThread();
 	static int32 CachedRayTracingSkeletalMeshes = RayTracingSkeletalMeshesCVar->GetValueOnGameThread();
+	static int32 CachedRayTracingISKM = RayTracingISKMCVar->GetValueOnGameThread();
 	static int32 CachedRayTracingLandscapeGrass = RayTracingLandscapeGrassCVar->GetValueOnGameThread();
 
 	const int32 RayTracingStaticMeshes = RayTracingStaticMeshesCVar->GetValueOnGameThread();
 	const int32 RayTracingHISM = RayTracingHISMCVar->GetValueOnGameThread();
 	const int32 RayTracingNaniteProxies = RayTracingNaniteProxiesCVar->GetValueOnGameThread();
 	const int32 RayTracingSkeletalMeshes = RayTracingSkeletalMeshesCVar->GetValueOnGameThread();
+	const int32 RayTracingISKM = RayTracingISKMCVar->GetValueOnGameThread();
 	const int32 RayTracingLandscapeGrass = RayTracingLandscapeGrassCVar->GetValueOnGameThread();
 
 	if (RayTracingStaticMeshes != CachedRayTracingStaticMeshes
 		|| RayTracingHISM != CachedRayTracingHISM
 		|| RayTracingNaniteProxies != CachedRayTracingNaniteProxies
 		|| RayTracingSkeletalMeshes != CachedRayTracingSkeletalMeshes
+		|| RayTracingISKM != CachedRayTracingISKM
 		|| RayTracingLandscapeGrass != CachedRayTracingLandscapeGrass)
 	{
 		ENQUEUE_RENDER_COMMAND(RefreshRayTracingInstancesCmd)(
@@ -134,6 +138,7 @@ static void RefreshRayTracingInstancesSinkFunction()
 		CachedRayTracingHISM = RayTracingHISM;
 		CachedRayTracingNaniteProxies = RayTracingNaniteProxies;
 		CachedRayTracingSkeletalMeshes = RayTracingSkeletalMeshes;
+		CachedRayTracingISKM = RayTracingISKM;
 		CachedRayTracingLandscapeGrass = RayTracingLandscapeGrass;
 	}
 }
