[Landscape][RayTracing] Refresh Cached Ray Tracing Instanced using Multithreading.

diff --git a/Engine/Source/Runtime/Renderer/Private/PrimitiveSceneInfo.cpp b/Engine/Source/Runtime/Renderer/Private/PrimitiveSceneInfo.cpp
--- a/Engine/Source/Runtime/Renderer/Private/PrimitiveSceneInfo.cpp
+++ b/Engine/Source/Runtime/Renderer/Private/PrimitiveSceneInfo.cpp
@@ -73,6 +73,20 @@ static FAutoConsoleVariableRef CVarRayTracingPrimitiveCacheMultithreaded(
 	TEXT("Enable multithreading of raytracing primitive mesh command caching. 0=disabled, 1=enabled (default)"),
 	ECVF_RenderThreadSafe);
 
+static int32 GRayTracingRefreshInstancesMultithreaded = 1;
+static FAutoConsoleVariableRef CVarRayTracingRefreshInstancesMultithreaded(
+	TEXT("r.RayTracing.RefreshInstancesMultithreaded"),
+	GRayTracingRefreshInstancesMultithreaded,
+	TEXT("Enable multithreading of ray tracing refresh primitive instances. 0=disables, 1=enabled (default)"),
+	ECVF_RenderThreadSafe);
+
+static int32 GRayTracingRefreshInstancesMultithreadedMinBatchSize = 2048;
+static FAutoConsoleVariableRef CVarRayTracingRefreshInstancesMultithreadedMinBatchSize(
+	TEXT("r.RayTracing.RefreshInstancesMultithreadedMinBatchSize"),
+	GRayTracingRefreshInstancesMultithreadedMinBatchSize,
+	TEXT("Determine minimum size of a single batch for the ray tracing refresh primitive instances."),
+	ECVF_RenderThreadSafe);
+
 /** An implementation of FStaticPrimitiveDrawInterface that stores the drawn elements for the rendering thread to use. */
 class FBatchingSPDI : public FStaticPrimitiveDrawInterface
 {
@@ -1000,34 +1014,39 @@ void FScene::RefreshRayTracingInstances()
 		SCOPED_NAMED_EVENT(FPrimitiveSceneInfo_UpdateCachedRayTracingInstances, FColor::Turquoise);
 		checkf(GRHISupportsMultithreadedShaderCreation, TEXT("Raytracing code needs the ability to create shaders from task threads."));
 
-		for (FPrimitiveSceneInfo* SceneInfo : Primitives)
-		{
-			const int32 PrimitiveIndex = SceneInfo->GetIndex();
-			FPrimitiveSceneProxy* SceneProxy = SceneInfo->Proxy;
+		EParallelForFlags ParallelForFlags = GRayTracingRefreshInstancesMultithreaded && FApp::ShouldUseThreadingForPerformance() ? EParallelForFlags::None : EParallelForFlags::ForceSingleThread;
 
-#if DO_CHECK
+		ParallelFor(TEXT("FPrimitiveSceneInfo_UpdateCachedRayTracingInstances"), Primitives.Num(), GRayTracingRefreshInstancesMultithreadedMinBatchSize, [this](int32 PrimitiveSceneInfoIndex)
 			{
-				Experimental::FHashElementId SceneRayTracingGroupId;
-				const int32 RayTracingGroupId = SceneProxy->GetRayTracingGroupId();
-				if (RayTracingGroupId != -1)
+				FPrimitiveSceneInfo* SceneInfo = Primitives[PrimitiveSceneInfoIndex];
+
+				const int32 PrimitiveIndex = SceneInfo->GetIndex();
+				FPrimitiveSceneProxy* SceneProxy = SceneInfo->Proxy;
+
+#if DO_CHECK
 				{
-					SceneRayTracingGroupId = PrimitiveRayTracingGroups.FindId(RayTracingGroupId);
-				}
+					Experimental::FHashElementId SceneRayTracingGroupId;
+					const int32 RayTracingGroupId = SceneProxy->GetRayTracingGroupId();
+					if (RayTracingGroupId != -1)
+					{
+						SceneRayTracingGroupId = PrimitiveRayTracingGroups.FindId(RayTracingGroupId);
+					}
 
-				check(PrimitiveRayTracingGroupIds[PrimitiveIndex] == SceneRayTracingGroupId);
+					check(PrimitiveRayTracingGroupIds[PrimitiveIndex] == SceneRayTracingGroupId);
 
-				const FScene::FPrimitiveRayTracingData& PrimitiveRayTracingData = PrimitiveRayTracingDatas[PrimitiveIndex];
-				CheckPrimitiveRayTracingData(PrimitiveRayTracingData, SceneProxy);
-			}
+					const FScene::FPrimitiveRayTracingData& PrimitiveRayTracingData = PrimitiveRayTracingDatas[PrimitiveIndex];
+					CheckPrimitiveRayTracingData(PrimitiveRayTracingData, SceneProxy);
+				}
 #endif
 
-			FRayTracingInstance CachedRayTracingInstance;
-			ERayTracingPrimitiveFlags& Flags = PrimitiveRayTracingFlags[PrimitiveIndex];
+				FRayTracingInstance CachedRayTracingInstance;
+				ERayTracingPrimitiveFlags& Flags = PrimitiveRayTracingFlags[PrimitiveIndex];
 
-			// Write flags
-			Flags = SceneProxy->GetCachedRayTracingInstance(CachedRayTracingInstance);
-			FPrimitiveSceneInfo::UpdateCachedRayTracingInstance(SceneInfo, CachedRayTracingInstance, Flags);
-		}
+				// Write flags
+				Flags = SceneProxy->GetCachedRayTracingInstance(CachedRayTracingInstance);
+				FPrimitiveSceneInfo::UpdateCachedRayTracingInstance(SceneInfo, CachedRayTracingInstance, Flags);
+
+			}, ParallelForFlags);
 	}
 }
 
