[HWRT] Invalidate all the cached ray tracing state for primitives when toggling r.RayTracing.Geometry.*.
- Disabling r.RayTracing.Geometry.* causes raytracing MDCs not to be cached so need to run RefreshRayTracingMeshCommandCache() instead of the lighterweight RefreshRayTracingInstances().

diff --git a/Engine/Source/Runtime/Renderer/Private/PrimitiveSceneInfo.cpp b/Engine/Source/Runtime/Renderer/Private/PrimitiveSceneInfo.cpp
--- a/Engine/Source/Runtime/Renderer/Private/PrimitiveSceneInfo.cpp
+++ b/Engine/Source/Runtime/Renderer/Private/PrimitiveSceneInfo.cpp
@@ -1018,6 +1018,8 @@ void FScene::RefreshRayTracingInstances()
 
 		ParallelFor(TEXT("FPrimitiveSceneInfo_UpdateCachedRayTracingInstances"), Primitives.Num(), GRayTracingRefreshInstancesMultithreadedMinBatchSize, [this](int32 PrimitiveSceneInfoIndex)
 			{
+				FTaskTagScope Scope(ETaskTag::EParallelRenderingThread);
+
 				FPrimitiveSceneInfo* SceneInfo = Primitives[PrimitiveSceneInfoIndex];
 
 				const int32 PrimitiveIndex = SceneInfo->GetIndex();
diff --git a/Engine/Source/Runtime/Renderer/Private/RayTracing/RayTracing.cpp b/Engine/Source/Runtime/Renderer/Private/RayTracing/RayTracing.cpp
--- a/Engine/Source/Runtime/Renderer/Private/RayTracing/RayTracing.cpp
+++ b/Engine/Source/Runtime/Renderer/Private/RayTracing/RayTracing.cpp
@@ -130,7 +130,12 @@ static void RefreshRayTracingInstancesSinkFunction()
 		ENQUEUE_RENDER_COMMAND(RefreshRayTracingInstancesCmd)(
 			[](FRHICommandListImmediate&)
 			{
-				bRefreshRayTracingInstances = true;
+				// When these cvars are disabled we skip caching ray tracing mesh commands
+				// so when they're toggled we need to invalidate the all the ray tracing data for the primitives
+				// not just the cached ray tracing instance data.
+				// TODO: implement more fine grained invalidation
+				//bRefreshRayTracingInstances = true;
+				bUpdateCachedRayTracingState = true;
 			}
 		);
 
