Default to going wide on GPU scene updates when the number of primitives / instances is large. Reduces hitches.

diff --git a/Engine/Source/Runtime/Renderer/Private/GPUScene.cpp b/Engine/Source/Runtime/Renderer/Private/GPUScene.cpp
--- a/Engine/Source/Runtime/Renderer/Private/GPUScene.cpp
+++ b/Engine/Source/Runtime/Renderer/Private/GPUScene.cpp
@@ -94,8 +94,9 @@ static TAutoConsoleVariable<int32> CVarGPUSceneMaxPooledUploadBufferSize(
 
 static TAutoConsoleVariable<int32> CVarGPUSceneParallelUpdate(
 	TEXT("r.GPUScene.ParallelUpdate"),
-	0,
-	TEXT(""),
+	2048,
+	TEXT("0  - disables parallel for updates of GPU scene primitives and instances.")
+	TEXT(">0 - enables parallel for updates if the number of items exceeds N."),
 	ECVF_RenderThreadSafe
 );
 
@@ -1231,7 +1232,7 @@ void FGPUScene::UploadGeneral(FRDGBuilder& GraphBuilder, const FRegisteredBuffer
 		return;
 	}
 
-	const bool bExecuteInParallel = CVarGPUSceneParallelUpdate.GetValueOnRenderThread() != 0 && FApp::ShouldUseThreadingForPerformance();
+	const int32 ParallelExecuteThreshold = FApp::ShouldUseThreadingForPerformance() ? CVarGPUSceneParallelUpdate.GetValueOnRenderThread() : 0;
 
 	SCOPED_NAMED_EVENT(UpdateGPUScene, FColor::Green);
 
@@ -1273,6 +1274,8 @@ void FGPUScene::UploadGeneral(FRDGBuilder& GraphBuilder, const FRegisteredBuffer
 		TaskContext.LightmapUploader = LightmapUploadBuffer.Begin(GraphBuilder, BufferState.LightmapDataBuffer, TaskContext.NumLightmapDataUploads, sizeof(FLightmapSceneShaderData::Data), TEXT("LightmapUploadBuffer"));
 	}
 
+	const bool bExecuteInParallel = ParallelExecuteThreshold > 0 && (TaskContext.NumInstancePayloadDataUploads > ParallelExecuteThreshold || TaskContext.NumPrimitiveDataUploads > ParallelExecuteThreshold);
+
 	GraphBuilder.AddCommandListSetupTask([&TaskContext, &UploadDataSourceAdapter, bExecuteInParallel, FeatureLevel = FeatureLevel](FRHICommandListBase& RHICmdList)
 	{
 		SCOPED_NAMED_EVENT(UpdateGPUScene_Primitives, FColor::Green);
@@ -1314,6 +1317,7 @@ void FGPUScene::UploadGeneral(FRDGBuilder& GraphBuilder, const FRegisteredBuffer
 
 				FVector4f* DstData = static_cast<FVector4f*>(TaskContext.PrimitiveUploader->GetRef(ItemIndex));
 				UploadDataSourceAdapter.GetPrimitiveShaderData(ItemIndex, DstData);
+
 			}, bExecuteInParallel ? EParallelForFlags::None : EParallelForFlags::ForceSingleThread);
 		}
 
