Fixed deadlock in -onethread -rdgimmediate where ray tracing task is waiting for frustum cull task that is kicked later in the frame.

diff --git a/Engine/Source/Runtime/Renderer/Private/DeferredShadingRenderer.cpp b/Engine/Source/Runtime/Renderer/Private/DeferredShadingRenderer.cpp
--- a/Engine/Source/Runtime/Renderer/Private/DeferredShadingRenderer.cpp
+++ b/Engine/Source/Runtime/Renderer/Private/DeferredShadingRenderer.cpp
@@ -1638,7 +1638,7 @@ void FDeferredShadingSceneRenderer::Render(FRDGBuilder& GraphBuilder, const FSce
 				GetViewPipelineState(ReferenceView).DiffuseIndirectMethod,
 				GetViewPipelineState(ReferenceView).ReflectionsMethod);
 
-			RayTracing::BeginGatherInstances(GraphBuilder, *InitViewTaskDatas.RayTracingGatherInstances, InitViewTaskDatas.VisibilityTaskData->GetFrustumCullTask());
+			RayTracing::BeginGatherInstances(*InitViewTaskDatas.RayTracingGatherInstances, InitViewTaskDatas.VisibilityTaskData->GetFrustumCullTask());
 		}
 #endif
 	}
diff --git a/Engine/Source/Runtime/Renderer/Private/RayTracing/RayTracing.cpp b/Engine/Source/Runtime/Renderer/Private/RayTracing/RayTracing.cpp
--- a/Engine/Source/Runtime/Renderer/Private/RayTracing/RayTracing.cpp
+++ b/Engine/Source/Runtime/Renderer/Private/RayTracing/RayTracing.cpp
@@ -1783,40 +1783,43 @@ namespace RayTracing
 		TaskData.DynamicRayTracingInstancesContext.AddInstancesToScene(RayTracingScene, RayTracingSBT, SharedBufferGenerationID);
 	}
 
-	void BeginGatherInstances(FRDGBuilder& GraphBuilder, FGatherInstancesTaskData& TaskData, UE::Tasks::FTask FrustumCullTask)
+	void BeginGatherInstances(FGatherInstancesTaskData& TaskData, UE::Tasks::FTask FrustumCullTask)
 	{
 		const bool bUsingReferenceBasedResidency = IsRayTracingUsingReferenceBasedResidency();
 
+		// Use high priority tasks to reduce stalls on the critical path.
+		const UE::Tasks::ETaskPriority TaskPriority = UE::Tasks::ETaskPriority::High;
+
 		UE::Tasks::FTask CacheRayTracingPrimitivesTask = TaskData.Scene.GetCacheRayTracingPrimitivesTask();
 
-		TaskData.GatherRelevantPrimitivesTask = GraphBuilder.AddSetupTask([&TaskData, bUsingReferenceBasedResidency]
+		TaskData.GatherRelevantPrimitivesTask = UE::Tasks::Launch(UE_SOURCE_LOCATION, [&TaskData, bUsingReferenceBasedResidency]
 			{
 				FTaskTagScope Scope(ETaskTag::EParallelRenderingThread);
 				GatherRelevantPrimitives(TaskData, bUsingReferenceBasedResidency);
-			}, MakeArrayView({ CacheRayTracingPrimitivesTask, FrustumCullTask }), UE::Tasks::ETaskPriority::High);
+			}, MakeArrayView({ CacheRayTracingPrimitivesTask, FrustumCullTask }), TaskPriority);
 
 		static const auto ICVarStaticMeshLODDistanceScale = IConsoleManager::Get().FindConsoleVariable(TEXT("r.StaticMeshLODDistanceScale"));
 		const float LODScaleCVarValue = ICVarStaticMeshLODDistanceScale->GetFloat();
 		const int32 ForcedLODLevel = GetCVarForceLOD();
 
-		TaskData.GatherRelevantStaticPrimitivesTask = GraphBuilder.AddSetupTask([&TaskData, LODScaleCVarValue, ForcedLODLevel, bUsingReferenceBasedResidency]
+		TaskData.GatherRelevantStaticPrimitivesTask = UE::Tasks::Launch(UE_SOURCE_LOCATION, [&TaskData, LODScaleCVarValue, ForcedLODLevel, bUsingReferenceBasedResidency]
 			{
 				FTaskTagScope Scope(ETaskTag::EParallelRenderingThread);
 				GatherRelevantStaticPrimitives(TaskData, LODScaleCVarValue, ForcedLODLevel, bUsingReferenceBasedResidency);
-			}, TaskData.GatherRelevantPrimitivesTask);
+			}, TaskData.GatherRelevantPrimitivesTask, TaskPriority);
 
 		TaskData.GatherDynamicRayTracingInstancesPrerequisites.AddPrerequisites(TaskData.GatherRelevantPrimitivesTask);
 		TaskData.GatherDynamicRayTracingInstancesPrerequisites.AddPrerequisites(TaskData.Scene.GetGPUSkinCacheTask());
 
 		// TODO: Could gather dynamic ray tracing instances using multiple tasks / FDynamicRayTracingInstancesContext
-		TaskData.GatherDynamicRayTracingInstancesTask = GraphBuilder.AddSetupTask([&TaskData]
+		TaskData.GatherDynamicRayTracingInstancesTask = UE::Tasks::Launch(UE_SOURCE_LOCATION, [&TaskData]
 			{
 				FTaskTagScope Scope(ETaskTag::EParallelRenderingThread);
 				TaskData.DynamicRayTracingInstancesContext.GatherDynamicRayTracingInstances(TaskData.DynamicPrimitivesIndices);
-			}, TaskData.GatherDynamicRayTracingInstancesPrerequisites);
+			}, TaskData.GatherDynamicRayTracingInstancesPrerequisites, TaskPriority);
 
 
-		TaskData.AddUsedStreamingHandlesTask = GraphBuilder.AddSetupTask([&TaskData, LODScaleCVarValue, ForcedLODLevel, bUsingReferenceBasedResidency]
+		TaskData.AddUsedStreamingHandlesTask = UE::Tasks::Launch(UE_SOURCE_LOCATION, [&TaskData, LODScaleCVarValue, ForcedLODLevel, bUsingReferenceBasedResidency]
 			{
 				FTaskTagScope Scope(ETaskTag::EParallelRenderingThread);
 
@@ -1826,7 +1829,7 @@ namespace RayTracing
 				{
 					CoarseMeshSM->AddUsedStreamingHandles(TaskData.UsedCoarseMeshStreamingHandles);
 				}
-			}, TaskData.GatherRelevantPrimitivesTask);
+			}, TaskData.GatherRelevantPrimitivesTask, TaskPriority);
 
 		TaskData.bValid = true;
 	}
diff --git a/Engine/Source/Runtime/Renderer/Private/RayTracing/RayTracing.h b/Engine/Source/Runtime/Renderer/Private/RayTracing/RayTracing.h
--- a/Engine/Source/Runtime/Renderer/Private/RayTracing/RayTracing.h
+++ b/Engine/Source/Runtime/Renderer/Private/RayTracing/RayTracing.h
@@ -53,7 +53,7 @@ namespace RayTracing
 		EDiffuseIndirectMethod DiffuseIndirectMethod,
 		EReflectionsMethod ReflectionsMethod);
 	
-	void BeginGatherInstances(FRDGBuilder& GraphBuilder, FGatherInstancesTaskData& TaskData, UE::Tasks::FTask FrustumCullTask);
+	void BeginGatherInstances(FGatherInstancesTaskData& TaskData, UE::Tasks::FTask FrustumCullTask);
 
 	void BeginGatherDynamicRayTracingInstances(FGatherInstancesTaskData& TaskData);
 
