Fixing ensure

diff --git a/Engine/Source/Runtime/Renderer/Private/HeterogeneousVolumes/HeterogeneousVolumes.cpp b/Engine/Source/Runtime/Renderer/Private/HeterogeneousVolumes/HeterogeneousVolumes.cpp
--- a/Engine/Source/Runtime/Renderer/Private/HeterogeneousVolumes/HeterogeneousVolumes.cpp
+++ b/Engine/Source/Runtime/Renderer/Private/HeterogeneousVolumes/HeterogeneousVolumes.cpp
@@ -654,7 +654,7 @@ namespace HeterogeneousVolumes
 			return 2;
 		}
 
-		int32 Value = FMath::Clamp(CVarHeterogeneousVolumesLightingCache.GetValueOnRenderThread(), 0, 2);
+		int32 Value = FMath::Clamp(CVarHeterogeneousVolumesLightingCache.GetValueOnAnyThread(), 0, 2);
 		return Value;
 	}
 
diff --git a/Engine/Source/Runtime/Renderer/Private/HeterogeneousVolumes/HeterogeneousVolumesLiveShadingPipeline.cpp b/Engine/Source/Runtime/Renderer/Private/HeterogeneousVolumes/HeterogeneousVolumesLiveShadingPipeline.cpp
--- a/Engine/Source/Runtime/Renderer/Private/HeterogeneousVolumes/HeterogeneousVolumesLiveShadingPipeline.cpp
+++ b/Engine/Source/Runtime/Renderer/Private/HeterogeneousVolumes/HeterogeneousVolumesLiveShadingPipeline.cpp
@@ -126,7 +126,7 @@ namespace HeterogeneousVolumes
 {
 	EScalabilityMode GetScalabilityMode()
 	{
-		int32 ScalabilityValue = FMath::Clamp(CVarHeterogeneousVolumesScalability.GetValueOnRenderThread(), 0, 3);
+		int32 ScalabilityValue = FMath::Clamp(CVarHeterogeneousVolumesScalability.GetValueOnAnyThread(), 0, 3);
 		return static_cast<EScalabilityMode>(ScalabilityValue);
 	}
 
