Remove deprecated eye adaptation texture code.

diff --git a/Engine/Source/Runtime/Engine/Private/SceneView.cpp b/Engine/Source/Runtime/Engine/Private/SceneView.cpp
--- a/Engine/Source/Runtime/Engine/Private/SceneView.cpp
+++ b/Engine/Source/Runtime/Engine/Private/SceneView.cpp
@@ -2852,17 +2852,6 @@ void FSceneView::SetupCommonViewUniformBufferParameters(
 	::SetupCommonViewUniformBufferParameters(ViewUniformShaderParameters, BufferSize, NumMSAASamples, EffectiveViewRect, InViewMatrices, InPrevViewMatrices, FSetupViewUniformParametersInputs::Create(*this));
 }
 
-bool FSceneView::HasValidEyeAdaptationTexture() const
-{
-	if (EyeAdaptationViewState)
-	{
-PRAGMA_DISABLE_DEPRECATION_WARNINGS
-		return EyeAdaptationViewState->HasValidEyeAdaptationTexture();
-PRAGMA_ENABLE_DEPRECATION_WARNINGS
-	}
-	return false;
-}
-
 bool FSceneView::HasValidEyeAdaptationBuffer() const
 {
 	if (EyeAdaptationViewState)
@@ -2872,18 +2861,6 @@ bool FSceneView::HasValidEyeAdaptationBuffer() const
 	return false;
 }
 
-IPooledRenderTarget* FSceneView::GetEyeAdaptationTexture() const
-{
-	checkf(FeatureLevel > ERHIFeatureLevel::ES3_1, TEXT("EyeAdaptation Texture is only available on SM5 and above."));
-	if (EyeAdaptationViewState)
-	{
-PRAGMA_DISABLE_DEPRECATION_WARNINGS
-		return EyeAdaptationViewState->GetCurrentEyeAdaptationTexture();
-PRAGMA_ENABLE_DEPRECATION_WARNINGS
-	}
-	return nullptr;
-}
-
 FRDGPooledBuffer* FSceneView::GetEyeAdaptationBuffer() const
 {
 	if (EyeAdaptationViewState)
diff --git a/Engine/Source/Runtime/Engine/Public/SceneManagement.h b/Engine/Source/Runtime/Engine/Public/SceneManagement.h
--- a/Engine/Source/Runtime/Engine/Public/SceneManagement.h
+++ b/Engine/Source/Runtime/Engine/Public/SceneManagement.h
@@ -128,8 +128,7 @@ class FSceneViewStateInterface
 {
 public:
 	FSceneViewStateInterface()
-		:	bValidEyeAdaptationTexture(0)
-		,	bValidEyeAdaptationBuffer(0)
+		:	bValidEyeAdaptationBuffer(0)
 	{}
 	
 	/** Called in the game thread to destroy the view state. */
@@ -198,20 +197,9 @@ public:
 	 */
 	virtual uint32 GetDistanceFieldTemporalSampleIndex() const = 0;
 
-	UE_DEPRECATED(5.2, "Use HasValidEyeAdaptationBuffer() instead.")
-	bool HasValidEyeAdaptationTexture() const
-	{
-PRAGMA_DISABLE_DEPRECATION_WARNINGS
-		return bValidEyeAdaptationTexture;
-PRAGMA_ENABLE_DEPRECATION_WARNINGS
-	}
-
 	/** Tells if the eye adaptation buffer exists without attempting to allocate it. */
 	bool HasValidEyeAdaptationBuffer() const { return bValidEyeAdaptationBuffer; }
 
-	UE_DEPRECATED(5.2, "Use GetCurrentEyeAdaptationBuffer() instead.")
-	virtual IPooledRenderTarget* GetCurrentEyeAdaptationTexture() const = 0;
-
 	/** Returns the eye adaptation buffer. */
 	virtual FRDGPooledBuffer* GetCurrentEyeAdaptationBuffer() const = 0;
 
@@ -279,8 +267,6 @@ protected:
 	// Don't allow direct deletion of the view state, Destroy should be called instead.
 	virtual ~FSceneViewStateInterface() {}
 
-	UE_DEPRECATED(5.2, "Use bValidEyeAdaptationBuffer instead.")
-	uint8 bValidEyeAdaptationTexture : 1;
 	uint8 bValidEyeAdaptationBuffer : 1;
 
 private:
diff --git a/Engine/Source/Runtime/Engine/Public/SceneView.h b/Engine/Source/Runtime/Engine/Public/SceneView.h
--- a/Engine/Source/Runtime/Engine/Public/SceneView.h
+++ b/Engine/Source/Runtime/Engine/Public/SceneView.h
@@ -2027,15 +2027,9 @@ public:
 	/** Current ray tracing debug visualization mode */
 	FName CurrentRayTracingDebugVisualizationMode;
 
-	UE_DEPRECATED(5.2, "Use HasValidEyeAdaptationBuffer() instead.")
-	ENGINE_API bool HasValidEyeAdaptationTexture() const;
-
 	/** Tells if the eye adaptation buffer exists without attempting to allocate it. */
 	ENGINE_API bool HasValidEyeAdaptationBuffer() const;
 
-	UE_DEPRECATED(5.2, "Use GetEyeAdaptationBuffer() instead.")
-	ENGINE_API IPooledRenderTarget* GetEyeAdaptationTexture() const;
-
 	/** Returns the eye adaptation buffer or null if it doesn't exist. */
 	ENGINE_API FRDGPooledBuffer* GetEyeAdaptationBuffer() const;
 
diff --git a/Engine/Source/Runtime/Renderer/Private/PostProcess/PostProcessEyeAdaptation.cpp b/Engine/Source/Runtime/Renderer/Private/PostProcess/PostProcessEyeAdaptation.cpp
--- a/Engine/Source/Runtime/Renderer/Private/PostProcess/PostProcessEyeAdaptation.cpp
+++ b/Engine/Source/Runtime/Renderer/Private/PostProcess/PostProcessEyeAdaptation.cpp
@@ -859,11 +859,6 @@ FRDGBufferRef AddHistogramEyeAdaptationPass(
 			PassParameters,
 			FIntVector(1, 1, 1));
 
-		{
-			FRDGTextureRef OutputTexture = GraphBuilder.RegisterExternalTexture(View.GetEyeAdaptationTexture(GraphBuilder), ERDGTextureFlags::MultiFrame);
-			AddCopyEyeAdaptationDataToTexturePass(GraphBuilder, View.ShaderMap, OutputBuffer, OutputTexture);
-		}
-
 		View.EnqueueEyeAdaptationExposureBufferReadback(GraphBuilder);
 	}
 	else
@@ -1013,13 +1008,6 @@ FRDGBufferRef AddBasicEyeAdaptationPass(
 			PassParameters,
 			FIntVector(1, 1, 1));
 
-PRAGMA_DISABLE_DEPRECATION_WARNINGS
-		{
-			FRDGTextureRef OutputTexture = GraphBuilder.RegisterExternalTexture(View.GetEyeAdaptationTexture(GraphBuilder), ERDGTextureFlags::MultiFrame);
-			AddCopyEyeAdaptationDataToTexturePass(GraphBuilder, View.ShaderMap, OutputBuffer, OutputTexture);
-		}
-PRAGMA_ENABLE_DEPRECATION_WARNINGS
-
 		View.EnqueueEyeAdaptationExposureBufferReadback(GraphBuilder);
 	}
 	else
diff --git a/Engine/Source/Runtime/Renderer/Private/RendererScene.cpp b/Engine/Source/Runtime/Renderer/Private/RendererScene.cpp
--- a/Engine/Source/Runtime/Renderer/Private/RendererScene.cpp
+++ b/Engine/Source/Runtime/Renderer/Private/RendererScene.cpp
@@ -393,9 +393,6 @@ void FSceneViewState::ReleaseRHI()
 	EyeAdaptationManager.SafeRelease(); 
 	SubstrateViewDebugData.SafeRelease();
 	OcclusionFeedback.ReleaseResource();
-PRAGMA_DISABLE_DEPRECATION_WARNINGS
-	bValidEyeAdaptationTexture = false;
-PRAGMA_ENABLE_DEPRECATION_WARNINGS
 	bValidEyeAdaptationBuffer = false;
 }
 
diff --git a/Engine/Source/Runtime/Renderer/Private/ScenePrivate.h b/Engine/Source/Runtime/Renderer/Private/ScenePrivate.h
--- a/Engine/Source/Runtime/Renderer/Private/ScenePrivate.h
+++ b/Engine/Source/Runtime/Renderer/Private/ScenePrivate.h
@@ -874,18 +874,6 @@ private:
 		/** Get the last frame average scene luminance (used for exposure compensation curve) */
 		float GetLastAverageSceneLuminance() const { return LastAverageSceneLuminance; }
 
-		UE_DEPRECATED(5.2, "Use GetCurrentBuffer() instead.")
-		const TRefCountPtr<IPooledRenderTarget>& GetCurrentTexture() const
-		{
-			return GetTexture(CurrentBufferIndex);
-		}
-
-		UE_DEPRECATED(5.2, "Use GetCurrentBuffer() instead.")
-		const TRefCountPtr<IPooledRenderTarget>& GetCurrentTexture(FRDGBuilder& GraphBuilder)
-		{
-			return GetOrCreateTexture(GraphBuilder, CurrentBufferIndex);
-		}
-
 		const TRefCountPtr<FRDGPooledBuffer>& GetCurrentBuffer() const
 		{
 			return GetBuffer(CurrentBufferIndex);
@@ -1283,25 +1271,6 @@ public:
 	 */
 	bool IsShadowOccluded(FSceneViewState::FProjectedShadowKey ShadowKey, int32 NumBufferedFrames) const;
 
-	UE_DEPRECATED(5.2, "Use GetCurrentEyeAdaptationBuffer() instead.")
-	IPooledRenderTarget* GetCurrentEyeAdaptationTexture() const override final
-	{
-PRAGMA_DISABLE_DEPRECATION_WARNINGS
-		IPooledRenderTarget* Texture = EyeAdaptationManager.GetCurrentTexture().GetReference();
-		check(bValidEyeAdaptationTexture && Texture);
-PRAGMA_ENABLE_DEPRECATION_WARNINGS
-		return Texture;
-	}
-
-	UE_DEPRECATED(5.2, "Use GetCurrentEyeAdaptationTexture() instead.")
-	IPooledRenderTarget* GetCurrentEyeAdaptationTexture(FRDGBuilder& GraphBuilder)
-	{
-PRAGMA_DISABLE_DEPRECATION_WARNINGS
-		bValidEyeAdaptationTexture = true;
-		return EyeAdaptationManager.GetCurrentTexture(GraphBuilder).GetReference();
-PRAGMA_ENABLE_DEPRECATION_WARNINGS
-	}
-
 	FRDGPooledBuffer* GetCurrentEyeAdaptationBuffer() const override final
 	{
 		FRDGPooledBuffer* Buffer = EyeAdaptationManager.GetCurrentBuffer().GetReference();
diff --git a/Engine/Source/Runtime/Renderer/Private/SceneRendering.cpp b/Engine/Source/Runtime/Renderer/Private/SceneRendering.cpp
--- a/Engine/Source/Runtime/Renderer/Private/SceneRendering.cpp
+++ b/Engine/Source/Runtime/Renderer/Private/SceneRendering.cpp
@@ -2358,19 +2358,6 @@ FSceneViewState* FViewInfo::GetEyeAdaptationViewState() const
 	return static_cast<FSceneViewState*>(EyeAdaptationViewState);
 }
 
-IPooledRenderTarget* FViewInfo::GetEyeAdaptationTexture(FRDGBuilder& GraphBuilder) const
-{
-	checkf(FeatureLevel > ERHIFeatureLevel::ES3_1, TEXT("EyeAdaptation Texture is only available on SM5 and above."));
-
-	if (FSceneViewState* EffectiveViewState = GetEyeAdaptationViewState())
-	{
-PRAGMA_DISABLE_DEPRECATION_WARNINGS
-		return EffectiveViewState->GetCurrentEyeAdaptationTexture(GraphBuilder);
-PRAGMA_ENABLE_DEPRECATION_WARNINGS
-	}
-	return nullptr;
-}
-
 FRDGPooledBuffer* FViewInfo::GetEyeAdaptationBuffer(FRDGBuilder& GraphBuilder) const
 {
 	if (FSceneViewState* EffectiveViewState = GetEyeAdaptationViewState())
diff --git a/Engine/Source/Runtime/Renderer/Private/SceneRendering.h b/Engine/Source/Runtime/Renderer/Private/SceneRendering.h
--- a/Engine/Source/Runtime/Renderer/Private/SceneRendering.h
+++ b/Engine/Source/Runtime/Renderer/Private/SceneRendering.h
@@ -1680,10 +1680,6 @@ public:
 	/** @return - whether this primitive has completely faded out */
 	bool UpdatePrimitiveFadingState(const FPrimitiveSceneInfo* PrimitiveSceneInfo, bool bFadingIn);
 
-	/** Allocates and returns the current eye adaptation texture. */
-	using FSceneView::GetEyeAdaptationTexture;
-	IPooledRenderTarget* GetEyeAdaptationTexture(FRDGBuilder& GraphBuilder) const;
-
 	/** Allocates and returns the current eye adaptation buffer. */
 	using FSceneView::GetEyeAdaptationBuffer;
 	FRDGPooledBuffer* GetEyeAdaptationBuffer(FRDGBuilder& GraphBuilder) const;
