Added VisualizeSkyLightIlluminanceMeter post process support on mobile path.

#rnx
#tests pceditor mobile preview

diff --git a/Engine/Source/Runtime/Renderer/Private/IlluminanceMeter.cpp b/Engine/Source/Runtime/Renderer/Private/IlluminanceMeter.cpp
--- a/Engine/Source/Runtime/Renderer/Private/IlluminanceMeter.cpp
+++ b/Engine/Source/Runtime/Renderer/Private/IlluminanceMeter.cpp
@@ -57,7 +57,7 @@ IMPLEMENT_GLOBAL_SHADER(FDownsampleIntegrateCubemapIlluminancePS, "/Engine/Priva
 static bool CanIlluminanceMeterDisplayOnPlatform(EShaderPlatform Platform)
 {
 	// On some consoles, this ALU heavy shader (and with optimisation disables for the sake of low compilation time) would spill registers. So only keep it for the editor.
-	return GetMaxSupportedFeatureLevel(Platform) >= ERHIFeatureLevel::SM5 && IsPCPlatform(Platform);
+	return GetMaxSupportedFeatureLevel(Platform) >= ERHIFeatureLevel::ES3_1;
 }
 
 class FPrintIlluminanceMeterPS : public FGlobalShader
diff --git a/Engine/Source/Runtime/Renderer/Private/PostProcess/PostProcessing.cpp b/Engine/Source/Runtime/Renderer/Private/PostProcess/PostProcessing.cpp
--- a/Engine/Source/Runtime/Renderer/Private/PostProcess/PostProcessing.cpp
+++ b/Engine/Source/Runtime/Renderer/Private/PostProcess/PostProcessing.cpp
@@ -1664,7 +1664,6 @@ void AddPostProcessingPasses(
 		}
 	}
 
-
 	if (PassSequence.IsEnabled(EPass::VisualizeLevelInstance))
 	{
 		FVisualizeLevelInstanceInputs PassInputs;
@@ -2390,6 +2389,7 @@ void AddMobilePostProcessingPasses(FRDGBuilder& GraphBuilder, FScene* Scene, con
 		FXAA,
 		PostProcessMaterialAfterTonemapping,
 		HighResolutionScreenshotMask,
+		VisualizeSkyLightIlluminanceMeter,
 		SelectionOutline,
 		EditorPrimitive,
 #if UE_ENABLE_DEBUG_DRAWING
@@ -2434,6 +2434,7 @@ void AddMobilePostProcessingPasses(FRDGBuilder& GraphBuilder, FScene* Scene, con
 		TEXT("PostProcessMaterial (AfterTonemapping)"),
 		TEXT("FXAA"),
 		TEXT("HighResolutionScreenshotMask"),
+		TEXT("VisualizeSkyLightIlluminanceMeter"),
 		TEXT("SelectionOutline"),
 		TEXT("EditorPrimitive"),
 #if UE_ENABLE_DEBUG_DRAWING
@@ -2489,9 +2490,11 @@ void AddMobilePostProcessingPasses(FRDGBuilder& GraphBuilder, FScene* Scene, con
 	PassSequence.SetEnabled(EPass::Tonemap, bUseToneMapper);
 	PassSequence.SetEnabled(EPass::HighResolutionScreenshotMask, bUseHighResolutionScreenshotMask);
 #if WITH_EDITOR
+	PassSequence.SetEnabled(EPass::VisualizeSkyLightIlluminanceMeter, Scene&& Scene->SkyLight&& View.Family&& View.Family->EngineShowFlags.VisualizeSkyLightIlluminance);
 	PassSequence.SetEnabled(EPass::SelectionOutline, GIsEditor && View.Family->EngineShowFlags.Selection && View.Family->EngineShowFlags.SelectionOutline && !View.Family->EngineShowFlags.Wireframe);
 	PassSequence.SetEnabled(EPass::EditorPrimitive, FSceneRenderer::ShouldCompositeEditorPrimitives(View));
 #else
+	PassSequence.SetEnabled(EPass::VisualizeSkyLightIlluminanceMeter, false);
 	PassSequence.SetEnabled(EPass::SelectionOutline, false);
 	PassSequence.SetEnabled(EPass::EditorPrimitive, false);
 #endif
@@ -3179,6 +3182,14 @@ void AddMobilePostProcessingPasses(FRDGBuilder& GraphBuilder, FScene* Scene, con
 	}
 
 #if WITH_EDITOR
+
+	if (PassSequence.IsEnabled(EPass::VisualizeSkyLightIlluminanceMeter))
+	{
+		FScreenPassRenderTarget OverrideOutput;
+		PassSequence.AcceptOverrideIfLastPass(EPass::VisualizeSkyLightIlluminanceMeter, OverrideOutput);
+		SceneColor = ProcessAndRenderIlluminanceMeter(GraphBuilder, View, SceneColor);
+	}
+
 	// Show the selection outline if it is in the editor and we aren't in wireframe 
 	// If the engine is in demo mode and game view is on we also do not show the selection outline
 	if (PassSequence.IsEnabled(EPass::SelectionOutline))
