Correct GPU markup for DistanceFieldShadows

#tests local replayrun

diff --git a/Engine/Source/Runtime/Renderer/Private/ShadowRendering.cpp b/Engine/Source/Runtime/Renderer/Private/ShadowRendering.cpp
--- a/Engine/Source/Runtime/Renderer/Private/ShadowRendering.cpp
+++ b/Engine/Source/Runtime/Renderer/Private/ShadowRendering.cpp
@@ -2182,7 +2182,7 @@ void FSceneRenderer::RenderShadowProjections(
 
 	if (ScreenShadowMaskTexture && DistanceFieldShadows.Num() > 0)
 	{
-		RDG_EVENT_SCOPE(GraphBuilder, "DistanceFieldShadows");
+		RDG_EVENT_SCOPE_STAT(GraphBuilder, DistanceFieldShadows, "DistanceFieldShadows");
 		RDG_GPU_STAT_SCOPE(GraphBuilder, DistanceFieldShadows);
 
 		// Distance field shadows need to be renderer last as they blend over far shadow cascades.
@@ -2228,7 +2228,7 @@ void FSceneRenderer::BeginAsyncDistanceFieldShadowProjections(FRDGBuilder& Graph
 
 	if (!!GDFShadowAsyncCompute && ViewFamily.EngineShowFlags.DynamicShadows && GetShadowQuality() > 0 && ProjectedDistanceFieldShadows.Num() > 0)
 	{
-		RDG_EVENT_SCOPE(GraphBuilder, "DistanceFieldShadows");
+		RDG_EVENT_SCOPE_STAT(GraphBuilder, DistanceFieldShadows, "DistanceFieldShadows");
 		RDG_GPU_STAT_SCOPE(GraphBuilder, DistanceFieldShadows);
 
 		for (int32 ViewIndex = 0; ViewIndex < Views.Num(); ViewIndex++)
