Fixed unhappy static analysis.

#tests PC with static code analyser
#rnx

diff --git a/Engine/Source/Runtime/Renderer/Private/PostProcess/PostProcessing.cpp b/Engine/Source/Runtime/Renderer/Private/PostProcess/PostProcessing.cpp
--- a/Engine/Source/Runtime/Renderer/Private/PostProcess/PostProcessing.cpp
+++ b/Engine/Source/Runtime/Renderer/Private/PostProcess/PostProcessing.cpp
@@ -2490,7 +2490,7 @@ void AddMobilePostProcessingPasses(FRDGBuilder& GraphBuilder, FScene* Scene, con
 	PassSequence.SetEnabled(EPass::Tonemap, bUseToneMapper);
 	PassSequence.SetEnabled(EPass::HighResolutionScreenshotMask, bUseHighResolutionScreenshotMask);
 #if WITH_EDITOR
-	PassSequence.SetEnabled(EPass::VisualizeSkyLightIlluminanceMeter, Scene&& Scene->SkyLight&& View.Family&& View.Family->EngineShowFlags.VisualizeSkyLightIlluminance);
+	PassSequence.SetEnabled(EPass::VisualizeSkyLightIlluminanceMeter, Scene && Scene->SkyLight && View.Family->EngineShowFlags.VisualizeSkyLightIlluminance);
 	PassSequence.SetEnabled(EPass::SelectionOutline, GIsEditor && View.Family->EngineShowFlags.Selection && View.Family->EngineShowFlags.SelectionOutline && !View.Family->EngineShowFlags.Wireframe);
 	PassSequence.SetEnabled(EPass::EditorPrimitive, FSceneRenderer::ShouldCompositeEditorPrimitives(View));
 #else
