[SingleLayerWater] Fix crash when trying to access CustomDepth in a SLW material in the hitproxy pass. Also fixed Custom Depth/Stencil and SceneDepth nodes not working properly in the SLW depth prepass.

diff --git a/Engine/Shaders/Private/MaterialTemplate.ush b/Engine/Shaders/Private/MaterialTemplate.ush
--- a/Engine/Shaders/Private/MaterialTemplate.ush
+++ b/Engine/Shaders/Private/MaterialTemplate.ush
@@ -3167,7 +3167,7 @@ float4 SceneTextureLookup(float2 UV, int SceneTextureIndex, bool bFiltered)
 		return 0.0f;
 	}
 
-#elif MATERIAL_SHADINGMODEL_SINGLELAYERWATER
+#elif MATERIAL_SHADINGMODEL_SINGLELAYERWATER && !HIT_PROXY_SHADER
 
 	if (SceneTextureIndex == PPI_CustomDepth)
 	{
diff --git a/Engine/Source/Runtime/Core/Private/UObject/FortniteReleaseBranchCustomObjectVersion.cpp b/Engine/Source/Runtime/Core/Private/UObject/FortniteReleaseBranchCustomObjectVersion.cpp
--- a/Engine/Source/Runtime/Core/Private/UObject/FortniteReleaseBranchCustomObjectVersion.cpp
+++ b/Engine/Source/Runtime/Core/Private/UObject/FortniteReleaseBranchCustomObjectVersion.cpp
@@ -8,7 +8,7 @@ TMap<FGuid, FGuid> FFortniteReleaseBranchCustomObjectVersion::GetSystemGuids()
 
 	SystemGuids.Add(DevGuids.Niagara_LatestScriptCompileVersion, FGuid("8BC8C557C87C4F629468F10EA4161F0A"));
 	SystemGuids.Add(DevGuids.SkeletalMeshDerivedDataVersion, FGuid("C96DEBAD45E3443D804017D4C29F7C4B"));
-	SystemGuids.Add(DevGuids.MaterialTranslationDDCVersion, FGuid("3EF80806F2E945B1ADAFC49154058C00"));
+	SystemGuids.Add(DevGuids.MaterialTranslationDDCVersion, FGuid("5958B746DC5E4046AA7FB7E30E7AC509"));
 	SystemGuids.Add(DevGuids.GLOBALSHADERMAP_DERIVEDDATA_VER, FGuid("0A0A1BF5D2A94105BB395EF0DA7CAEB9"));
 	SystemGuids.Add(DevGuids.MATERIALSHADERMAP_DERIVEDDATA_VER, FGuid("C0E1F59386894C9D855B85178B2553CB"));
 	SystemGuids.Add(DevGuids.STATICMESH_DERIVEDDATA_VER, FGuid("0D7C30DE6E5F489A831CEBF1346725C4"));
diff --git a/Engine/Source/Runtime/Renderer/Private/DeferredShadingRenderer.cpp b/Engine/Source/Runtime/Renderer/Private/DeferredShadingRenderer.cpp
--- a/Engine/Source/Runtime/Renderer/Private/DeferredShadingRenderer.cpp
+++ b/Engine/Source/Runtime/Renderer/Private/DeferredShadingRenderer.cpp
@@ -2400,16 +2400,6 @@ void FDeferredShadingSceneRenderer::Render(FRDGBuilder& GraphBuilder, const FSce
 			ViewExtension->PreRenderBasePass_RenderThread(GraphBuilder, ShouldRenderPrePass() /*bDepthBufferIsPopulated*/);
 		}
 
-		// Single layer water depth prepass. Needs to run before VSM page allocation. If there's a full depth prepass, it can run before the base pass, otherwise after.
-		// Running before the base pass allows for some optimizations to save work in the base pass and lighting stages.
-		const FSingleLayerWaterPrePassResult* SingleLayerWaterPrePassResult = nullptr;
-		const ESingleLayerWaterPrepassLocation SingleLayerWaterPrepassLocation = GetSingleLayerWaterDepthPrepassLocation(bIsEarlyDepthComplete);
-		const bool bShouldRenderSingleLayerWaterDepthPrepass = !bHasRayTracedOverlay && ShouldRenderSingleLayerWaterDepthPrepass(Views);
-		if (bShouldRenderSingleLayerWaterDepthPrepass && SingleLayerWaterPrepassLocation == ESingleLayerWaterPrepassLocation::BeforeBasePass)
-		{
-			SingleLayerWaterPrePassResult = RenderSingleLayerWaterDepthPrepass(GraphBuilder, Views, SceneTextures, SingleLayerWaterPrepassLocation, NaniteRasterResults);
-		}
-
 		{
 			SCOPE_CYCLE_COUNTER(STAT_WaitGatherAndSortLightsTask);
 			GatherAndSortLightsTask.Wait();
@@ -2495,6 +2485,16 @@ void FDeferredShadingSceneRenderer::Render(FRDGBuilder& GraphBuilder, const FSce
 			}
 		}
 
+		// Single layer water depth prepass. Needs to run before VSM page allocation. If there's a full depth prepass, it can run before the base pass, otherwise after.
+		// Running before the base pass allows for some optimizations to save work in the base pass and lighting stages.
+		const FSingleLayerWaterPrePassResult* SingleLayerWaterPrePassResult = nullptr;
+		const ESingleLayerWaterPrepassLocation SingleLayerWaterPrepassLocation = GetSingleLayerWaterDepthPrepassLocation(bIsEarlyDepthComplete, CustomDepthPassLocation);
+		const bool bShouldRenderSingleLayerWaterDepthPrepass = !bHasRayTracedOverlay && ShouldRenderSingleLayerWaterDepthPrepass(Views);
+		if (bShouldRenderSingleLayerWaterDepthPrepass && SingleLayerWaterPrepassLocation == ESingleLayerWaterPrepassLocation::BeforeBasePass)
+		{
+			SingleLayerWaterPrePassResult = RenderSingleLayerWaterDepthPrepass(GraphBuilder, Views, SceneTextures, SingleLayerWaterPrepassLocation, NaniteRasterResults);
+		}
+
 		// Lumen updates need access to sky atmosphere LUT.
 		ExternalAccessQueue.Submit(GraphBuilder);
 
diff --git a/Engine/Source/Runtime/Renderer/Private/DepthRendering.h b/Engine/Source/Runtime/Renderer/Private/DepthRendering.h
--- a/Engine/Source/Runtime/Renderer/Private/DepthRendering.h
+++ b/Engine/Source/Runtime/Renderer/Private/DepthRendering.h
@@ -145,7 +145,11 @@ public:
 		FMeshMaterialShader::ModifyCompilationEnvironment(Parameters, OutEnvironment);
 		
 		OutEnvironment.SetDefine(TEXT("ALLOW_DEBUG_VIEW_MODES"), AllowDebugViewmodes(Parameters.Platform));
-		OutEnvironment.SetDefine(TEXT("SCENE_TEXTURES_DISABLED"), 1u);
+
+		// Regular opaque materials do not support scene textures, but SLW does support reading custom depth/stencil,
+		// so we need to make an exception here for the water depth prepass.
+		const bool bIsSingleLayerWater = Parameters.MaterialParameters.ShadingModels.HasShadingModel(MSM_SingleLayerWater);
+		OutEnvironment.SetDefine(TEXT("SCENE_TEXTURES_DISABLED"), bIsSingleLayerWater ? 0u : 1u);
 	}
 
 	FDepthOnlyPS() {}
diff --git a/Engine/Source/Runtime/Renderer/Private/SingleLayerWaterRendering.cpp b/Engine/Source/Runtime/Renderer/Private/SingleLayerWaterRendering.cpp
--- a/Engine/Source/Runtime/Renderer/Private/SingleLayerWaterRendering.cpp
+++ b/Engine/Source/Runtime/Renderer/Private/SingleLayerWaterRendering.cpp
@@ -22,6 +22,7 @@
 #include "Lumen/LumenTracingUtils.h"
 #include "RenderCore.h"
 #include "UnrealEngine.h"
+#include "CustomDepthRendering.h"
 
 #include "DepthCopy.h"
 using namespace DepthCopy;
@@ -265,9 +266,9 @@ bool ShouldRenderSingleLayerWaterDepthPrepass(TArrayView<const FViewInfo> Views)
 	return bPrepassEnabled && bShouldRenderWater;
 }
 
-ESingleLayerWaterPrepassLocation GetSingleLayerWaterDepthPrepassLocation(bool bFullDepthPrepass)
+ESingleLayerWaterPrepassLocation GetSingleLayerWaterDepthPrepassLocation(bool bFullDepthPrepass, ECustomDepthPassLocation CustomDepthPassLocation)
 {
-	if (bFullDepthPrepass)
+	if (bFullDepthPrepass && CustomDepthPassLocation == ECustomDepthPassLocation::BeforeBasePass)
 	{
 		return ESingleLayerWaterPrepassLocation::BeforeBasePass;
 	}
@@ -346,6 +347,29 @@ bool NeedsSeparatedMainDirectionalLightTexture_Runtime(const FStaticShaderPlatfo
 	return IsWaterDistanceFieldShadowEnabled_Runtime(Platform) || IsWaterVirtualShadowMapFilteringEnabled_Runtime(Platform);
 }
 
+BEGIN_UNIFORM_BUFFER_STRUCT(FSingleLayerWaterPassUniformParameters, )
+	SHADER_PARAMETER_RDG_TEXTURE(Texture2D, SceneColorWithoutSingleLayerWaterTexture)
+	SHADER_PARAMETER_SAMPLER(SamplerState, SceneColorWithoutSingleLayerWaterSampler)
+	SHADER_PARAMETER_RDG_TEXTURE(Texture2D, SceneDepthWithoutSingleLayerWaterTexture)
+	SHADER_PARAMETER_SAMPLER(SamplerState, SceneDepthWithoutSingleLayerWaterSampler)
+	SHADER_PARAMETER_RDG_TEXTURE(Texture2D, CustomDepthTexture)
+	SHADER_PARAMETER_RDG_TEXTURE_SRV(Texture2D<uint2>, CustomStencilTexture)
+	SHADER_PARAMETER_SAMPLER(SamplerState, CustomDepthSampler)
+	SHADER_PARAMETER_RDG_TEXTURE(Texture2D, RefractionMaskTexture)
+	SHADER_PARAMETER(FVector4f, SceneWithoutSingleLayerWaterMinMaxUV)
+	SHADER_PARAMETER(FVector4f, DistortionParams)
+	SHADER_PARAMETER(FVector2f, SceneWithoutSingleLayerWaterTextureSize)
+	SHADER_PARAMETER(FVector2f, SceneWithoutSingleLayerWaterInvTextureSize)
+	SHADER_PARAMETER(uint32, bMainDirectionalLightVSMFiltering)
+	SHADER_PARAMETER(uint32, bSeparateMainDirLightLuminance)
+	SHADER_PARAMETER_STRUCT(FLightCloudTransmittanceParameters, ForwardDirLightCloudShadow)
+	SHADER_PARAMETER_STRUCT(FBlueNoiseParameters, BlueNoise)
+END_UNIFORM_BUFFER_STRUCT()
+
+// At the moment we reuse the DeferredDecals static uniform buffer slot because it is currently unused in this pass.
+// When we add support for decals on SLW in the future, we might need to find another solution.
+IMPLEMENT_STATIC_UNIFORM_BUFFER_STRUCT(FSingleLayerWaterPassUniformParameters, "SingleLayerWater", DeferredDecals);
+
 BEGIN_SHADER_PARAMETER_STRUCT(FSingleLayerWaterCommonShaderParameters, )
 	SHADER_PARAMETER_RDG_TEXTURE_SRV(Texture2D, ScreenSpaceReflectionsTexture)
 	SHADER_PARAMETER_SAMPLER(SamplerState, ScreenSpaceReflectionsSampler)
@@ -646,17 +670,62 @@ class FWaterRefractionCopyPS : public FGlobalShader
 
 IMPLEMENT_GLOBAL_SHADER(FWaterRefractionCopyPS, "/Engine/Private/SingleLayerWaterComposite.usf", "WaterRefractionCopyPS", SF_Pixel);
 
+bool IsVSMTranslucentHighQualityEnabled();
+
+FSingleLayerWaterPassUniformParameters* CreateSingleLayerWaterPassUniformParameters(FRDGBuilder& GraphBuilder, const FScene* Scene, const FViewInfo& View, const FSceneTextures& SceneTextures, FRDGTextureRef SceneDepthWithoutWater, FRDGTextureRef SceneColorWithoutWater, FRDGTextureRef RefractionMaskTexture, const FVector4f& MinMaxUV)
+{
+	FSingleLayerWaterPassUniformParameters* SLWUniformParameters = GraphBuilder.AllocParameters<FSingleLayerWaterPassUniformParameters>();
+	const bool bShouldUseBilinearSamplerForDepth = ShouldUseBilinearSamplerForDepthWithoutSingleLayerWater(SceneDepthWithoutWater->Desc.Format);
+	const bool bCustomDepthTextureProduced = HasBeenProduced(SceneTextures.CustomDepth.Depth);
+	const FIntVector DepthTextureSize = SceneDepthWithoutWater->Desc.GetSize();
+	const FRDGSystemTextures& SystemTextures = FRDGSystemTextures::Get(GraphBuilder);
+
+	SLWUniformParameters->SceneColorWithoutSingleLayerWaterTexture = SceneColorWithoutWater ? SceneColorWithoutWater : SystemTextures.Black;
+	SLWUniformParameters->SceneColorWithoutSingleLayerWaterSampler = TStaticSamplerState<SF_Bilinear>::GetRHI();
+	SLWUniformParameters->SceneDepthWithoutSingleLayerWaterTexture = SceneDepthWithoutWater;
+	SLWUniformParameters->SceneDepthWithoutSingleLayerWaterSampler = bShouldUseBilinearSamplerForDepth ? TStaticSamplerState<SF_Bilinear>::GetRHI() : TStaticSamplerState<SF_Point>::GetRHI();
+	SLWUniformParameters->CustomDepthTexture = bCustomDepthTextureProduced ? SceneTextures.CustomDepth.Depth : SystemTextures.DepthDummy;
+	SLWUniformParameters->CustomStencilTexture = bCustomDepthTextureProduced ? SceneTextures.CustomDepth.Stencil : SystemTextures.StencilDummySRV;
+	SLWUniformParameters->CustomDepthSampler = TStaticSamplerState<SF_Point>::GetRHI();
+	SLWUniformParameters->RefractionMaskTexture = RefractionMaskTexture ? RefractionMaskTexture : SystemTextures.White;
+	SLWUniformParameters->SceneWithoutSingleLayerWaterMinMaxUV = MinMaxUV;
+	SetupDistortionParams(SLWUniformParameters->DistortionParams, View);
+	SLWUniformParameters->SceneWithoutSingleLayerWaterTextureSize = FVector2f(DepthTextureSize.X, DepthTextureSize.Y);
+	SLWUniformParameters->SceneWithoutSingleLayerWaterInvTextureSize = FVector2f(1.0f / DepthTextureSize.X, 1.0f / DepthTextureSize.Y);
+	SLWUniformParameters->bMainDirectionalLightVSMFiltering = IsWaterVirtualShadowMapFilteringEnabled_Runtime(View.GetShaderPlatform());
+	SLWUniformParameters->bSeparateMainDirLightLuminance = NeedsSeparatedMainDirectionalLightTexture_Runtime(View.GetShaderPlatform());
+	// Only use blue noise resources if VSM quality is set to high
+	if (IsVSMTranslucentHighQualityEnabled())
+	{
+		SLWUniformParameters->BlueNoise = GetBlueNoiseParameters();
+	}
+	else
+	{
+		SLWUniformParameters->BlueNoise = GetBlueNoiseDummyParameters();
+	}
+
+	const FLightSceneProxy* SelectedForwardDirectionalLightProxy = View.ForwardLightingResources.SelectedForwardDirectionalLightProxy;
+	SetupLightCloudTransmittanceParameters(GraphBuilder, Scene, View, SelectedForwardDirectionalLightProxy ? SelectedForwardDirectionalLightProxy->GetLightSceneInfo() : nullptr, SLWUniformParameters->ForwardDirLightCloudShadow);
+
+	return SLWUniformParameters;
+}
+
 BEGIN_SHADER_PARAMETER_STRUCT(FSingleLayerWaterDepthPassParameters, )
 	SHADER_PARAMETER_STRUCT_INCLUDE(FViewShaderParameters, View)
 	SHADER_PARAMETER_STRUCT_INCLUDE(FInstanceCullingDrawParams, InstanceCullingDrawParams)
+	SHADER_PARAMETER_RDG_UNIFORM_BUFFER(FSingleLayerWaterPassUniformParameters, SingleLayerWater)
 	RENDER_TARGET_BINDING_SLOTS()
 END_SHADER_PARAMETER_STRUCT()
 
-static FSingleLayerWaterDepthPassParameters* GetSingleLayerWaterDepthPassParameters(FRDGBuilder& GraphBuilder, const FViewInfo& View, FRDGTextureRef DepthTexture)
+static FSingleLayerWaterDepthPassParameters* GetSingleLayerWaterDepthPassParameters(FRDGBuilder& GraphBuilder, const FScene* Scene, const FViewInfo& View, const FSceneTextures& SceneTextures, FRDGTextureRef SceneDepthWithoutWater, FRDGTextureRef DepthTexture)
 {
 	FSingleLayerWaterDepthPassParameters* PassParameters = GraphBuilder.AllocParameters<FSingleLayerWaterDepthPassParameters>();
 	PassParameters->View = View.GetShaderParameters();
 	PassParameters->RenderTargets.DepthStencil = FDepthStencilBinding(DepthTexture, ERenderTargetLoadAction::ELoad, ERenderTargetLoadAction::ELoad, FExclusiveDepthStencil::DepthWrite_StencilWrite);
+
+	FSingleLayerWaterPassUniformParameters* SLWUniformParameters = CreateSingleLayerWaterPassUniformParameters(GraphBuilder, Scene, View, SceneTextures, SceneDepthWithoutWater, nullptr, nullptr, FVector4f(0.0f, 0.0f, 1.0f, 1.0f));
+	PassParameters->SingleLayerWater = GraphBuilder.CreateUniformBuffer(SLWUniformParameters);
+
 	return PassParameters;
 }
 
@@ -664,7 +733,7 @@ static FSingleLayerWaterDepthPassParameters* GetSingleLayerWaterDepthPassParamet
  * Build lists of 8x8 tiles used by water pixels
  * Mark and build list steps are separated in order to build a more coherent list (z-ordered over a larger region), which is important for the performance of future passes like ray traced Lumen reflections
  */
-static FSingleLayerWaterTileClassification ClassifyTiles(FRDGBuilder& GraphBuilder, const FViewInfo &View, const FSceneTextures& SceneTextures, const FRDGTextureRef& DepthPrepassTexture, const Froxel::FViewData* OutFroxelViewData, EReflectionsMethod ReflectionsMethod)
+static FSingleLayerWaterTileClassification ClassifyTiles(FRDGBuilder& GraphBuilder, const FViewInfo &View, const FSceneTextures& SceneTextures, FRDGTextureRef DepthPrepassTexture, const Froxel::FViewData* OutFroxelViewData, EReflectionsMethod ReflectionsMethod)
 {
 	FSingleLayerWaterTileClassification Result;
 	const bool bRunTiled = UseSingleLayerWaterIndirectDraw(View.GetShaderPlatform()) && CVarWaterSingleLayerTiledComposite.GetValueOnRenderThread();
@@ -1027,7 +1096,7 @@ FSingleLayerWaterPrePassResult* FDeferredShadingSceneRenderer::RenderSingleLayer
 		RDG_EVENT_SCOPE_CONDITIONAL(GraphBuilder, InViews.Num() > 1, "View%d", ViewIndex);
 		View.BeginRenderView();
 
-		FSingleLayerWaterDepthPassParameters* PassParameters = GetSingleLayerWaterDepthPassParameters(GraphBuilder, View, OutDepthPrepassTexture.Target);
+		FSingleLayerWaterDepthPassParameters* PassParameters = GetSingleLayerWaterDepthPassParameters(GraphBuilder, Scene, View, SceneTextures, Result->SceneDepthWithoutWater, OutDepthPrepassTexture.Target);
 
 		Pass->BuildRenderingCommands(GraphBuilder, Scene->GPUScene, PassParameters->InstanceCullingDrawParams);
 
@@ -1663,29 +1732,6 @@ void FDeferredShadingSceneRenderer::RenderSingleLayerWater(
 	}
 }
 
-BEGIN_UNIFORM_BUFFER_STRUCT(FSingleLayerWaterPassUniformParameters,)
-	SHADER_PARAMETER_RDG_TEXTURE(Texture2D, SceneColorWithoutSingleLayerWaterTexture)
-	SHADER_PARAMETER_SAMPLER(SamplerState, SceneColorWithoutSingleLayerWaterSampler)
-	SHADER_PARAMETER_RDG_TEXTURE(Texture2D, SceneDepthWithoutSingleLayerWaterTexture)
-	SHADER_PARAMETER_SAMPLER(SamplerState, SceneDepthWithoutSingleLayerWaterSampler)
-	SHADER_PARAMETER_RDG_TEXTURE(Texture2D, CustomDepthTexture)
-	SHADER_PARAMETER_RDG_TEXTURE_SRV(Texture2D<uint2>, CustomStencilTexture)
-	SHADER_PARAMETER_SAMPLER(SamplerState, CustomDepthSampler)
-	SHADER_PARAMETER_RDG_TEXTURE(Texture2D, RefractionMaskTexture)
-	SHADER_PARAMETER(FVector4f, SceneWithoutSingleLayerWaterMinMaxUV)
-	SHADER_PARAMETER(FVector4f, DistortionParams)
-	SHADER_PARAMETER(FVector2f, SceneWithoutSingleLayerWaterTextureSize)
-	SHADER_PARAMETER(FVector2f, SceneWithoutSingleLayerWaterInvTextureSize)
-	SHADER_PARAMETER(uint32, bMainDirectionalLightVSMFiltering)
-	SHADER_PARAMETER(uint32, bSeparateMainDirLightLuminance)
-	SHADER_PARAMETER_STRUCT(FLightCloudTransmittanceParameters, ForwardDirLightCloudShadow)
-	SHADER_PARAMETER_STRUCT(FBlueNoiseParameters, BlueNoise)
-END_UNIFORM_BUFFER_STRUCT()
-
-// At the moment we reuse the DeferredDecals static uniform buffer slot because it is currently unused in this pass.
-// When we add support for decals on SLW in the future, we might need to find another solution.
-IMPLEMENT_STATIC_UNIFORM_BUFFER_STRUCT(FSingleLayerWaterPassUniformParameters, "SingleLayerWater", DeferredDecals);
-
 BEGIN_SHADER_PARAMETER_STRUCT(FSingleLayerWaterPassParameters, )
 	SHADER_PARAMETER_STRUCT_INCLUDE(FViewShaderParameters, View)
 	SHADER_PARAMETER_STRUCT_REF(FReflectionCaptureShaderData, ReflectionCapture)
@@ -1696,8 +1742,6 @@ BEGIN_SHADER_PARAMETER_STRUCT(FSingleLayerWaterPassParameters, )
 	RENDER_TARGET_BINDING_SLOTS()
 END_SHADER_PARAMETER_STRUCT()
 
-bool IsVSMTranslucentHighQualityEnabled();
-
 void FDeferredShadingSceneRenderer::RenderSingleLayerWaterInner(
 	FRDGBuilder& GraphBuilder,
 	TArrayView<FViewInfo> InViews,
@@ -1753,46 +1797,15 @@ void FDeferredShadingSceneRenderer::RenderSingleLayerWaterInner(
 		RDG_EVENT_SCOPE_CONDITIONAL(GraphBuilder, InViews.Num() > 1, "View%d", ViewIndex);
 		View.BeginRenderView();
 
-		FSingleLayerWaterPassUniformParameters& SLWUniformParameters = *GraphBuilder.AllocParameters<FSingleLayerWaterPassUniformParameters>();
-		{
-			const bool bShouldUseBilinearSamplerForDepth = ShouldUseBilinearSamplerForDepthWithoutSingleLayerWater(SceneWithoutWaterTextures.DepthTexture->Desc.Format);
-			const bool bCustomDepthTextureProduced = HasBeenProduced(SceneTextures.CustomDepth.Depth);
-			const FIntVector DepthTextureSize = SceneWithoutWaterTextures.DepthTexture->Desc.GetSize();
-
-			SLWUniformParameters.SceneColorWithoutSingleLayerWaterTexture = SceneWithoutWaterTextures.ColorTexture;
-			SLWUniformParameters.SceneColorWithoutSingleLayerWaterSampler = TStaticSamplerState<SF_Bilinear>::GetRHI();
-			SLWUniformParameters.SceneDepthWithoutSingleLayerWaterTexture = SceneWithoutWaterTextures.DepthTexture;
-			SLWUniformParameters.SceneDepthWithoutSingleLayerWaterSampler = bShouldUseBilinearSamplerForDepth ? TStaticSamplerState<SF_Bilinear>::GetRHI() : TStaticSamplerState<SF_Point>::GetRHI();
-			SLWUniformParameters.CustomDepthTexture = bCustomDepthTextureProduced ? SceneTextures.CustomDepth.Depth : SystemTextures.DepthDummy;
-			SLWUniformParameters.CustomStencilTexture = bCustomDepthTextureProduced ? SceneTextures.CustomDepth.Stencil : SystemTextures.StencilDummySRV;
-			SLWUniformParameters.CustomDepthSampler = TStaticSamplerState<SF_Point>::GetRHI();
-			SLWUniformParameters.RefractionMaskTexture = SingleLayerWaterPrePassResult && SingleLayerWaterPrePassResult->RefractionMaskTexture ? SingleLayerWaterPrePassResult->RefractionMaskTexture : SystemTextures.White;
-			SLWUniformParameters.SceneWithoutSingleLayerWaterMinMaxUV = SceneWithoutWaterTextures.Views[ViewIndex].MinMaxUV;
-			SetupDistortionParams(SLWUniformParameters.DistortionParams, View);
-			SLWUniformParameters.SceneWithoutSingleLayerWaterTextureSize = FVector2f(DepthTextureSize.X, DepthTextureSize.Y);
-			SLWUniformParameters.SceneWithoutSingleLayerWaterInvTextureSize = FVector2f(1.0f / DepthTextureSize.X, 1.0f / DepthTextureSize.Y);
-			SLWUniformParameters.bMainDirectionalLightVSMFiltering = IsWaterVirtualShadowMapFilteringEnabled_Runtime(View.GetShaderPlatform());
-			SLWUniformParameters.bSeparateMainDirLightLuminance = NeedsSeparatedMainDirectionalLightTexture_Runtime(View.GetShaderPlatform());
-			// Only use blue noise resources if VSM quality is set to high
-			if (IsVSMTranslucentHighQualityEnabled())
-			{
-				SLWUniformParameters.BlueNoise = GetBlueNoiseParameters();
-			}
-			else
-			{
-				SLWUniformParameters.BlueNoise = GetBlueNoiseDummyParameters();
-			}
-
-			const FLightSceneProxy* SelectedForwardDirectionalLightProxy = View.ForwardLightingResources.SelectedForwardDirectionalLightProxy;
-			SetupLightCloudTransmittanceParameters(GraphBuilder, Scene, View, SelectedForwardDirectionalLightProxy ? SelectedForwardDirectionalLightProxy->GetLightSceneInfo() : nullptr, SLWUniformParameters.ForwardDirLightCloudShadow);
-		}
+		FRDGTextureRef RefractionMaskTexture = SingleLayerWaterPrePassResult && SingleLayerWaterPrePassResult->RefractionMaskTexture ? SingleLayerWaterPrePassResult->RefractionMaskTexture : SystemTextures.White;
+		FSingleLayerWaterPassUniformParameters* SLWUniformParameters = CreateSingleLayerWaterPassUniformParameters(GraphBuilder, Scene, View, SceneTextures, SceneWithoutWaterTextures.DepthTexture, SceneWithoutWaterTextures.ColorTexture, RefractionMaskTexture, SceneWithoutWaterTextures.Views[ViewIndex].MinMaxUV);
 
 		FSingleLayerWaterPassParameters* PassParameters = GraphBuilder.AllocParameters<FSingleLayerWaterPassParameters>();
 		PassParameters->View = View.GetShaderParameters();
 		PassParameters->ReflectionCapture = View.ReflectionCaptureUniformBuffer;
 		PassParameters->BasePass = CreateOpaqueBasePassUniformBuffer(GraphBuilder, View, ViewIndex);
 		PassParameters->VirtualShadowMapSamplingParameters = VirtualShadowMapArray.GetSamplingParameters(GraphBuilder, ViewIndex);
-		PassParameters->SingleLayerWater = GraphBuilder.CreateUniformBuffer(&SLWUniformParameters);
+		PassParameters->SingleLayerWater = GraphBuilder.CreateUniformBuffer(SLWUniformParameters);
 		PassParameters->RenderTargets = GetRenderTargetBindings(ERenderTargetLoadAction::ELoad, BasePassTexturesView);
 		PassParameters->RenderTargets.DepthStencil = DepthStencilBinding;
 
diff --git a/Engine/Source/Runtime/Renderer/Private/SingleLayerWaterRendering.h b/Engine/Source/Runtime/Renderer/Private/SingleLayerWaterRendering.h
--- a/Engine/Source/Runtime/Renderer/Private/SingleLayerWaterRendering.h
+++ b/Engine/Source/Runtime/Renderer/Private/SingleLayerWaterRendering.h
@@ -9,6 +9,7 @@
 #include "Froxel/Froxel.h"
 
 class FViewInfo;
+enum class ECustomDepthPassLocation : uint32;
 
 struct FSingleLayerWaterTileClassification
 {
@@ -55,7 +56,7 @@ enum class ESingleLayerWaterPrepassLocation : uint8
 bool ShouldRenderSingleLayerWater(TArrayView<const FViewInfo> Views);
 bool ShouldRenderSingleLayerWaterSkippedRenderEditorNotification(TArrayView<const FViewInfo> Views);
 bool ShouldRenderSingleLayerWaterDepthPrepass(TArrayView<const FViewInfo> Views);
-ESingleLayerWaterPrepassLocation GetSingleLayerWaterDepthPrepassLocation(bool bFullDepthPrepass);
+ESingleLayerWaterPrepassLocation GetSingleLayerWaterDepthPrepassLocation(bool bFullDepthPrepass, ECustomDepthPassLocation CustomDepthPassLocation);
 bool ShouldUseBilinearSamplerForDepthWithoutSingleLayerWater(EPixelFormat DepthTextureFormat);
 
 class FWaterTileVS : public FGlobalShader
