Draw debug will crash visionOS
-Debug draw would attempt to use a shader permutation with FUseMetalMSAAHDRDecodeDim=true and FUseMSAADimension=false, which doesn't make sense and does not match the ShouldCompilePermutation for FPopulateCompositeDepthPS and is therefore a "missing shader" resulting in an assert.  Fixed it so that if we are not using msaa we use the shader permutation without MetalMSAAHDRDecodeDim.

diff --git a/Engine/Source/Runtime/Renderer/Private/PostProcess/PostProcessCompositePrimitivesCommon.cpp b/Engine/Source/Runtime/Renderer/Private/PostProcess/PostProcessCompositePrimitivesCommon.cpp
--- a/Engine/Source/Runtime/Renderer/Private/PostProcess/PostProcessCompositePrimitivesCommon.cpp
+++ b/Engine/Source/Runtime/Renderer/Private/PostProcess/PostProcessCompositePrimitivesCommon.cpp
@@ -296,7 +296,7 @@ void PopulateDepthPass(FRDGBuilder& GraphBuilder,
 	FPopulateCompositeDepthPS::FPermutationDomain PermutationVector;
 	PermutationVector.Set<FPopulateCompositeDepthPS::FUseMSAADimension>(NumMSAASamples > 1);
 	PermutationVector.Set< FPopulateCompositeDepthPS::FForceDebugDrawColorOutput>(bForceDrawColor);
-	PermutationVector.Set< FPopulateCompositeDepthPS::FUseMetalMSAAHDRDecodeDim>(bForceDrawColor && bUseMetalPlatformHDRDecode);
+	PermutationVector.Set< FPopulateCompositeDepthPS::FUseMetalMSAAHDRDecodeDim>(NumMSAASamples > 1 && bForceDrawColor && bUseMetalPlatformHDRDecode);
 	TShaderMapRef<FPopulateCompositeDepthPS> PixelShader(View.ShaderMap, PermutationVector);
 
 	FPixelShaderUtils::AddFullscreenPass(
