Remove visual artifacting on LensFlareBlur for mali devices.

The underlying problem is that RDG barriers currently do not support VERTEX_SHADER_READ bits. So LensFlareBlur (with it's VS texture sample) is fundamentally synchronization unclean.

This alteration with moving the texture sample out of the VS and into a preceding CS should be synchronization clean.

This change is locked behind a CVar. the CVarLensFlareBlurComputeShader set to 0 will use the old path.

diff --git a/Engine/Config/Android/AndroidEngine.ini b/Engine/Config/Android/AndroidEngine.ini
--- a/Engine/Config/Android/AndroidEngine.ini
+++ b/Engine/Config/Android/AndroidEngine.ini
@@ -93,6 +93,9 @@ r.Vulkan.RHIThread=1
 r.Vulkan.AllowAsyncCompute=0
 ; Disable split barriers until its effects are measured on all devices
 r.Vulkan.AllowSplitBarriers=0
+; Split LensFlareBlur into two passes to remove vertex shader texture sample
+; Splitting into compute shader is necessary on Android. See UE-243673
+r.LensFlareBlurComputeShader=1
 
 [/Script/Engine.RendererSettings]
 r.DistanceFields=1
diff --git a/Engine/Shaders/Private/PostProcessLensFlares.usf b/Engine/Shaders/Private/PostProcessLensFlares.usf
--- a/Engine/Shaders/Private/PostProcessLensFlares.usf
+++ b/Engine/Shaders/Private/PostProcessLensFlares.usf
@@ -48,6 +48,91 @@ float2 PixelToScreenPos(float2 PixelPos)
 	return (PixelPos - Input_ScreenPosToViewportBias) / Input_ScreenPosToViewportScale;
 }
 
+#if LENS_FLARE_BLUR_CS
+
+#ifndef THREADGROUP_SIZEX
+  #define THREADGROUP_SIZEX 64
+#endif
+
+/*  RWVertexBufferBytes layout
+struct VertexBuffer
+{
+	half2 Position;
+	half4 Color;
+};
+*/
+RWBuffer<uint> RWVertexBuffer;
+[numthreads(THREADGROUP_SIZEX, 1, 1)]
+void LensFlareBlurCS(
+	uint3 GroupId : SV_GroupID,
+	uint3 DispatchThreadId : SV_DispatchThreadID,
+	uint3 GroupThreadId : SV_GroupThreadID)
+{
+	uint QuadId = DispatchThreadId.x * QUADS_PER_INSTANCE;
+	
+	float2 TilePos = float2(QuadId % TileCount.x, QuadId / TileCount.x) + Input_ViewportMin;	
+	
+	float2 TileUV = clamp(Input_ExtentInverse * TilePos * TileSize, Input_UVViewportBilinearMin, Input_UVViewportBilinearMax);
+	TilePos = PixelToScreenPos(TilePos * TileSize) * GuardBandScaleInverse;
+
+	float4 TileColor = Texture2DSampleLevel(InputTexture, InputSampler, TileUV, 0);
+	TileColor.rgb *= DiscMask(TilePos);
+	
+	uint PackedTilePosition = f32tof16(TilePos.x) | (f32tof16(TilePos.y) << 16);
+	uint PackedTileColorRG = f32tof16(TileColor.r) | (f32tof16(TileColor.g) << 16);
+	uint PackedTileColorBA = f32tof16(TileColor.b) | (f32tof16(TileColor.a) << 16);
+	
+	RWVertexBuffer[(3 * DispatchThreadId.x) + 0] = PackedTilePosition;
+	RWVertexBuffer[(3 * DispatchThreadId.x) + 1] = PackedTileColorRG;
+	RWVertexBuffer[(3 * DispatchThreadId.x) + 2] = PackedTileColorBA;
+
+}
+
+Buffer<uint> VertexBuffer;
+void LensFlareBlurComputeShaderVS(
+    in uint VId : SV_VertexID,
+	in uint IId : SV_InstanceID,
+	out noperspective float2 OutUV : TEXCOORD0,
+	out noperspective float4 OutColor : TEXCOORD1,
+	out float4 OutPosition : SV_POSITION)
+{
+	// Get packed input data
+	uint PackedTilePosition = VertexBuffer[(3 * IId) + 0];
+	uint PackedTileColorRG = VertexBuffer[(3 * IId) + 1];
+	uint PackedTileColorBA = VertexBuffer[(3 * IId) + 2];
+	
+	// Unpack input data
+	float2 TilePosition = float2(f16tof32(PackedTilePosition & 0xFFFF), f16tof32(PackedTilePosition >> 16));
+	float4 TileColor = float4(
+		f16tof32(PackedTileColorRG & 0xFFFF),
+	    f16tof32(PackedTileColorRG >> 16),
+		f16tof32(PackedTileColorBA & 0xFFFF),
+	    f16tof32(PackedTileColorBA >> 16));
+	
+	const float TileLuminance = dot(TileColor.rgb, 1);
+	const float ThresholdScale = (TileLuminance < Threshold) ? 0.0f : 1.0f;
+
+	TileColor *= KernelAreaInverse;
+	
+	// Triangle A: 0:left top, 1:right top, 2: left bottom
+	// Triangle B: 3:right bottom, 4:left bottom, 5: right top
+	float2 VertexPosWithinTile = float2(VId % 2, VId > 1 && VId < 5);
+	float2 VertexDelta = VertexPosWithinTile - float2(0.5, 0.5);
+		
+	// Set output vertex UV
+	OutUV = VertexPosWithinTile;
+	
+	// Set output vertex color as tile corner
+	OutColor = TileColor;
+	
+	// Set output vertex position
+	float2 PositionScale = 2.0 * Input_ViewportSizeInverse * KernelSize * ThresholdScale;
+	float2 VertexPosition = TilePosition + (PositionScale * VertexDelta);
+	OutPosition = float4(VertexPosition.x, VertexPosition.y, 0.0f, 1.0f);
+}
+
+#else // LENS_FLARE_BLUR_CS
+
 void LensFlareBlurVS(
 	uint VId : SV_VertexID,
 	uint IId : SV_InstanceID,
@@ -83,6 +168,8 @@ void LensFlareBlurVS(
 	OutPosition.xy += 2.0 * (LocalPos - 0.5f) * Input_ViewportSizeInverse * KernelSize * ThresholdScale;
 }
 
+#endif // LENS_FLARE_BLUR_CS
+
 void LensFlareBlurPS(
 	noperspective float2 InUV : TEXCOORD0,
 	noperspective float4 InColor : TEXCOORD1,
diff --git a/Engine/Source/Runtime/Renderer/Private/PostProcess/PostProcessLensFlares.cpp b/Engine/Source/Runtime/Renderer/Private/PostProcess/PostProcessLensFlares.cpp
--- a/Engine/Source/Runtime/Renderer/Private/PostProcess/PostProcessLensFlares.cpp
+++ b/Engine/Source/Runtime/Renderer/Private/PostProcess/PostProcessLensFlares.cpp
@@ -25,6 +25,13 @@ TAutoConsoleVariable<int32> CVarLensFlareQuality(
 	TEXT(" 3: very good quality but bad performance"),
 	ECVF_Scalability | ECVF_RenderThreadSafe);
 
+TAutoConsoleVariable<int32> CVarLensFlareBlurComputeShader(
+	TEXT("r.LensFlareBlurComputeShader"),
+	0,
+	TEXT(" 0: LensFlareBlur is computed exclusively within a single draw call (default).\n")
+	TEXT(" 1: LensFlareBlur is split into a Compute Shader invocation followed by a draw call.\n"),
+	ECVF_RenderThreadSafe);
+
 // The RDG inputs shared by all lens flare passes.
 BEGIN_SHADER_PARAMETER_STRUCT(FLensFlarePassParameters, )
 	SHADER_PARAMETER_RDG_TEXTURE_SRV(Texture2D, InputTexture)
@@ -73,6 +80,66 @@ public:
 
 IMPLEMENT_GLOBAL_SHADER(FLensFlareBlurVS, "/Engine/Private/PostProcessLensFlares.usf", "LensFlareBlurVS", SF_Vertex);
 
+class FLensFlareBlurCS : public FLensFlareShader
+{
+public:
+	DECLARE_GLOBAL_SHADER(FLensFlareBlurCS);
+	SHADER_USE_PARAMETER_STRUCT(FLensFlareBlurCS, FLensFlareShader);
+
+	BEGIN_SHADER_PARAMETER_STRUCT(FParameters, )
+		SHADER_PARAMETER_STRUCT(FScreenPassTextureViewportParameters, Input)
+		SHADER_PARAMETER_RDG_TEXTURE_SRV(Texture2D, InputTexture)
+		SHADER_PARAMETER_SAMPLER(SamplerState, InputSampler)
+		SHADER_PARAMETER(FIntPoint, TileCount)
+		SHADER_PARAMETER(uint32, TileSize)
+		SHADER_PARAMETER(float, GuardBandScaleInverse)
+		SHADER_PARAMETER_RDG_BUFFER_UAV(RWBuffer<uint>, RWVertexBuffer)
+	END_SHADER_PARAMETER_STRUCT()
+
+	static uint32 GetGroupSize()
+	{
+		return 64;
+	}
+
+	static void ModifyCompilationEnvironment(const FGlobalShaderPermutationParameters& Parameters, FShaderCompilerEnvironment& OutEnvironment)
+	{
+		FLensFlareShader::ModifyCompilationEnvironment(Parameters, OutEnvironment);
+		OutEnvironment.SetDefine(TEXT("THREADGROUP_SIZEX"), GetGroupSize());
+		OutEnvironment.SetDefine(TEXT("LENS_FLARE_BLUR_CS"), 1);
+	}
+};
+
+IMPLEMENT_GLOBAL_SHADER(FLensFlareBlurCS, "/Engine/Private/PostProcessLensFlares.usf", "LensFlareBlurCS", SF_Compute);
+
+BEGIN_SHADER_PARAMETER_STRUCT(FLensFlareBlurPassParameters, )
+	SHADER_PARAMETER_STRUCT_INCLUDE(FLensFlarePassParameters, Pass)
+	SHADER_PARAMETER_RDG_BUFFER_SRV(Buffer<uint>, VertexBuffer)
+END_SHADER_PARAMETER_STRUCT()
+
+class FLensFlareBlurComputeShaderVS : public FLensFlareShader
+{
+public:
+	DECLARE_GLOBAL_SHADER(FLensFlareBlurComputeShaderVS);
+	SHADER_USE_PARAMETER_STRUCT(FLensFlareBlurComputeShaderVS, FLensFlareShader);
+
+	BEGIN_SHADER_PARAMETER_STRUCT(FParameters, )
+		SHADER_PARAMETER_STRUCT_INCLUDE(FLensFlareBlurPassParameters, Pass)
+		SHADER_PARAMETER_STRUCT(FScreenPassTextureViewportParameters, Input)
+		SHADER_PARAMETER(float, KernelSize)
+		SHADER_PARAMETER(float, KernelAreaInverse)
+		SHADER_PARAMETER(float, Threshold)
+	END_SHADER_PARAMETER_STRUCT()
+
+
+	static void ModifyCompilationEnvironment(const FGlobalShaderPermutationParameters& Parameters, FShaderCompilerEnvironment& OutEnvironment)
+	{
+		FLensFlareShader::ModifyCompilationEnvironment(Parameters, OutEnvironment);
+		OutEnvironment.SetDefine(TEXT("LENS_FLARE_BLUR_CS"), 1);
+	}
+};
+
+IMPLEMENT_GLOBAL_SHADER(FLensFlareBlurComputeShaderVS, "/Engine/Private/PostProcessLensFlares.usf", "LensFlareBlurComputeShaderVS", SF_Vertex);
+
 class FLensFlareBlurPS : public FLensFlareShader
 {
 public:
@@ -175,6 +242,8 @@ FScreenPassTexture AddLensFlaresPass(
 	}
 
 	// Lens flare blur pass. Rasterizes a bokeh quad for each pixel on the screen based on the intensity threshold.
+	// Decide between one draw call versus one compute and one draw
+	if (CVarLensFlareBlurComputeShader.GetValueOnRenderThread() == 0)
 	{
 		const uint32 TileSizeInPixels = 1;
 
@@ -210,36 +279,133 @@ FScreenPassTexture AddLensFlaresPass(
 			RDG_EVENT_NAME("LensFlareBlur %dx%d", FlareViewSize.X, FlareViewSize.Y),
 			PassParameters,
 			ERDGPassFlags::Raster,
-			[VertexShader, PixelShader, VertexParameters, PixelParameters, AdditiveBlendState, FlareViewport, TileCount] (FRDGAsyncTask, FRHICommandList& RHICmdList)
+			[VertexShader, PixelShader, VertexParameters, PixelParameters, AdditiveBlendState, FlareViewport, TileCount](FRDGAsyncTask, FRHICommandList& RHICmdList)
+			{
+				// Viewport is the same as the input.
+				RHICmdList.SetViewport(
+					FlareViewport.Rect.Min.X,
+					FlareViewport.Rect.Min.Y,
+					0.0f,
+					FlareViewport.Rect.Max.X,
+					FlareViewport.Rect.Max.Y,
+					1.0f);
+
+				// Apply additive blending pipeline state.
+				FGraphicsPipelineStateInitializer GraphicsPSOInit;
+				RHICmdList.ApplyCachedRenderTargets(GraphicsPSOInit);
+				GraphicsPSOInit.BlendState = AdditiveBlendState;
+				GraphicsPSOInit.RasterizerState = TStaticRasterizerState<>::GetRHI();
+				GraphicsPSOInit.DepthStencilState = TStaticDepthStencilState<false, CF_Always>::GetRHI();
+				GraphicsPSOInit.BoundShaderState.VertexDeclarationRHI = GEmptyVertexDeclaration.VertexDeclarationRHI;
+				GraphicsPSOInit.BoundShaderState.VertexShaderRHI = VertexShader.GetVertexShader();
+				GraphicsPSOInit.BoundShaderState.PixelShaderRHI = PixelShader.GetPixelShader();
+				GraphicsPSOInit.PrimitiveType = PT_TriangleList;
+				SetGraphicsPipelineState(RHICmdList, GraphicsPSOInit, 0);
+				SetShaderParameters(RHICmdList, VertexShader, VertexShader.GetVertexShader(), VertexParameters);
+				SetShaderParameters(RHICmdList, PixelShader, PixelShader.GetPixelShader(), PixelParameters);
+				// Emit an instanced quad draw call on the order of the number of pixels on the screen.
+				RHICmdList.SetStreamSource(0, nullptr, 0);
+				RHICmdList.DrawPrimitive(0, 2, FMath::DivideAndRoundUp(TileCount.X * TileCount.Y, GLensFlareQuadsPerInstance));
+			});
+	}
+	else
+	{
+		const uint32 TileSizeInPixels = 1;
+
+		const FIntPoint TileCount = FlareViewSize / TileSizeInPixels;
+
+		const int32 NumQuads = TileCount.X * TileCount.Y;
+		const uint32 NumInstances = FMath::DivideAndRoundUp(NumQuads, GLensFlareQuadsPerInstance);
+		const uint32 UintPerInstance = 3;
+		const uint32 TotalUints = NumInstances * UintPerInstance;
+
+		FRDGBufferRef LensFlareVertexBuffer = GraphBuilder.CreateBuffer(
+			FRDGBufferDesc::CreateBufferDesc(sizeof(uint32), TotalUints),
+			TEXT("LensFlareVertexBuffer"));
+		FRDGBufferUAVRef LensFlareVertexBufferUAV = GraphBuilder.CreateUAV(LensFlareVertexBuffer, PF_R32_UINT);
+		FRDGBufferSRVRef LensFlareVertexBufferSRV = GraphBuilder.CreateSRV(LensFlareVertexBuffer, PF_R32_UINT);
+
+		// Setup compute shader parameters
 		{
-			// Viewport is the same as the input.
-			RHICmdList.SetViewport(
-				FlareViewport.Rect.Min.X,
-				FlareViewport.Rect.Min.Y,
-				0.0f,
-				FlareViewport.Rect.Max.X,
-				FlareViewport.Rect.Max.Y,
-				1.0f);
-
-			// Apply additive blending pipeline state.
-			FGraphicsPipelineStateInitializer GraphicsPSOInit;
-			RHICmdList.ApplyCachedRenderTargets(GraphicsPSOInit);
-			GraphicsPSOInit.BlendState = AdditiveBlendState;
-			GraphicsPSOInit.RasterizerState = TStaticRasterizerState<>::GetRHI();
-			GraphicsPSOInit.DepthStencilState = TStaticDepthStencilState<false, CF_Always>::GetRHI();
-			GraphicsPSOInit.BoundShaderState.VertexDeclarationRHI = GEmptyVertexDeclaration.VertexDeclarationRHI;
-			GraphicsPSOInit.BoundShaderState.VertexShaderRHI = VertexShader.GetVertexShader();
-			GraphicsPSOInit.BoundShaderState.PixelShaderRHI = PixelShader.GetPixelShader();
-			GraphicsPSOInit.PrimitiveType = PT_TriangleList;
-			SetGraphicsPipelineState(RHICmdList, GraphicsPSOInit, 0);
-
-			SetShaderParameters(RHICmdList, VertexShader, VertexShader.GetVertexShader(), VertexParameters);
-			SetShaderParameters(RHICmdList, PixelShader, PixelShader.GetPixelShader(), PixelParameters);
-
-			// Emit an instanced quad draw call on the order of the number of pixels on the screen.
-			RHICmdList.SetStreamSource(0, nullptr, 0);
-			RHICmdList.DrawPrimitive(0, 2, FMath::DivideAndRoundUp(TileCount.X * TileCount.Y, GLensFlareQuadsPerInstance));
-		});
+			FLensFlareBlurCS::FParameters* ComputeParameters = GraphBuilder.AllocParameters<FLensFlareBlurCS::FParameters>();
+			ComputeParameters->Input = GetScreenPassTextureViewportParameters(FlareViewport);
+			ComputeParameters->InputSampler = BilinearClampSampler;
+			ComputeParameters->InputTexture = Inputs.Flare.TextureSRV;
+			ComputeParameters->TileCount = TileCount;
+			ComputeParameters->TileSize = TileSizeInPixels;
+			ComputeParameters->GuardBandScaleInverse = 1.0f / GuardBandScale;
+			ComputeParameters->RWVertexBuffer = LensFlareVertexBufferUAV;
+			TShaderMapRef<FLensFlareBlurCS> ComputeShader(View.ShaderMap);
+			//ComputeShader->GroupSize = NumInstances;
+
+			FComputeShaderUtils::AddPass(
+				GraphBuilder,
+				RDG_EVENT_NAME("LensFlareBlurCS %dx%d", FlareViewSize.X, FlareViewSize.Y),
+				ERDGPassFlags::Compute,
+				ComputeShader,
+				ComputeParameters,
+				FIntVector(NumInstances, 1, 1));
+		}
+
+		{
+			const float KernelSize = (Inputs.BokehSizePercent * static_cast<float>(FlareViewSize.X)) * PercentToScale;
+
+			FLensFlareBlurPassParameters* PassParameters = GraphBuilder.AllocParameters<FLensFlareBlurPassParameters>();
+			//PassParameters->Pass.InputTexture = Inputs.Flare.TextureSRV;
+			PassParameters->Pass.RenderTargets[0] = FRenderTargetBinding(BlurOutputTexture, ERenderTargetLoadAction::EClear);
+			PassParameters->VertexBuffer = LensFlareVertexBufferSRV;
+
+			// Setup vertex shader parameters.
+			FLensFlareBlurComputeShaderVS::FParameters VertexParameters;
+			VertexParameters.Pass = *PassParameters;
+			VertexParameters.Input = GetScreenPassTextureViewportParameters(FlareViewport);
+			VertexParameters.KernelSize = KernelSize;
+			VertexParameters.Threshold = Inputs.Threshold;
+			VertexParameters.KernelAreaInverse = 1.0f / FMath::Max(1.0f, KernelSize * KernelSize);
+
+			// Setup pixel shader parameters.
+			FLensFlareBlurPS::FParameters PixelParameters;
+			PixelParameters.BokehTexture = Inputs.BokehShapeTexture;
+			PixelParameters.BokehSampler = BilinearClampSampler;
+
+			TShaderMapRef<FLensFlareBlurComputeShaderVS> VertexShader(View.ShaderMap);
+			TShaderMapRef<FLensFlareBlurPS> PixelShader(View.ShaderMap);
+
+			GraphBuilder.AddPass(
+				RDG_EVENT_NAME("LensFlareBlur %dx%d", FlareViewSize.X, FlareViewSize.Y),
+				PassParameters,
+				ERDGPassFlags::Raster,
+				[VertexShader, PixelShader, VertexParameters, PixelParameters, AdditiveBlendState, FlareViewport, TileCount](FRDGAsyncTask, FRHICommandList& RHICmdList)
+				{
+					// Viewport is the same as the input.
+					RHICmdList.SetViewport(
+						FlareViewport.Rect.Min.X,
+						FlareViewport.Rect.Min.Y,
+						0.0f,
+						FlareViewport.Rect.Max.X,
+						FlareViewport.Rect.Max.Y,
+						1.0f);
+
+					// Apply additive blending pipeline state.
+					FGraphicsPipelineStateInitializer GraphicsPSOInit;
+					RHICmdList.ApplyCachedRenderTargets(GraphicsPSOInit);
+					GraphicsPSOInit.BlendState = AdditiveBlendState;
+					GraphicsPSOInit.RasterizerState = TStaticRasterizerState<>::GetRHI();
+					GraphicsPSOInit.DepthStencilState = TStaticDepthStencilState<false, CF_Always>::GetRHI();
+					GraphicsPSOInit.BoundShaderState.VertexDeclarationRHI = GEmptyVertexDeclaration.VertexDeclarationRHI;
+					GraphicsPSOInit.BoundShaderState.VertexShaderRHI = VertexShader.GetVertexShader();
+					GraphicsPSOInit.BoundShaderState.PixelShaderRHI = PixelShader.GetPixelShader();
+					GraphicsPSOInit.PrimitiveType = PT_TriangleList;
+					SetGraphicsPipelineState(RHICmdList, GraphicsPSOInit, 0);
+
+					SetShaderParameters(RHICmdList, VertexShader, VertexShader.GetVertexShader(), VertexParameters);
+					SetShaderParameters(RHICmdList, PixelShader, PixelShader.GetPixelShader(), PixelParameters);
+
+					// Emit an instanced quad draw call on the order of the number of pixels on the screen.
+					RHICmdList.SetStreamSource(0, nullptr, 0);
+					RHICmdList.DrawPrimitive(0, 2, FMath::DivideAndRoundUp(TileCount.X * TileCount.Y, GLensFlareQuadsPerInstance));
+				});
+		}
 	}
 
 	FRDGTextureRef LensFlareTexture = nullptr;
