Remove ensure from FScene::UpdateCachedRayTracingState(...)
- In some situations FPrimitiveSceneInfo can end up being used before it's registered with FScene.

diff --git a/Engine/Source/Runtime/Renderer/Private/RendererScene.cpp b/Engine/Source/Runtime/Renderer/Private/RendererScene.cpp
--- a/Engine/Source/Runtime/Renderer/Private/RendererScene.cpp
+++ b/Engine/Source/Runtime/Renderer/Private/RendererScene.cpp
@@ -4093,7 +4093,7 @@ void FScene::UpdateCachedRayTracingState(FPrimitiveSceneProxy* SceneProxy)
 	if (SceneProxy->GetPrimitiveSceneInfo())
 	{
 		const int32 PrimitiveIndex = SceneProxy->GetPrimitiveSceneInfo()->GetIndex();
-		if (ensure(PrimitiveIndex != INDEX_NONE))
+		if (PrimitiveIndex != INDEX_NONE)
 		{
 			PrimitiveRayTracingDatas[PrimitiveIndex].bCachedRaytracingDataDirty = true;
 		}
