Substrate - improvemnt to translucentrough refraction. Aliasing is reduced considerably now.
Small change of blur behavior for this experimental feature.

#rnx
#tests Editor dx12

diff --git a/Engine/Shaders/Private/DistortFiltering.usf b/Engine/Shaders/Private/DistortFiltering.usf
--- a/Engine/Shaders/Private/DistortFiltering.usf
+++ b/Engine/Shaders/Private/DistortFiltering.usf
@@ -45,7 +45,6 @@ void CopySceneColorDepthTexturePS(FScreenVertexOutput Input, out float4 OutColor
 
 #ifdef DOWNSAMPLECOLORCS
 
-uint SrcMipIndex;
 int2 SrcMipResolution;
 int2 DstMipResolution;
 SamplerState SourceSampler;
@@ -59,9 +58,9 @@ void DownsampleColorCS(uint3 ThreadId : SV_DispatchThreadID)
 	{
 		return;
 	}
-	float2 ScaledUVs = (float2(ThreadId.xy) + 0.5f) / float2(DstMipResolution);
 
-	float4 OutColor = SourceTexture.SampleLevel(SourceSampler, ScaledUVs, SrcMipIndex);
+	float2 ScaledUVs = (float2(ThreadId.xy) * 2.0 + 1.0f) / float2(SrcMipResolution);
+	float4 OutColor = SourceTexture.Sample(SourceSampler, ScaledUVs, 0); // No need to specify mip index, the SRV is created for the acual mip level.
 
 	OutTextureMipColor[uint2(ThreadId.xy)] = OutColor;
 }
@@ -73,8 +72,8 @@ void DownsampleColorCS(uint3 ThreadId : SV_DispatchThreadID)
 #ifdef FILTERCOLORCS
 
 uint SampleCount;
-uint SrcMipIndex;
-int2 MipResolution;
+int2 SrcMipResolution;
+int2 DstMipResolution;
 float4 BlurDirection;
 SamplerState SourceSampler;
 Texture2D<float4> SourceTexture;
@@ -87,11 +86,11 @@ float4 SampleOffsetsWeights[MAX_FILTER_SAMPLE_COUNT];
 [numthreads(THREADGROUP_SIZE, THREADGROUP_SIZE, 1)]
 void FilterColorCS(uint3 ThreadId : SV_DispatchThreadID)
 {
-	if (any(ThreadId.xy >= uint2(MipResolution)))
+	if (any(ThreadId.xy >= uint2(DstMipResolution)))
 	{
 		return;
 	}
-	float2 ScaledUVs = (float2(ThreadId.xy) + 0.5f) / float2(MipResolution);
+	float2 ScaledUVs = (float2(ThreadId.xy) * 2.0 + 1.0f) / float2(SrcMipResolution);
 
 	float4 OutColor = 0;
 	
@@ -99,7 +98,7 @@ void FilterColorCS(uint3 ThreadId : SV_DispatchThreadID)
 	{
 		float2 SampleOffsetWeight = SampleOffsetsWeights[i].xy;
 
-		float4 Sample = SourceTexture.SampleLevel(SourceSampler, ScaledUVs + SampleOffsetWeight.xx * BlurDirection.xy, SrcMipIndex);
+		float4 Sample = SourceTexture.SampleLevel(SourceSampler, ScaledUVs + SampleOffsetWeight.xx * BlurDirection.xy, 0); // No need to specify mip index, the SRV is created for the acual mip level.
 
 		OutColor += Sample * SampleOffsetWeight.y;
 	}
diff --git a/Engine/Source/Runtime/Renderer/Private/DistortionRendering.cpp b/Engine/Source/Runtime/Renderer/Private/DistortionRendering.cpp
--- a/Engine/Source/Runtime/Renderer/Private/DistortionRendering.cpp
+++ b/Engine/Source/Runtime/Renderer/Private/DistortionRendering.cpp
@@ -327,8 +327,8 @@ static void AddCopySceneColorDepthPass(FRDGBuilder& GraphBuilder, const FViewInf
 	PassParameters->View = View.ViewUniformBuffer;
 	PassParameters->SceneColorTexture = SceneColorTexture;
 	PassParameters->SceneDepthTexture = SceneDepthTexture;
-	PassParameters->SceneColorSampler = TStaticSamplerState<SF_Point>::GetRHI();
-	PassParameters->SceneDepthSampler = TStaticSamplerState<SF_Point>::GetRHI();
+	PassParameters->SceneColorSampler = TStaticSamplerState<SF_Bilinear>::GetRHI();
+	PassParameters->SceneDepthSampler = TStaticSamplerState<SF_Bilinear>::GetRHI();
 	PassParameters->RenderTargets[0] = FRenderTargetBinding(SceneColorCopyTexture, ERenderTargetLoadAction::ENoAction);
 	PassParameters->RenderTargets[1] = FRenderTargetBinding(SceneDepthCopyTexture, ERenderTargetLoadAction::ENoAction);
 	PassParameters->MaxExposedLuminance = FMath::Max(0.0f, CVarRefractionBlurMaxExposedLuminance.GetValueOnRenderThread());
@@ -356,7 +356,6 @@ class FDownsampleSceneColorCS : public FGlobalShader
 	using FPermutationDomain = TShaderPermutationDomain<>;
 
 	BEGIN_SHADER_PARAMETER_STRUCT(FParameters, )
-		SHADER_PARAMETER(uint32, SrcMipIndex)
 		SHADER_PARAMETER(FIntPoint, SrcMipResolution)
 		SHADER_PARAMETER(FIntPoint, DstMipResolution)
 		SHADER_PARAMETER_RDG_TEXTURE_SRV(Texture2D, SourceTexture)
@@ -458,8 +457,8 @@ class FFilterSceneColorCS : public FGlobalShader
 
 	BEGIN_SHADER_PARAMETER_STRUCT(FParameters, )
 		SHADER_PARAMETER(uint32, SampleCount)
-		SHADER_PARAMETER(uint32, SrcMipIndex)
-		SHADER_PARAMETER(FIntPoint, MipResolution)
+		SHADER_PARAMETER(FIntPoint, SrcMipResolution)
+		SHADER_PARAMETER(FIntPoint, DstMipResolution)
 		SHADER_PARAMETER(FVector4f, BlurDirection)
 		SHADER_PARAMETER_RDG_TEXTURE_SRV(Texture2D, SourceTexture)
 		SHADER_PARAMETER_SAMPLER(SamplerState, SourceSampler)
@@ -481,7 +480,7 @@ IMPLEMENT_GLOBAL_SHADER(FFilterSceneColorCS, "/Engine/Private/DistortFiltering.u
 static void AddFilterSceneColorPass(FRDGBuilder& GraphBuilder, FFilterSceneColorCS::FParameters* PassParameters, const FViewInfo& View)
 {
 	FIntVector NumGroups = FIntVector::DivideAndRoundUp(
-		FIntVector(PassParameters->MipResolution.X, PassParameters->MipResolution.Y, 1),
+		FIntVector(PassParameters->DstMipResolution.X, PassParameters->DstMipResolution.Y, 1),
 		FIntVector(FFilterSceneColorCS::ThreadGroupSize, FFilterSceneColorCS::ThreadGroupSize, 1));
 
 	FFilterSceneColorCS::FPermutationDomain PermutationVector;
@@ -608,7 +607,6 @@ void FDeferredShadingSceneRenderer::RenderDistortion(
 				for (int32 DstMipIndex = 1; DstMipIndex < MipCount; DstMipIndex++)
 				{
 					FDownsampleSceneColorCS::FParameters* PassParameters = GraphBuilder.AllocParameters<FDownsampleSceneColorCS::FParameters>();
-					PassParameters->SrcMipIndex = DstMipIndex - 1;
 					PassParameters->SrcMipResolution = SrcMipResolution;
 					PassParameters->DstMipResolution = DstMipResolution;
 					PassParameters->SourceSampler = TStaticSamplerState<SF_Bilinear>::GetRHI();
@@ -633,14 +631,15 @@ void FDeferredShadingSceneRenderer::RenderDistortion(
 			{
 				RDG_EVENT_SCOPE(GraphBuilder, "SceneColorDepthMipHBlur");
 
-				FIntPoint MipResolution = SceneColorMip0Resolution / 2;
+				FIntPoint SrcMipResolution = SceneColorMip0Resolution;
+				FIntPoint DstMipResolution = SceneColorMip0Resolution / 2;
 				for (int32 MipIndex = 1; MipIndex < MipCount; MipIndex++)
 				{
-					const FVector2D InverseFilterTextureExtent(1.0f / static_cast<float>(MipResolution.X), 1.0f / static_cast<float>(MipResolution.Y));
+					const FVector2D InverseFilterTextureExtent(1.0f / static_cast<float>(DstMipResolution.X), 1.0f / static_cast<float>(DstMipResolution.Y));
 
 					FFilterSceneColorCS::FParameters* PassParameters = GraphBuilder.AllocParameters<FFilterSceneColorCS::FParameters>();
-					PassParameters->SrcMipIndex = MipIndex;
-					PassParameters->MipResolution = MipResolution;
+					PassParameters->SrcMipResolution = SrcMipResolution;
+					PassParameters->DstMipResolution = DstMipResolution;
 					PassParameters->SourceSampler = TStaticSamplerState<SF_Bilinear>::GetRHI();
 					PassParameters->SampleCount = SampleCount;
 					PassParameters->BlurDirection = FVector4f(1.0f, 0.0f, 0.0f, 0.0f);
@@ -660,7 +659,8 @@ void FDeferredShadingSceneRenderer::RenderDistortion(
 					PassParametersDepth->OutTextureMipColor = GraphBuilder.CreateUAV(FRDGTextureUAVDesc(TempSceneDepthMeterMipchainTexture, MipIndex));
 					AddFilterSceneColorPass(GraphBuilder, PassParametersDepth, View);
 
-					MipResolution = MipResolution / 2;
+					SrcMipResolution = DstMipResolution;
+					DstMipResolution = DstMipResolution / 2;
 				}
 			}
 
@@ -669,14 +669,15 @@ void FDeferredShadingSceneRenderer::RenderDistortion(
 			{
 				RDG_EVENT_SCOPE(GraphBuilder, "SceneColorMipVBlur");
 
-				FIntPoint MipResolution = SceneColorMip0Resolution / 2;
+				FIntPoint SrcMipResolution = SceneColorMip0Resolution;
+				FIntPoint DstMipResolution = SceneColorMip0Resolution / 2;
 				for (int32 MipIndex = 1; MipIndex < MipCount; MipIndex++)
 				{
-					const FVector2D InverseFilterTextureExtent(1.0f / static_cast<float>(MipResolution.X), 1.0f / static_cast<float>(MipResolution.Y));
+					const FVector2D InverseFilterTextureExtent(1.0f / static_cast<float>(DstMipResolution.X), 1.0f / static_cast<float>(DstMipResolution.Y));
 
 					FFilterSceneColorCS::FParameters* PassParameters = GraphBuilder.AllocParameters<FFilterSceneColorCS::FParameters>();
-					PassParameters->SrcMipIndex = MipIndex;
-					PassParameters->MipResolution = MipResolution;
+					PassParameters->SrcMipResolution = SrcMipResolution;
+					PassParameters->DstMipResolution = DstMipResolution;
 					PassParameters->SourceSampler = TStaticSamplerState<SF_Bilinear>::GetRHI();
 					PassParameters->SampleCount = SampleCount;
 					PassParameters->BlurDirection = FVector4f(0.0f, 1.0f, 0.0f, 0.0f);
@@ -696,7 +697,8 @@ void FDeferredShadingSceneRenderer::RenderDistortion(
 					PassParametersDepth->OutTextureMipColor = GraphBuilder.CreateUAV(FRDGTextureUAVDesc(SceneDepthMeterMipchainTexture, MipIndex));
 					AddFilterSceneColorPass(GraphBuilder, PassParametersDepth, View);
 
-					MipResolution = MipResolution / 2;
+					SrcMipResolution = DstMipResolution;
+					DstMipResolution = DstMipResolution / 2;
 				}
 			}
 		}
