Substrate - rough refraction - added blur uv dithering to better hide low resolution texels.
Work ok but too large value can result in color splotches.

#tests PC editor dx12
#rnx

diff --git a/Engine/Shaders/Private/DistortApplyScreenPS.usf b/Engine/Shaders/Private/DistortApplyScreenPS.usf
--- a/Engine/Shaders/Private/DistortApplyScreenPS.usf
+++ b/Engine/Shaders/Private/DistortApplyScreenPS.usf
@@ -10,6 +10,7 @@
 
 #if USE_ROUGH_REFRACTION && SUBSTRATE_ENABLED
 #include "/Engine/Private/Substrate/SubstrateStatisticalOperators.ush"
+#include "BlueNoise.ush"
 #endif
 
 
@@ -55,6 +56,7 @@ SamplerState		StandardTranslucentTransmittanceTextureSampler;
 
 float RefractionMaxMipLevel;
 float StandardDevInPixelForMip0;
+float BlurDitheringScaleFactor;
 
 // TexCoords [0,1] where 0 is the left edge and 1 is the right edge.
 // This means that in ISR they go over the whole doublewide.
@@ -137,15 +139,11 @@ void Main(
 
 		const float PixelStandardDevInPixel = StandardDeviationCmForThickness / PixelRadiusCm;
 
-		// Jitter UV according to MIP level to hide texel
-		// => disabled because it is a bit unpleasant over bright area (TSR seems to have trouble resolving it)
-		if(false)
+		// Jitter UV according to MIP level to hide low res texels. (also tried mip index dithering: good but less effective)
+		// It does add a bit of extra blurring, but practically acceptable and it is not that much more.
+		if(BlurDitheringScaleFactor > 0.0f) // make this an option
 		{
-			const float BlurRandomness = InterleavedGradientNoise(SVPos.xy, float(View.StateFrameIndexMod8)) - 0.5; // In [-0.5, 0.5]
-			const float RandAngle = BlurRandomness * 2.0f * PI;
-			float CosA, SinA;
-			sincos(RandAngle, SinA, CosA);
-			NewBufferUV += 0.5 * float2(CosA, SinA) * PixelStandardDevInPixel * View.ViewSizeAndInvSize.zw;
+			NewBufferUV += (BlueNoiseVec2(SVPos.xy, View.StateFrameIndexMod8) - 0.5f) * BlurDitheringScaleFactor * PixelStandardDevInPixel * View.ViewSizeAndInvSize.zw;
 		}
 
 		// The standard deviation in pixel for mip0 = StandardDevInPixelForMip0.
diff --git a/Engine/Source/Runtime/Renderer/Private/DistortionRendering.cpp b/Engine/Source/Runtime/Renderer/Private/DistortionRendering.cpp
--- a/Engine/Source/Runtime/Renderer/Private/DistortionRendering.cpp
+++ b/Engine/Source/Runtime/Renderer/Private/DistortionRendering.cpp
@@ -19,6 +19,7 @@
 #include "PostProcess/TemporalAA.h"
 #include "ShaderPlatformCachedIniValue.h"
 #include "HeterogeneousVolumes/HeterogeneousVolumes.h"
+#include "BlueNoise.h"
 
 DECLARE_GPU_DRAWCALL_STAT(Distortion);
 
@@ -31,7 +32,7 @@ static TAutoConsoleVariable<int32> CVarDisableDistortion(
 static TAutoConsoleVariable<int32> CVarRefractionBlur(
 	TEXT("r.Refraction.Blur"),
 	1,
-	TEXT("Enables rough refractions, i.e. blurring of the background."),
+	TEXT("Enables rough refractions, i.e. blurring of the background. Only enabled when Substrate is enabled."),
 	ECVF_Scalability | ECVF_RenderThreadSafe);
 
 static TAutoConsoleVariable<int32> CVarRefractionOffsetQuality(
@@ -58,6 +59,12 @@ static TAutoConsoleVariable<float> CVarRefractionBlurMaxExposedLuminance(
 	TEXT("Clamp scene pre-exposed luminance to this maximum value. It helps to reduce bright specular highlight flickering, even when r.Refraction.Blur.TemporalAA=1."),
 	ECVF_Scalability | ECVF_RenderThreadSafe);
 
+static TAutoConsoleVariable<float> CVarRefractionBlurDitheringScaleFactor(
+	TEXT("r.Refraction.Blur.DitheringScaleFactor"),
+	0.3f,
+	TEXT("Represents the radius scale factor of the dithered noise applied on the rough refraction UV sample. A value of 1 means that the UV offset will be rotated on circle of radius = pixel standard deviation. A large value can confuse TAA/TSR and result on color splotches. Disabled when 0, accept [0,1] range."),
+	ECVF_Scalability | ECVF_RenderThreadSafe);
+
 static TAutoConsoleVariable<int32> CVarRefractionUseAVSM(
 	TEXT("r.HeterogeneousVolumes.CompositeWithTranslucency.Refraction.UseAVSM"),
 	1,
@@ -135,9 +142,9 @@ TRDGUniformBufferRef<FDistortionPassUniformParameters> CreateDistortionPassUnifo
 	return GraphBuilder.CreateUniformBuffer(Parameters);
 }
 
-static bool GetUseRoughRefraction()
+bool GetSubstrateEnabledUseRoughRefraction()
 {
-	return Substrate::IsSubstrateEnabled() && CVarRefractionBlur.GetValueOnAnyThread() > 0; // Any thread since it can be called when creating the PassProcessor
+	return Substrate::IsSubstrateEnabled() && !CVarDisableDistortion.GetValueOnAnyThread() && CVarRefractionBlur.GetValueOnAnyThread() > 0; // Any thread since it can be called when creating the PassProcessor
 }
 
 class FDistortionScreenPS : public FGlobalShader
@@ -150,6 +157,7 @@ public:
 
 	BEGIN_SHADER_PARAMETER_STRUCT(FParameters, )
 		SHADER_PARAMETER_STRUCT_INCLUDE(FViewShaderParameters, View)
+		SHADER_PARAMETER_STRUCT_REF(FBlueNoise, BlueNoise)
 		SHADER_PARAMETER_RDG_TEXTURE(Texture2DMS<float4>, RefractionVarianceCoverageMSAATexture)
 		SHADER_PARAMETER_RDG_TEXTURE(Texture2DMS<float4>, RefractionClosestDepthMeterMSAATexture)
 		SHADER_PARAMETER_RDG_TEXTURE(Texture2DMS<float4>, DistortionMSAATexture)
@@ -173,6 +181,7 @@ public:
 		SHADER_PARAMETER_SAMPLER(SamplerState, StandardTranslucentTransmittanceTextureSampler)
 		SHADER_PARAMETER(float, RefractionMaxMipLevel)
 		SHADER_PARAMETER(float, StandardDevInPixelForMip0)
+		SHADER_PARAMETER(float, BlurDitheringScaleFactor)
 		RENDER_TARGET_BINDING_SLOTS()
 	END_SHADER_PARAMETER_STRUCT()
 
@@ -532,7 +541,7 @@ void FDeferredShadingSceneRenderer::RenderDistortion(
 	TArray<FRDGTextureRef> ViewsSceneDepthMeterMipchain;
 	ViewsSceneColorMipchain.SetNum(Views.Num());
 	ViewsSceneDepthMeterMipchain.SetNum(Views.Num());
-	const bool bUseRoughRefraction = GetUseRoughRefraction();
+	const bool bUseRoughRefraction = GetSubstrateEnabledUseRoughRefraction();
 	if(bUseRoughRefraction)
 	{
 		FVector2D OffsetAndWeight[MAX_FILTER_SAMPLE_COUNT];
@@ -830,7 +839,10 @@ void FDeferredShadingSceneRenderer::RenderDistortion(
 	CommonParameters.StandardTranslucentColorAlphaTextureSampler = TStaticSamplerState<>::GetRHI();
 	CommonParameters.StandardTranslucentTransmittanceTextureSampler = TStaticSamplerState<>::GetRHI();
 	CommonParameters.StandardDevInPixelForMip0 = FilterSigmaPixel;
+	CommonParameters.BlurDitheringScaleFactor = FMath::Clamp(CVarRefractionBlurDitheringScaleFactor.GetValueOnRenderThread(), 0.0f, 1.0f);
 
+	FBlueNoise BlueNoise = GetBlueNoiseGlobalParameters();
+	CommonParameters.BlueNoise = CreateUniformBufferImmediate(BlueNoise, EUniformBufferUsage::UniformBuffer_SingleDraw);
 
 	const bool bAllowStandardTranslucencySeparated = ViewFamily.AllowStandardTranslucencySeparated();
 	FDistortionScreenPS::FPermutationDomain PermutationVector;
@@ -1033,7 +1045,7 @@ void FDistortionMeshProcessor::CollectPSOInitializers(const FSceneTexturesConfig
 	FGraphicsPipelineRenderTargetsInfo RenderTargetsInfo;
 	RenderTargetsInfo.NumSamples = 1;
 	AddRenderTargetInfo(GetRefractionOffsetQuality(SceneTexturesConfig.ShaderPlatform) > 0 ? PF_FloatRGBA : PF_B8G8R8A8, GFastVRamConfig.Distortion | TexCreate_RenderTargetable | TexCreate_ShaderResource, RenderTargetsInfo);
-	if (GetUseRoughRefraction())
+	if (GetSubstrateEnabledUseRoughRefraction())
 	{
 		AddRenderTargetInfo(PF_R16F, GFastVRamConfig.Distortion | TexCreate_RenderTargetable | TexCreate_ShaderResource, RenderTargetsInfo);
 	}
@@ -1152,7 +1164,7 @@ FMeshPassProcessor* CreateDistortionPassProcessor(ERHIFeatureLevel::Type Feature
 
 	DistortionPassState.SetStencilRef(DISTORTION_STENCIL_MASK_BIT);
 
-	if (GetUseRoughRefraction())
+	if (GetSubstrateEnabledUseRoughRefraction())
 	{
 		DistortionPassState.SetBlendState(TStaticBlendState<CW_RGBA, BO_Add, BF_One, BF_One, BO_Add, BF_One, BF_One,
 															CW_RGBA, BO_Add, BF_One, BF_One, BO_Add, BF_One, BF_One,
@@ -1182,7 +1194,7 @@ FMeshPassProcessor* CreateMobileDistortionPassProcessor(ERHIFeatureLevel::Type F
 	// We don't have depth, render all pixels, pixel shader will sample SceneDepth from SceneColor.A and discard if occluded
 	DistortionPassState.SetDepthStencilState(TStaticDepthStencilState<false, CF_Always>::GetRHI());
 
-	if (GetUseRoughRefraction())
+	if (GetSubstrateEnabledUseRoughRefraction())
 	{
 		DistortionPassState.SetBlendState(TStaticBlendState<CW_RGBA, BO_Add, BF_One, BF_One, BO_Add, BF_One, BF_One,
 															CW_RGBA, BO_Add, BF_One, BF_One, BO_Add, BF_One, BF_One,
diff --git a/Engine/Source/Runtime/Renderer/Private/SceneRendering.cpp b/Engine/Source/Runtime/Renderer/Private/SceneRendering.cpp
--- a/Engine/Source/Runtime/Renderer/Private/SceneRendering.cpp
+++ b/Engine/Source/Runtime/Renderer/Private/SceneRendering.cpp
@@ -492,6 +492,7 @@ void CleanupOcclusionSubmittedFence()
 }
 
 extern int32 GetTranslucencyLightingVolumeDim();
+extern bool GetSubstrateEnabledUseRoughRefraction();
 
 DECLARE_CYCLE_STAT(TEXT("DeferredShadingSceneRenderer ViewExtensionPostRenderView"), STAT_FDeferredShadingSceneRenderer_ViewExtensionPostRenderView, STATGROUP_SceneRendering);
 DECLARE_CYCLE_STAT(TEXT("DeferredShadingSceneRenderer ViewExtensionPreRenderView"), STAT_FDeferredShadingSceneRenderer_ViewExtensionPreRenderView, STATGROUP_SceneRendering);
@@ -2727,7 +2728,8 @@ FSceneRenderer::FSceneRenderer(const FSceneViewFamily* InViewFamily, FHitProxyCo
 			UseVirtualShadowMaps(Scene->GetShaderPlatform(), Scene->GetFeatureLevel()) ||
 			Substrate::IsGlintEnabled(ViewInfo->GetShaderPlatform()) ||
 			IsHairStrandsSupported(EHairStrandsShaderType::Strands, ViewInfo->GetShaderPlatform()) ||
-			IsTranslucencyLightingVolumeUsingBlueNoise();
+			IsTranslucencyLightingVolumeUsingBlueNoise() ||
+			GetSubstrateEnabledUseRoughRefraction();
 		GEngine->LoadBlueNoiseTexture(LoadVector2BlueNoiseTexture);
 
 		if (Substrate::IsGlintEnabled(ViewInfo->GetShaderPlatform()))
