[Postprocessing] Disable film grain in tonemapper if intensity is 0. Enabling film gain may also enable other effects in the shader, so having film grain disabled can save significant time.

diff --git a/Engine/Source/Runtime/Renderer/Private/PostProcess/PostProcessTonemap.cpp b/Engine/Source/Runtime/Renderer/Private/PostProcess/PostProcessTonemap.cpp
--- a/Engine/Source/Runtime/Renderer/Private/PostProcess/PostProcessTonemap.cpp
+++ b/Engine/Source/Runtime/Renderer/Private/PostProcess/PostProcessTonemap.cpp
@@ -141,7 +141,7 @@ FCommonDomain BuildCommonPermutationDomain(const FViewInfo& View, bool bGammaOnl
 	}
 
 	const FPostProcessSettings& Settings = View.FinalPostProcessSettings;
-	PermutationVector.Set<FTonemapperFilmGrainDim>(View.FilmGrainTexture != nullptr);
+	PermutationVector.Set<FTonemapperFilmGrainDim>(View.FinalPostProcessSettings.FilmGrainIntensity > 0.0f && View.FilmGrainTexture != nullptr);
 	PermutationVector.Set<FTonemapperSharpenDim>(GetSharpenSetting(Settings) > 0.0f);
 	PermutationVector.Set<FTonemapperMsaaDim>(bMetalMSAAHDRDecode);
 
