Support baking exposure compensation curve into LUT and sample directly in shader (r.EyeAdaptation.ExposureCompensationCurveLUT)
- Enabled by default.
- LUT texture is recreated on a worker thread whenever the curve changes.
- Use OnUpdateCurve delegate to get notified about changes to the curves. This logic runs on game thread and enqueues a render command that bumps a global "version" number to notify the render thread that LUT needs to be updated.

diff --git a/Engine/Shaders/Private/PostProcessEyeAdaptation.usf b/Engine/Shaders/Private/PostProcessEyeAdaptation.usf
--- a/Engine/Shaders/Private/PostProcessEyeAdaptation.usf
+++ b/Engine/Shaders/Private/PostProcessEyeAdaptation.usf
@@ -174,7 +174,11 @@ float4 EyeAdaptationCommon(float AverageSceneLuminance, float OldExposureScale)
 	// White point luminance is target luminance divided by 0.18 (18% grey).
 	const float TargetExposure = TargetAverageLuminance / 0.18;
 
-	const float MiddleGreyExposureCompensation = EyeAdaptation_ExposureCompensationSettings * EyeAdaptation_ExposureCompensationCurve; // we want the average luminance remapped to 0.18, not 1.0
+	const float EyeAdaptationEV100 = log2((AverageSceneLuminance / 0.18) / EyeAdaptation_LuminanceMax);
+	const float LUT_U = EyeAdaptationEV100 * EyeAdaptation_EV100ToExposureCompensationCurveLUTScaleBias.x + EyeAdaptation_EV100ToExposureCompensationCurveLUTScaleBias.y;
+	const float ExposureCompensationFromCurveLUT = Texture2DSampleLevel(EyeAdaptation_ExposureCompensationCurveLUT, EyeAdaptation_ExposureCompensationCurveSampler, float2(LUT_U, 0.5f), 0).x;
+
+	const float MiddleGreyExposureCompensation = EyeAdaptation_ExposureCompensationSettings * EyeAdaptation_ExposureCompensationCurve * ExposureCompensationFromCurveLUT; // we want the average luminance remapped to 0.18, not 1.0
 	const float OldExposure = MiddleGreyExposureCompensation / (OldExposureScale != 0 ? OldExposureScale : 1.0f);
 
 	// eye adaptation changes over time
diff --git a/Engine/Shaders/Private/PostProcessHistogramCommon.ush b/Engine/Shaders/Private/PostProcessHistogramCommon.ush
--- a/Engine/Shaders/Private/PostProcessHistogramCommon.ush
+++ b/Engine/Shaders/Private/PostProcessHistogramCommon.ush
@@ -31,6 +31,9 @@ float EyeAdaptation_StartDistance;
 float EyeAdaptation_LuminanceMax;
 float EyeAdaptation_ForceTarget;
 int EyeAdaptation_VisualizeDebugType;
+float2 EyeAdaptation_EV100ToExposureCompensationCurveLUTScaleBias;
+Texture2D EyeAdaptation_ExposureCompensationCurveLUT;
+SamplerState EyeAdaptation_ExposureCompensationCurveSampler;
 Texture2D EyeAdaptation_MeterMaskTexture;
 SamplerState EyeAdaptation_MeterMaskSampler;
 
diff --git a/Engine/Source/Runtime/Engine/Private/SceneView.cpp b/Engine/Source/Runtime/Engine/Private/SceneView.cpp
--- a/Engine/Source/Runtime/Engine/Private/SceneView.cpp
+++ b/Engine/Source/Runtime/Engine/Private/SceneView.cpp
@@ -2363,6 +2363,11 @@ void FSceneView::EndFinalPostprocessSettings(const FSceneViewInitOptions& ViewIn
 	}
 #endif // WITH_EDITOR
 
+	if (State != nullptr)
+	{
+		State->SetExposureCompensationCurve(FinalPostProcessSettings.AutoExposureBiasCurve);
+	}
+
 #if DO_CHECK || USING_CODE_ANALYSIS
 	check(VerifyMembersChecks());
 #endif
diff --git a/Engine/Source/Runtime/Engine/Public/SceneManagement.h b/Engine/Source/Runtime/Engine/Public/SceneManagement.h
--- a/Engine/Source/Runtime/Engine/Public/SceneManagement.h
+++ b/Engine/Source/Runtime/Engine/Public/SceneManagement.h
@@ -206,6 +206,8 @@ public:
 	/** Returns the eye adaptation exposure. */
 	virtual float GetLastEyeAdaptationExposure() const = 0;
 
+	virtual void SetExposureCompensationCurve(class UCurveFloat* NewCurve) {}
+
 	virtual void SetSequencerState(ESequencerState InSequencerState) = 0;
 
 	virtual ESequencerState GetSequencerState() = 0;
diff --git a/Engine/Source/Runtime/Renderer/Private/DeferredShadingRenderer.cpp b/Engine/Source/Runtime/Renderer/Private/DeferredShadingRenderer.cpp
--- a/Engine/Source/Runtime/Renderer/Private/DeferredShadingRenderer.cpp
+++ b/Engine/Source/Runtime/Renderer/Private/DeferredShadingRenderer.cpp
@@ -1174,6 +1174,12 @@ void FDeferredShadingSceneRenderer::CommitFinalPipelineState()
 	}
 }
 
+struct FUpdateExposureCompensationCurveLUTTaskData
+{
+	UE::Tasks::FTask Task;
+	FRHICommandList* CmdList = nullptr;
+};
+
 void FDeferredShadingSceneRenderer::RenderNanite(FRDGBuilder& GraphBuilder, const TArray<FViewInfo>& InViews, FSceneTextures& SceneTextures, bool bIsEarlyDepthComplete,
 	FNaniteBasePassVisibility& InNaniteBasePassVisibility,
 	TArray<Nanite::FRasterResults, TInlineAllocator<2>>& NaniteRasterResults,
@@ -1487,6 +1493,25 @@ void FDeferredShadingSceneRenderer::Render(FRDGBuilder& GraphBuilder, const FSce
 
 	FInitViewTaskDatas InitViewTaskDatas = OnRenderBegin(GraphBuilder, SceneUpdateInputs);
 
+	FUpdateExposureCompensationCurveLUTTaskData UpdateExposureCompensationCurveLUTTaskData;
+	UpdateExposureCompensationCurveLUTTaskData.Task = LaunchSceneRenderTask(TEXT("UpdateExposureCompensationCurveLUTs"),
+		[TaskData = &UpdateExposureCompensationCurveLUTTaskData, Views = MakeConstArrayView(Views)]
+		{
+			TaskData->CmdList = new FRHICommandList();
+			TaskData->CmdList->SwitchPipeline(ERHIPipeline::Graphics);
+
+			for (int32 ViewIndex = 0; ViewIndex < Views.Num(); ++ViewIndex)
+			{
+				const FViewInfo& View = Views[ViewIndex];
+				if (View.ViewState != nullptr)
+				{
+					View.ViewState->UpdateExposureCompensationCurveLUT(*TaskData->CmdList, View.FinalPostProcessSettings.AutoExposureBiasCurve);
+				}
+			}
+
+			TaskData->CmdList->FinishRecording();
+		});
+
 	FRDGExternalAccessQueue ExternalAccessQueue;
 	TUniquePtr<FVirtualTextureUpdater> VirtualTextureUpdater;
 	FLumenSceneFrameTemporaries LumenFrameTemporaries(Views);
@@ -3520,6 +3545,15 @@ void FDeferredShadingSceneRenderer::Render(FRDGBuilder& GraphBuilder, const FSce
 			RDG_GPU_STAT_SCOPE(GraphBuilder, Postprocessing);
 			SCOPED_NAMED_EVENT(PostProcessing, FColor::Emerald);
 
+			if (UpdateExposureCompensationCurveLUTTaskData.Task.IsValid())
+			{
+				UpdateExposureCompensationCurveLUTTaskData.Task.Wait();
+				UpdateExposureCompensationCurveLUTTaskData.Task = {};
+
+				GraphBuilder.RHICmdList.QueueAsyncCommandListSubmit(UpdateExposureCompensationCurveLUTTaskData.CmdList);
+				UpdateExposureCompensationCurveLUTTaskData.CmdList = nullptr;
+			}
+
 			FPostProcessingInputs PostProcessingInputs;
 			PostProcessingInputs.ViewFamilyTexture = ViewFamilyTexture;
 			PostProcessingInputs.ViewFamilyDepthTexture = ViewFamilyDepthTexture;
@@ -3624,6 +3658,15 @@ void FDeferredShadingSceneRenderer::Render(FRDGBuilder& GraphBuilder, const FSce
 
 		GEngine->GetPostRenderDelegateEx().Broadcast(GraphBuilder);
 	}
+
+	if (UpdateExposureCompensationCurveLUTTaskData.Task.IsValid())
+	{
+		UpdateExposureCompensationCurveLUTTaskData.Task.Wait();
+		UpdateExposureCompensationCurveLUTTaskData.Task = {};
+
+		GraphBuilder.RHICmdList.QueueAsyncCommandListSubmit(UpdateExposureCompensationCurveLUTTaskData.CmdList);
+		UpdateExposureCompensationCurveLUTTaskData.CmdList = nullptr;
+	}
 	
 	GetSceneExtensionsRenderers().PostRender(GraphBuilder);
 
diff --git a/Engine/Source/Runtime/Renderer/Private/PostProcess/PostProcessEyeAdaptation.cpp b/Engine/Source/Runtime/Renderer/Private/PostProcess/PostProcessEyeAdaptation.cpp
--- a/Engine/Source/Runtime/Renderer/Private/PostProcess/PostProcessEyeAdaptation.cpp
+++ b/Engine/Source/Runtime/Renderer/Private/PostProcess/PostProcessEyeAdaptation.cpp
@@ -73,6 +73,31 @@ namespace
 		TEXT("will have no influence.\n"),
 		ECVF_RenderThreadSafe);
 
+	TAutoConsoleVariable<bool> CVarEyeAdaptationExposureCompensationCurveLUT(
+		TEXT("r.EyeAdaptation.ExposureCompensationCurveLUT"),
+		true,
+		TEXT("Whether to bake the Exposure Compensation Curve to a LUT.\n")
+		TEXT("This allows (approximate) evaluation of the curve on the GPU, avoiding delays due to GPU->CPU readback."),
+		ECVF_RenderThreadSafe);
+
+	TAutoConsoleVariable<int32> CVarEyeAdaptationExposureCompensationCurveLUTNumSamples(
+		TEXT("r.EyeAdaptation.ExposureCompensationCurveLUT.NumSamples"),
+		64,
+		TEXT("Number of samples to use when baking Exposure Compensation Curve to a LUT."),
+		ECVF_RenderThreadSafe);
+
+	TAutoConsoleVariable<float> CVarEyeAdaptationExposureCompensationCurveLUTMinX(
+		TEXT("r.EyeAdaptation.ExposureCompensationCurveLUT.MinX"),
+		-10.0f,
+		TEXT("Minimum X to use when baking Exposure Compensation Curve to a LUT."),
+		ECVF_RenderThreadSafe);
+
+	TAutoConsoleVariable<float> CVarEyeAdaptationExposureCompensationCurveLUTMaxX(
+		TEXT("r.EyeAdaptation.ExposureCompensationCurveLUT.MaxX"),
+		20.0f,
+		TEXT("Maximum X to use when baking Exposure Compensation Curve to a LUT."),
+		ECVF_RenderThreadSafe);
+
 	TAutoConsoleVariable<bool> CVarAutoExposureIgnoreMaterials(
 		TEXT("r.AutoExposure.IgnoreMaterials"),
 		false,
@@ -374,6 +399,18 @@ float CalculateManualAutoExposure(const FViewInfo& View, bool bForceDisablePhysi
 	return FoundLuminance;
 }
 
+FVector2f GetExposureCompensationCurveLUTScaleBias()
+{
+	const float MinX = CVarEyeAdaptationExposureCompensationCurveLUTMinX.GetValueOnRenderThread();
+	const float MaxX = CVarEyeAdaptationExposureCompensationCurveLUTMaxX.GetValueOnRenderThread();
+	const int32 NumSamples = FMath::Max(1, CVarEyeAdaptationExposureCompensationCurveLUTNumSamples.GetValueOnRenderThread());
+
+	const float Scale = (NumSamples - 1) / (NumSamples * (MaxX - MinX));
+	const float Bias = (0.5f / NumSamples) - (MinX * Scale);
+
+	return FVector2f(Scale, Bias);
+}
+
 FEyeAdaptationParameters GetEyeAdaptationParameters(const FViewInfo& View)
 {
 	const bool bExtendedLuminanceRange = IsExtendLuminanceRangeEnabled();
@@ -399,11 +436,29 @@ FEyeAdaptationParameters GetEyeAdaptationParameters(const FViewInfo& View)
 	float MinWhitePointLuminance = 1.0f;
 	float MaxWhitePointLuminance = 1.0f;
 
+	const bool bUseCurveLUT = 
+		Settings.AutoExposureBiasCurve != nullptr
+		&& View.ViewState != nullptr
+		&& View.ViewState->GetExposureCompensationCurveLUT() != nullptr
+		&& CVarEyeAdaptationExposureCompensationCurveLUT.GetValueOnRenderThread();
+
+	FTextureRHIRef ExposureCompensationCurveLUT = nullptr;
+
+	if (bUseCurveLUT)
+	{
+		ExposureCompensationCurveLUT = View.ViewState->GetExposureCompensationCurveLUT();
+	}
+	else
+	{
+		ExposureCompensationCurveLUT = GWhiteTexture->TextureRHI;
+	}
+
 	// Get the exposure compensation from the post process volume settings (everything except the curve)
 	float ExposureCompensationSettings = GetAutoExposureCompensationFromSettings(View);
 
 	// Get the exposure compensation from the curve
-	float ExposureCompensationCurve = GetAutoExposureCompensationFromCurve(View);
+	float ExposureCompensationCurve = bUseCurveLUT ? 1.0f : GetAutoExposureCompensationFromCurve(View);
+
 	const float BlackHistogramBucketInfluence = CVarEyeAdaptationBlackHistogramBucketInfluence.GetValueOnRenderThread();
 
 	const float kMiddleGrey = 0.18f;
@@ -417,12 +472,14 @@ FEyeAdaptationParameters GetEyeAdaptationParameters(const FViewInfo& View)
 	{
 		ExposureCompensationSettings = 1.0f;
 		ExposureCompensationCurve = 1.0f;
+		ExposureCompensationCurveLUT = GWhiteTexture->TextureRHI;
 	}
 	// Fixed exposure override in effect.
 	else if (View.Family->ExposureSettings.bFixed)
 	{
 		ExposureCompensationSettings = 1.0f;
 		ExposureCompensationCurve = 1.0f;
+		ExposureCompensationCurveLUT = GWhiteTexture->TextureRHI;
 
 		// ignores bExtendedLuminanceRange
 		MinWhitePointLuminance = MaxWhitePointLuminance = CalculateFixedAutoExposure(View);
@@ -541,6 +598,9 @@ FEyeAdaptationParameters GetEyeAdaptationParameters(const FViewInfo& View)
 	Parameters.IgnoreMaterialsReconstructFromSceneColor = CVarAutoExposureIgnoreMaterialsReconstructFromSceneColor.GetValueOnRenderThread();
 	Parameters.ForceTarget = ForceTarget;
 	Parameters.VisualizeDebugType = CVarEyeAdaptationVisualizeDebugType.GetValueOnRenderThread();
+	Parameters.EV100ToExposureCompensationCurveLUTScaleBias = GetExposureCompensationCurveLUTScaleBias();
+	Parameters.ExposureCompensationCurveLUT = ExposureCompensationCurveLUT;
+	Parameters.ExposureCompensationCurveSampler = TStaticSamplerState<SF_Bilinear, AM_Clamp, AM_Clamp, AM_Clamp>::GetRHI();
 	Parameters.MeterMaskTexture = MeterMask;
 	Parameters.MeterMaskSampler = TStaticSamplerState<SF_Bilinear, AM_Clamp, AM_Clamp, AM_Clamp>::GetRHI();
 
@@ -577,6 +637,111 @@ float GetEyeAdaptationFixedExposure(const FViewInfo& View)
 	return ExposureScale * Parameters.ExposureCompensationSettings;
 }
 
+static FTextureRHIRef CreateCurveLUT(FRHICommandList& RHICmdList, UCurveFloat* Curve, float MinX, float MaxX, int32 NumSamples)
+{
+	TArray<FFloat16> Data;
+	Data.SetNumUninitialized(NumSamples);
+
+	const float Delta = (MaxX - MinX) / (NumSamples - 1);
+
+	for (int32 Index = 0; Index < NumSamples; ++Index)
+	{
+		const float LuminanceEV100 = MinX + Delta * Index;
+
+		const float Bias = Curve->GetFloatValue(LuminanceEV100);
+
+		FFloat16& Value = Data[Index];
+		Value = FMath::Pow(2.0f, Bias);
+	}
+
+	FResourceBulkDataArrayView BulkDataInterface(Data);
+
+	const FRHITextureCreateDesc Desc = FRHITextureCreateDesc::Create2D(TEXT("ExposureCompensationCurveLUT"), NumSamples, 1, PF_R16F)
+		.SetFlags(ETextureCreateFlags::ShaderResource)
+		.SetBulkData(&BulkDataInterface);
+
+	return RHICmdList.CreateTexture(Desc);
+}
+
+static uint32 GlobalExposureCompensationCurveLUTVersion = 0;
+
+static void BumpGlobalExposureCompensationCurveLUTVersion()
+{
+	ENQUEUE_RENDER_COMMAND(BumpExposureCompensationCurveLUTVersionCommand)(
+		[](FRHICommandListImmediate& RHICmdList)
+		{
+			++GlobalExposureCompensationCurveLUTVersion;
+		});
+}
+
+void FSceneViewState::SetExposureCompensationCurve(UCurveFloat* NewCurve)
+{
+	check(IsInGameThread());
+
+	// Use OnUpdateCurve delegate to get notified about changes to the curves.
+	// At the moment the render thread is informed that curve LUTs need to be update by bumping a global version number.
+	// TODO: could setup a more fine grained invalidation mechanism.
+
+#if WITH_EDITORONLY_DATA
+	if (!CVarEyeAdaptationExposureCompensationCurveLUT.GetValueOnAnyThread())
+	{
+		NewCurve = nullptr;
+	}
+
+	if (ExposureCompensationCurveLUT.Curve_GameThread != NewCurve)
+	{
+		if (ExposureCompensationCurveLUT.Curve_GameThread != nullptr)
+		{
+			ExposureCompensationCurveLUT.Curve_GameThread->OnUpdateCurve.Remove(ExposureCompensationCurveLUT.CurveUpdateDelegateHandle);
+			ExposureCompensationCurveLUT.CurveUpdateDelegateHandle.Reset();
+		}
+
+		ExposureCompensationCurveLUT.Curve_GameThread = NewCurve;
+
+		if (ExposureCompensationCurveLUT.Curve_GameThread != nullptr)
+		{
+			ExposureCompensationCurveLUT.CurveUpdateDelegateHandle = ExposureCompensationCurveLUT.Curve_GameThread->OnUpdateCurve.AddLambda([](UCurveBase* /*Curve*/, EPropertyChangeType::Type /*ChangeType*/) { BumpGlobalExposureCompensationCurveLUTVersion(); });
+		}
+	}
+#endif
+}
+
+void FSceneViewState::UpdateExposureCompensationCurveLUT(FRHICommandList& RHICmdList, UCurveFloat* Curve)
+{
+	if (!CVarEyeAdaptationExposureCompensationCurveLUT.GetValueOnRenderThread())
+	{
+		Curve = nullptr;
+	}
+
+	TRACE_CPUPROFILER_EVENT_SCOPE(UpdateExposureCompensationCurveLUT);
+
+	const float MinX = CVarEyeAdaptationExposureCompensationCurveLUTMinX.GetValueOnRenderThread();
+	const float MaxX = CVarEyeAdaptationExposureCompensationCurveLUTMaxX.GetValueOnRenderThread();
+	const int32 NumSamples = FMath::Max(1, CVarEyeAdaptationExposureCompensationCurveLUTNumSamples.GetValueOnRenderThread());
+
+	bool bNeedToUpdate = false;
+	bNeedToUpdate |= (ExposureCompensationCurveLUT.Texture == nullptr || ExposureCompensationCurveLUT.Texture->GetSizeX() != NumSamples);
+	bNeedToUpdate |= (ExposureCompensationCurveLUT.Curve != Curve);
+#if WITH_EDITORONLY_DATA
+	bNeedToUpdate |= (ExposureCompensationCurveLUT.Version != GlobalExposureCompensationCurveLUTVersion);
+#endif
+	
+	if (Curve != nullptr && bNeedToUpdate)
+	{
+		ExposureCompensationCurveLUT.Texture = CreateCurveLUT(RHICmdList, Curve, MinX, MaxX, NumSamples);
+
+#if WITH_EDITORONLY_DATA
+		ExposureCompensationCurveLUT.Version = GlobalExposureCompensationCurveLUTVersion;
+#endif
+	}
+	else if (Curve == nullptr)
+	{
+		ExposureCompensationCurveLUT.Texture = nullptr;
+	}
+
+	ExposureCompensationCurveLUT.Curve = Curve;
+}
+
 //////////////////////////////////////////////////////////////////////////
 //! CopyEyeAdaptationToTexture
 //////////////////////////////////////////////////////////////////////////
diff --git a/Engine/Source/Runtime/Renderer/Private/PostProcess/PostProcessEyeAdaptation.h b/Engine/Source/Runtime/Renderer/Private/PostProcess/PostProcessEyeAdaptation.h
--- a/Engine/Source/Runtime/Renderer/Private/PostProcess/PostProcessEyeAdaptation.h
+++ b/Engine/Source/Runtime/Renderer/Private/PostProcess/PostProcessEyeAdaptation.h
@@ -48,6 +48,8 @@ int32 GetAutoExposureIlluminanceDownscaleFactor();
 // Returns the auto exposure method enabled by the view (including CVar override).
 EAutoExposureMethod GetAutoExposureMethod(const FViewInfo& View);
 
+void SetupExposureCompensationCurveLUT(FRHICommandListImmediate& RHICmdList, const FViewInfo& View);
+
 BEGIN_SHADER_PARAMETER_STRUCT(FEyeAdaptationParameters, )
 	SHADER_PARAMETER(float, ExposureLowPercent)
 	SHADER_PARAMETER(float, ExposureHighPercent)
@@ -74,6 +76,9 @@ BEGIN_SHADER_PARAMETER_STRUCT(FEyeAdaptationParameters, )
 	SHADER_PARAMETER(uint32, IgnoreMaterialsReconstructFromSceneColor)
 	SHADER_PARAMETER(float, ForceTarget)
 	SHADER_PARAMETER(int, VisualizeDebugType)
+	SHADER_PARAMETER(FVector2f, EV100ToExposureCompensationCurveLUTScaleBias)
+	SHADER_PARAMETER_TEXTURE(Texture2D, ExposureCompensationCurveLUT)
+	SHADER_PARAMETER_SAMPLER(SamplerState, ExposureCompensationCurveSampler)
 	SHADER_PARAMETER_TEXTURE(Texture2D, MeterMaskTexture)
 	SHADER_PARAMETER_SAMPLER(SamplerState, MeterMaskSampler)
 END_SHADER_PARAMETER_STRUCT()
diff --git a/Engine/Source/Runtime/Renderer/Private/RendererScene.cpp b/Engine/Source/Runtime/Renderer/Private/RendererScene.cpp
--- a/Engine/Source/Runtime/Renderer/Private/RendererScene.cpp
+++ b/Engine/Source/Runtime/Renderer/Private/RendererScene.cpp
@@ -381,6 +381,14 @@ FSceneViewState::~FSceneViewState()
 	{
 		Scene->RemoveViewState_RenderThread(this);
 	}
+
+#if WITH_EDITORONLY_DATA
+	if (ExposureCompensationCurveLUT.Curve_GameThread != nullptr)
+	{
+		ExposureCompensationCurveLUT.Curve_GameThread->OnUpdateCurve.Remove(ExposureCompensationCurveLUT.CurveUpdateDelegateHandle);
+		ExposureCompensationCurveLUT.CurveUpdateDelegateHandle.Reset();
+	}
+#endif
 }
 
 void FSceneViewState::ReleaseRHI()
diff --git a/Engine/Source/Runtime/Renderer/Private/ScenePrivate.h b/Engine/Source/Runtime/Renderer/Private/ScenePrivate.h
--- a/Engine/Source/Runtime/Renderer/Private/ScenePrivate.h
+++ b/Engine/Source/Runtime/Renderer/Private/ScenePrivate.h
@@ -926,6 +926,16 @@ private:
 
 	} EyeAdaptationManager;
 
+	struct {
+		FTextureRHIRef Texture;
+		UCurveFloat* Curve = nullptr;
+#if WITH_EDITORONLY_DATA
+		UCurveFloat* Curve_GameThread = nullptr;
+		FDelegateHandle CurveUpdateDelegateHandle;
+		uint32 Version = 0;
+#endif
+	} ExposureCompensationCurveLUT;
+
 	FSubstrateViewDebugData SubstrateViewDebugData;
 
 	// The LUT used by tonemapping.  In stereo this is only computed and stored by the Left Eye.
@@ -1320,6 +1330,15 @@ public:
 		return EyeAdaptationManager.GetLastAverageSceneLuminance();
 	}
 
+	FRHITexture* GetExposureCompensationCurveLUT() const
+	{
+		return ExposureCompensationCurveLUT.Texture;
+	}
+
+	virtual void SetExposureCompensationCurve(UCurveFloat* NewCurve) override;
+
+	void UpdateExposureCompensationCurveLUT(FRHICommandList& RHICmdList, UCurveFloat* Curve);
+
 	bool HasValidTonemappingLUT() const
 	{
 		return bValidTonemappingLUT;
