Fix deprecation warning
#rnx

diff --git a/Engine/Source/Runtime/Renderer/Private/PostProcess/PostProcessEyeAdaptation.cpp b/Engine/Source/Runtime/Renderer/Private/PostProcess/PostProcessEyeAdaptation.cpp
--- a/Engine/Source/Runtime/Renderer/Private/PostProcess/PostProcessEyeAdaptation.cpp
+++ b/Engine/Source/Runtime/Renderer/Private/PostProcess/PostProcessEyeAdaptation.cpp
@@ -658,7 +658,7 @@ static FTextureRHIRef CreateCurveLUT(FRHICommandList& RHICmdList, UCurveFloat* C
 
 	const FRHITextureCreateDesc Desc = FRHITextureCreateDesc::Create2D(TEXT("ExposureCompensationCurveLUT"), NumSamples, 1, PF_R16F)
 		.SetFlags(ETextureCreateFlags::ShaderResource)
-		.SetBulkData(&BulkDataInterface);
+		.SetInitActionBulkData(&BulkDataInterface);
 
 	return RHICmdList.CreateTexture(Desc);
 }
