[VSM] Fixed incorrect cache behavior for added movable primitives that should be treated as dynamic immediately. Not doing so invalidates static cache, adding potentially significant cost.

diff --git a/Engine/Source/Runtime/Renderer/Private/VirtualShadowMaps/VirtualShadowMapCacheManager.cpp b/Engine/Source/Runtime/Renderer/Private/VirtualShadowMaps/VirtualShadowMapCacheManager.cpp
--- a/Engine/Source/Runtime/Renderer/Private/VirtualShadowMaps/VirtualShadowMapCacheManager.cpp
+++ b/Engine/Source/Runtime/Renderer/Private/VirtualShadowMaps/VirtualShadowMapCacheManager.cpp
@@ -565,6 +565,7 @@ void FVirtualShadowMapArrayCacheManager::FInvalidatingPrimitiveCollector::AddInv
 		if (PrimitiveSceneInfo->Proxy->IsMeshShapeOftenMoving() && PrimitiveSceneInfo->Proxy->GetShadowCacheInvalidationBehavior() != EShadowCacheInvalidationBehavior::Static)
 		{
 			InvalidatedPrimitives[PersistentPrimitiveIndex.Index] = true;
+			Manager.CachePrimitiveAsDynamic[PersistentPrimitiveIndex.Index] = true;
 		}
 	}
 	
