Heterogeneous Volumes: Reinstating Velocity write permutation in live shading pipeline

Using a live variable was causing regression failures on some platforms.

diff --git a/Engine/Shaders/Private/HeterogeneousVolumes/HeterogeneousVolumesLiveShadingPipeline.usf b/Engine/Shaders/Private/HeterogeneousVolumes/HeterogeneousVolumesLiveShadingPipeline.usf
--- a/Engine/Shaders/Private/HeterogeneousVolumes/HeterogeneousVolumesLiveShadingPipeline.usf
+++ b/Engine/Shaders/Private/HeterogeneousVolumes/HeterogeneousVolumesLiveShadingPipeline.usf
@@ -32,6 +32,10 @@
 #define USE_ANALYTIC_DERIVATIVES 0
 #endif //
 
+#ifndef DIM_WRITE_VELOCITY
+#define DIM_WRITE_VELOCITY 0
+#endif //DIM_WRITE_VELOCITY
+
 #define USE_CAMERA_AVSM (DIM_USE_ADAPTIVE_VOLUMETRIC_SHADOW_MAP && (SUPPORT_OVERLAPPING_VOLUMES != 0))
 
 #if SUPPORT_OVERLAPPING_VOLUMES
@@ -383,7 +387,7 @@ int GetIndirectLightingMode()
 bool ShouldWriteVelocity()
 {
 #if HV_SCALABILITY_MODE == HV_SCALABILITY_MODE_CINEMATIC
-	return bWriteVelocity;
+	return DIM_WRITE_VELOCITY == 1;
 #else
 	return false;
 #endif
diff --git a/Engine/Source/Runtime/Renderer/Private/HeterogeneousVolumes/HeterogeneousVolumes.cpp b/Engine/Source/Runtime/Renderer/Private/HeterogeneousVolumes/HeterogeneousVolumes.cpp
--- a/Engine/Source/Runtime/Renderer/Private/HeterogeneousVolumes/HeterogeneousVolumes.cpp
+++ b/Engine/Source/Runtime/Renderer/Private/HeterogeneousVolumes/HeterogeneousVolumes.cpp
@@ -1086,7 +1086,7 @@ void FDeferredShadingSceneRenderer::RenderHeterogeneousVolumes(
 			HeterogeneousVolumeBeerShadowMap = GraphBuilder.CreateTexture(Desc, TEXT("HeterogeneousVolume.BeerShadowMap"));
 			AddClearUAVPass(GraphBuilder, GraphBuilder.CreateUAV(HeterogeneousVolumeBeerShadowMap), FVector4d(HeterogeneousVolumes::GetMaxTraceDistance(), 0.0, 1.0, 1.0));
 		}
-
+#if 0
 		// Create velocity texture if it doesn't already exist
 		if (!HasBeenProduced(HeterogeneousVolumeVelocity))
 		{
@@ -1095,6 +1095,7 @@ void FDeferredShadingSceneRenderer::RenderHeterogeneousVolumes(
 			HeterogeneousVolumeVelocity = GraphBuilder.CreateTexture(VelocityDesc, TEXT("HeterogeneousVolume.Velocity"));
 			AddClearUAVPass(GraphBuilder, GraphBuilder.CreateUAV(HeterogeneousVolumeVelocity), FLinearColor::Black);
 		}
+#endif
 	}
 
 	for (int32 ViewIndex = 0; ViewIndex < Views.Num(); ViewIndex++)
diff --git a/Engine/Source/Runtime/Renderer/Private/HeterogeneousVolumes/HeterogeneousVolumesLiveShadingPipeline.cpp b/Engine/Source/Runtime/Renderer/Private/HeterogeneousVolumes/HeterogeneousVolumesLiveShadingPipeline.cpp
--- a/Engine/Source/Runtime/Renderer/Private/HeterogeneousVolumes/HeterogeneousVolumesLiveShadingPipeline.cpp
+++ b/Engine/Source/Runtime/Renderer/Private/HeterogeneousVolumes/HeterogeneousVolumesLiveShadingPipeline.cpp
@@ -464,7 +464,14 @@ class FRenderSingleScatteringWithLiveShadingCS : public FMeshMaterialShader
 	class FUseInscatteringVolume : SHADER_PERMUTATION_BOOL("DIM_USE_INSCATTERING_VOLUME");
 	class FAVSMSampleMode : SHADER_PERMUTATION_INT("AVSM_SAMPLE_MODE", 3);
 	class FSupportOverlappingVolumes : SHADER_PERMUTATION_BOOL("SUPPORT_OVERLAPPING_VOLUMES");
-	using FPermutationDomain = TShaderPermutationDomain<FScalabilityMode, FUseInscatteringVolume, FAVSMSampleMode, FSupportOverlappingVolumes>;
+	class FWriteVelocity : SHADER_PERMUTATION_BOOL("DIM_WRITE_VELOCITY");
+	using FPermutationDomain = TShaderPermutationDomain<
+		FScalabilityMode, 
+		FUseInscatteringVolume, 
+		FAVSMSampleMode, 
+		FSupportOverlappingVolumes,
+		FWriteVelocity
+	>;
 
 	BEGIN_SHADER_PARAMETER_STRUCT(FParameters, )
 		// Scene data
@@ -1370,9 +1377,10 @@ void RenderSingleScatteringWithLiveShading(
 		{
 			PassParameters->RWHoldoutTexture = GraphBuilder.CreateUAV(HeterogeneousVolumeHoldoutTexture);
 		}
-
-		PassParameters->RWVelocityTexture = GraphBuilder.CreateUAV(HeterogeneousVolumeVelocityTexture);
-
+		if (bWriteVelocity)
+		{
+			PassParameters->RWVelocityTexture = GraphBuilder.CreateUAV(HeterogeneousVolumeVelocityTexture);
+		}
 		//if (bCreateBeerShadowMap)
 		{
 			PassParameters->RWBeerShadowMapTexture = GraphBuilder.CreateUAV(HeterogeneousVolumeBeerShadowMapTexture);
@@ -1399,6 +1407,7 @@ void RenderSingleScatteringWithLiveShading(
 	PermutationVector.template Set<typename FRenderSingleScatteringWithLiveShadingDispatchTypeCS::FUseInscatteringVolume>(HeterogeneousVolumes::UseLightingCacheForInscattering());
 	PermutationVector.template Set<typename FRenderSingleScatteringWithLiveShadingDispatchTypeCS::FAVSMSampleMode>(static_cast<int32>(HeterogeneousVolumes::GetAVSMSampleMode(bUseAVSM)));
 	PermutationVector.template Set<typename FRenderSingleScatteringWithLiveShadingDispatchTypeCS::FSupportOverlappingVolumes>(HeterogeneousVolumes::SupportsOverlappingVolumes());
+	PermutationVector.template Set<typename FRenderSingleScatteringWithLiveShadingDispatchTypeCS::FWriteVelocity>(bWriteVelocity);
 	PermutationVector = FRenderSingleScatteringWithLiveShadingDispatchTypeCS::RemapPermutation(PermutationVector);
 	TShaderRef<FRenderSingleScatteringWithLiveShadingDispatchTypeCS> ComputeShader = Material.GetShader<FRenderSingleScatteringWithLiveShadingDispatchTypeCS>(&FLocalVertexFactory::StaticType, PermutationVector, false);
 	if (!ComputeShader.IsNull())
