Enable Translucent object receive CSM shadow when surface lighting

diff --git a/Engine/Source/Runtime/Engine/Private/Materials/MaterialShared.cpp b/Engine/Source/Runtime/Engine/Private/Materials/MaterialShared.cpp
--- a/Engine/Source/Runtime/Engine/Private/Materials/MaterialShared.cpp
+++ b/Engine/Source/Runtime/Engine/Private/Materials/MaterialShared.cpp
@@ -5680,6 +5680,7 @@ FMaterialShaderParametersBase::FMaterialShaderParametersBase(const FMaterial* In
 	bAllowVariableRateShading = InMaterial->IsVariableRateShadingAllowed();
 	bSupportsMaterialCache = InMaterial->SupportsMaterialCache();
 	bIsUsedWithStaticMesh = InMaterial->IsUsedWithStaticMesh();
+	bIsTranslucentSurfaceLighting = InMaterial->GetTranslucencyLightingMode() == ETranslucencyLightingMode::TLM_SurfacePerPixelLighting; // get lighting mode for mobile translucent receive CSM 
 }
 
 FMaterialShaderParameters::FMaterialShaderParameters(const FMaterial* InMaterial)
diff --git a/Engine/Source/Runtime/Engine/Public/MaterialShared.h b/Engine/Source/Runtime/Engine/Public/MaterialShared.h
--- a/Engine/Source/Runtime/Engine/Public/MaterialShared.h
+++ b/Engine/Source/Runtime/Engine/Public/MaterialShared.h
@@ -1004,6 +1004,7 @@ struct FMaterialShaderParametersBase
 		struct
 		{
 			uint16 bIsUsedWithStaticMesh : 1;
+			uint16 bIsTranslucentSurfaceLighting : 1;
 		};
 	};
 	EBlendMode BlendMode;
diff --git a/Engine/Source/Runtime/Renderer/Private/LightMapRendering.cpp b/Engine/Source/Runtime/Renderer/Private/LightMapRendering.cpp
--- a/Engine/Source/Runtime/Renderer/Private/LightMapRendering.cpp
+++ b/Engine/Source/Runtime/Renderer/Private/LightMapRendering.cpp
@@ -170,8 +170,9 @@ bool FMobileDirectionalLightAndCSMPolicy::ShouldCompilePermutation(const FMeshMa
 		return false;
 	}
 	
-	return Parameters.MaterialParameters.ShadingModels.IsLit() && !IsTranslucentBlendMode(Parameters.MaterialParameters) 
-		&& (!IsStaticLightingAllowed() || FReadOnlyCVARCache::MobileEnableStaticAndCSMShadowReceivers());
+	return Parameters.MaterialParameters.ShadingModels.IsLit() &&
+		(!IsTranslucentBlendMode(Parameters.MaterialParameters) || !!Parameters.MaterialParameters.bIsTranslucentSurfaceLighting) && // receive CSM when opaque or translucent surface lighting
+		(!IsStaticLightingAllowed() || FReadOnlyCVARCache::MobileEnableStaticAndCSMShadowReceivers());
 }
 
 bool FMobileDistanceFieldShadowsAndLQLightMapPolicy::ShouldCompilePermutation(const FMeshMaterialShaderPermutationParameters& Parameters)
@@ -179,8 +180,8 @@ bool FMobileDistanceFieldShadowsAndLQLightMapPolicy::ShouldCompilePermutation(co
 	static auto* CVarMobileAllowDistanceFieldShadows = IConsoleManager::Get().FindTConsoleVariableDataInt(TEXT("r.Mobile.AllowDistanceFieldShadows"));
 	const bool bMobileAllowDistanceFieldShadows = CVarMobileAllowDistanceFieldShadows->GetValueOnAnyThread() == 1;
 	
-	return bMobileAllowDistanceFieldShadows && 
-		!IsTranslucentBlendMode(Parameters.MaterialParameters) &&
+	return bMobileAllowDistanceFieldShadows &&
+		(!IsTranslucentBlendMode(Parameters.MaterialParameters) || !!Parameters.MaterialParameters.bIsTranslucentSurfaceLighting) && // receive CSM when opaque or translucent surface lighting
 		Super::ShouldCompilePermutation(Parameters);
 }
 
@@ -213,7 +214,8 @@ bool FMobileDirectionalLightCSMAndLightMapPolicy::ShouldCompilePermutation(const
 		return false;
 	}
 
-	return !IsTranslucentBlendMode(Parameters.MaterialParameters) && Super::ShouldCompilePermutation(Parameters);
+	return (!IsTranslucentBlendMode(Parameters.MaterialParameters) || !!Parameters.MaterialParameters.bIsTranslucentSurfaceLighting) && // receive CSM when opaque or translucent surface lighting
+		Super::ShouldCompilePermutation(Parameters);
 }
 
 void FMobileDirectionalLightCSMAndLightMapPolicy::ModifyCompilationEnvironment(const FMaterialShaderPermutationParameters& Parameters, FShaderCompilerEnvironment& OutEnvironment)
@@ -242,7 +244,7 @@ bool FMobileDirectionalLightCSMAndSHIndirectPolicy::ShouldCompilePermutation(con
 	}
 
 	return FReadOnlyCVARCache::MobileEnableStaticAndCSMShadowReceivers() &&
-		!IsTranslucentBlendMode(Parameters.MaterialParameters) &&
+		(!IsTranslucentBlendMode(Parameters.MaterialParameters) || !!Parameters.MaterialParameters.bIsTranslucentSurfaceLighting) && // receive CSM when opaque or translucent surface lighting
 		Super::ShouldCompilePermutation(Parameters);
 }
 
diff --git a/Engine/Source/Runtime/Renderer/Private/MobileBasePass.cpp b/Engine/Source/Runtime/Renderer/Private/MobileBasePass.cpp
--- a/Engine/Source/Runtime/Renderer/Private/MobileBasePass.cpp
+++ b/Engine/Source/Runtime/Renderer/Private/MobileBasePass.cpp
@@ -290,7 +290,7 @@ ELightMapPolicyType MobileBasePass::SelectMeshLightmapPolicy(
 	bool bPrimReceivesCSM,
 	bool bUsesDeferredShading,
 	bool bIsLitMaterial,
-	bool bIsTranslucent)
+	bool bIsTranslucent) // not used
 {
 	// Unlit uses NoLightmapPolicy with 0 point lights
 	ELightMapPolicyType SelectedLightmapPolicy = LMP_NO_LIGHTMAP;
@@ -332,7 +332,8 @@ ELightMapPolicyType MobileBasePass::SelectMeshLightmapPolicy(
 				MobileDirectionalLight 
 				&& bPrimReceivesCSM
 				&& FReadOnlyCVARCache::MobileEnableStaticAndCSMShadowReceivers()
-				&& MobileDirectionalLight->ShouldRenderViewIndependentWholeSceneShadows();
+				&& MobileDirectionalLight->ShouldRenderViewIndependentWholeSceneShadows()
+				&& PrimitiveSceneProxy->ShouldReceiveMobileCSMShadows(); // receive CSM shadow enabled on meshcomponent
 
 			const bool bPrimitiveUsesILC = 
 				PrimitiveSceneProxy
@@ -345,7 +346,7 @@ ELightMapPolicyType MobileBasePass::SelectMeshLightmapPolicy(
 
 			if (LightMapInteraction.GetType() == LMIT_Texture && FReadOnlyCVARCache::EnableLowQualityLightmaps())
 			{
-				const FShadowMapInteraction ShadowMapInteraction = (Mesh.LCI != nullptr && !bIsTranslucent)
+				const FShadowMapInteraction ShadowMapInteraction = Mesh.LCI != nullptr
 					? Mesh.LCI->GetShadowMapInteraction(FeatureLevel)
 					: FShadowMapInteraction();
 
@@ -385,7 +386,7 @@ ELightMapPolicyType MobileBasePass::SelectMeshLightmapPolicy(
 
 typedef TArray<ELightMapPolicyType, TInlineAllocator<4>> FMobileLightMapPolicyTypeList;
 
-static FMobileLightMapPolicyTypeList GetUniformLightMapPolicyTypeForPSOCollection(bool bLitMaterial, bool bTranslucent, bool bUsesDeferredShading, bool bCanReceiveCSM, bool bMovable)
+static FMobileLightMapPolicyTypeList GetUniformLightMapPolicyTypeForPSOCollection(bool bLitMaterial, bool bTranslucent/*not used*/, bool bUsesDeferredShading, bool bCanReceiveCSM, bool bMovable)
 {
 	FMobileLightMapPolicyTypeList Result;
 	
@@ -407,7 +408,7 @@ static FMobileLightMapPolicyTypeList GetUniformLightMapPolicyTypeForPSOCollectio
 			{
 				if (FReadOnlyCVARCache::MobileEnableStaticAndCSMShadowReceivers() && !bUsesDeferredShading && bCanReceiveCSM)
 				{
-					if (FReadOnlyCVARCache::MobileAllowDistanceFieldShadows() && !bTranslucent)
+					if (FReadOnlyCVARCache::MobileAllowDistanceFieldShadows())
 					{
 						Result.Add(LMP_MOBILE_DISTANCE_FIELD_SHADOWS_LIGHTMAP_AND_CSM);
 					}
@@ -415,7 +416,7 @@ static FMobileLightMapPolicyTypeList GetUniformLightMapPolicyTypeForPSOCollectio
 					Result.Add(LMP_MOBILE_DIRECTIONAL_LIGHT_CSM_AND_LIGHTMAP);
 				}
 
-				if (FReadOnlyCVARCache::MobileAllowDistanceFieldShadows() && !bCanReceiveCSM && !bTranslucent)
+				if (FReadOnlyCVARCache::MobileAllowDistanceFieldShadows() && !bCanReceiveCSM)
 				{
 					Result.Add(LMP_MOBILE_DISTANCE_FIELD_SHADOWS_AND_LQ_LIGHTMAP);
 				}
@@ -736,7 +737,6 @@ bool FMobileBasePassMeshProcessor::ShouldDraw(const FMaterial& Material) const
 {
 	const FMaterialShadingModelField ShadingModels = Material.GetShadingModels();
 	const bool bIsTranslucent = IsTranslucentBlendMode(Material.GetBlendMode()) || ShadingModels.HasShadingModel(MSM_SingleLayerWater); // Water goes into the translucent pass;
-	const bool bCanReceiveCSM = ((Flags & FMobileBasePassMeshProcessor::EFlags::CanReceiveCSM) == FMobileBasePassMeshProcessor::EFlags::CanReceiveCSM);
 	if (bTranslucentBasePass)
 	{
 		// Skipping TPT_TranslucencyAfterDOFModulate. That pass is only needed for Dual Blending, which is not supported on Mobile.
@@ -745,7 +745,6 @@ bool FMobileBasePassMeshProcessor::ShouldDraw(const FMaterial& Material) const
 				|| (TranslucencyPassType == ETranslucencyPass::TPT_TranslucencyStandard && !Material.IsMobileSeparateTranslucencyEnabled())
 				|| (TranslucencyPassType == ETranslucencyPass::TPT_TranslucencyAfterDOF && Material.IsMobileSeparateTranslucencyEnabled()));
 
-		check(!bShouldDraw || bCanReceiveCSM == false);
 		return bShouldDraw;
 	}
 	else
@@ -754,6 +753,15 @@ bool FMobileBasePassMeshProcessor::ShouldDraw(const FMaterial& Material) const
 		return !bIsTranslucent;
 	}
 }
+static bool MaterialCanReceiveCSM(const FMaterial& Material)
+{
+	const EBlendMode BlendMode = Material.GetBlendMode();
+	const FMaterialShadingModelField ShadingModels = Material.GetShadingModels();
+	const bool bIsTranslucent = IsTranslucentBlendMode(BlendMode);
+	const bool bSingleLayerWater = ShadingModels.HasShadingModel(MSM_SingleLayerWater);
+	const bool IsTranslucentPerPixelLighting = bIsTranslucent && Material.GetTranslucencyLightingMode() == ETranslucencyLightingMode::TLM_SurfacePerPixelLighting;
+	return !bIsTranslucent || bSingleLayerWater || IsTranslucentPerPixelLighting;
+}
 
 bool FMobileBasePassMeshProcessor::TryAddMeshBatch(const FMeshBatch& RESTRICT MeshBatch, uint64 BatchElementMask, const FPrimitiveSceneProxy* RESTRICT PrimitiveSceneProxy, int32 StaticMeshId, const FMaterialRenderProxy& MaterialRenderProxy, const FMaterial& Material)
 {
@@ -766,11 +774,11 @@ bool FMobileBasePassMeshProcessor::TryAddMeshBatch(const FMeshBatch& RESTRICT Me
 		UMaterialInterface::FilterOutPlatformShadingModels(ShaderPlatform, ShadingModels);
 #endif
 		const bool bSingleLayerWater = ShadingModels.HasShadingModel(MSM_SingleLayerWater);
-		const bool bCanReceiveCSM = bSingleLayerWater || ((Flags & EFlags::CanReceiveCSM) == EFlags::CanReceiveCSM);
 		const EBlendMode BlendMode = Material.GetBlendMode();
 		const bool bIsLitMaterial = ShadingModels.IsLit();
 		const bool bIsTranslucent = IsTranslucentBlendMode(BlendMode) || bSingleLayerWater; // Water goes into the translucent pass;
 		const bool bIsMasked = IsMaskedBlendMode(Material);		
+		const bool bCanReceiveCSM = ((Flags & EFlags::CanReceiveCSM) == EFlags::CanReceiveCSM) && MaterialCanReceiveCSM(Material);
 		ELightMapPolicyType LightmapPolicyType = MobileBasePass::SelectMeshLightmapPolicy(Scene, MeshBatch, PrimitiveSceneProxy, bCanReceiveCSM, bPassUsesDeferredShading, bIsLitMaterial, bIsTranslucent);
 		return Process(MeshBatch, BatchElementMask, StaticMeshId, PrimitiveSceneProxy, MaterialRenderProxy, Material, bIsMasked, bIsTranslucent, ShadingModels, LightmapPolicyType, bCanReceiveCSM, MeshBatch.LCI);
 	}
@@ -810,7 +818,7 @@ bool FMobileBasePassMeshProcessor::Process(
 		const FMaterialRenderProxy& RESTRICT MaterialRenderProxy,
 		const FMaterial& RESTRICT MaterialResource,
 		const bool bIsMasked,
-		const bool bIsTranslucent,
+		const bool bIsTranslucent, //not used
 		FMaterialShadingModelField ShadingModels,
 		const ELightMapPolicyType LightMapPolicyType,
 		const bool bCanReceiveCSM,
@@ -1061,7 +1069,7 @@ void FMobileBasePassMeshProcessor::CollectPSOInitializers(const FSceneTexturesCo
 	}
 	const bool bUseLocalLightPermutation = (LocalLightSetting != EMobileLocalLightSetting::LOCAL_LIGHTS_DISABLED);
 
-	const bool bCanReceiveCSM = ((Flags & FMobileBasePassMeshProcessor::EFlags::CanReceiveCSM) == FMobileBasePassMeshProcessor::EFlags::CanReceiveCSM);
+	const bool bCanReceiveCSM = ((Flags & EFlags::CanReceiveCSM) == EFlags::CanReceiveCSM) && MaterialCanReceiveCSM(Material);
 
 	FMobileLightMapPolicyTypeList UniformLightMapPolicyTypes = GetUniformLightMapPolicyTypeForPSOCollection(bLitMaterial, bTranslucentBasePass, bPassUsesDeferredShading, bCanReceiveCSM, bMovable);
 	
@@ -1115,7 +1123,7 @@ FMeshPassProcessor* CreateMobileTranslucencyStandardPassProcessor(ERHIFeatureLev
 	PassDrawRenderState.SetDepthStencilAccess(FExclusiveDepthStencil::DepthRead_StencilRead);
 	
 
-	const FMobileBasePassMeshProcessor::EFlags Flags = FMobileBasePassMeshProcessor::EFlags::CanUseDepthStencil;
+	const FMobileBasePassMeshProcessor::EFlags Flags = FMobileBasePassMeshProcessor::EFlags::CanUseDepthStencil | FMobileBasePassMeshProcessor::EFlags::CanReceiveCSM;
 
 	return new FMobileBasePassMeshProcessor(EMeshPass::TranslucencyStandard, Scene, InViewIfDynamicMeshCommand, PassDrawRenderState, InDrawListContext, Flags, ETranslucencyPass::TPT_TranslucencyStandard);
 }
@@ -1126,7 +1134,7 @@ FMeshPassProcessor* CreateMobileTranslucencyAfterDOFProcessor(ERHIFeatureLevel::
 	PassDrawRenderState.SetDepthStencilState(TStaticDepthStencilState<false, CF_DepthNearOrEqual>::GetRHI());
 	PassDrawRenderState.SetDepthStencilAccess(FExclusiveDepthStencil::DepthRead_StencilRead);
 
-	const FMobileBasePassMeshProcessor::EFlags Flags = FMobileBasePassMeshProcessor::EFlags::CanUseDepthStencil;
+	const FMobileBasePassMeshProcessor::EFlags Flags = FMobileBasePassMeshProcessor::EFlags::CanUseDepthStencil | FMobileBasePassMeshProcessor::EFlags::CanReceiveCSM;
 
 	return new FMobileBasePassMeshProcessor(EMeshPass::TranslucencyAfterDOF, Scene, InViewIfDynamicMeshCommand, PassDrawRenderState, InDrawListContext, Flags, ETranslucencyPass::TPT_TranslucencyAfterDOF);
 }
@@ -1137,7 +1145,7 @@ FMeshPassProcessor* CreateMobileTranslucencyAllPassProcessor(ERHIFeatureLevel::T
 	PassDrawRenderState.SetDepthStencilState(TStaticDepthStencilState<false, CF_DepthNearOrEqual>::GetRHI());
 	PassDrawRenderState.SetDepthStencilAccess(FExclusiveDepthStencil::DepthRead_StencilRead);
 
-	const FMobileBasePassMeshProcessor::EFlags Flags = FMobileBasePassMeshProcessor::EFlags::CanUseDepthStencil;
+	const FMobileBasePassMeshProcessor::EFlags Flags = FMobileBasePassMeshProcessor::EFlags::CanUseDepthStencil | FMobileBasePassMeshProcessor::EFlags::CanReceiveCSM;
 
 	return new FMobileBasePassMeshProcessor(EMeshPass::TranslucencyAll, Scene, InViewIfDynamicMeshCommand, PassDrawRenderState, InDrawListContext, Flags, ETranslucencyPass::TPT_AllTranslucency);
 }
