check PrimitiveSceneProxy pointer not null before using it

diff --git a/Engine/Source/Runtime/Renderer/Private/MobileBasePass.cpp b/Engine/Source/Runtime/Renderer/Private/MobileBasePass.cpp
--- a/Engine/Source/Runtime/Renderer/Private/MobileBasePass.cpp
+++ b/Engine/Source/Runtime/Renderer/Private/MobileBasePass.cpp
@@ -330,6 +330,7 @@ ELightMapPolicyType MobileBasePass::SelectMeshLightmapPolicy(
 			// Primitive can receive both pre-computed and CSM shadows
 			const bool bPrimReceivesStaticAndCSM = 
 				MobileDirectionalLight 
+				&& PrimitiveSceneProxy
 				&& bPrimReceivesCSM
 				&& FReadOnlyCVARCache::MobileEnableStaticAndCSMShadowReceivers()
 				&& MobileDirectionalLight->ShouldRenderViewIndependentWholeSceneShadows()
