[TSR] Default to not use thin geometry weight relaxation to deal with background sky. It is already addressed in the volumetric pass.

diff --git a/Engine/Source/Runtime/Renderer/Private/PostProcess/TemporalSuperResolution.cpp b/Engine/Source/Runtime/Renderer/Private/PostProcess/TemporalSuperResolution.cpp
--- a/Engine/Source/Runtime/Renderer/Private/PostProcess/TemporalSuperResolution.cpp
+++ b/Engine/Source/Runtime/Renderer/Private/PostProcess/TemporalSuperResolution.cpp
@@ -335,7 +335,7 @@ TAutoConsoleVariable<int32> CVarTSRThinGeometryDetectionWeightRelaxation(
 	ECVF_RenderThreadSafe);
 
 TAutoConsoleVariable<int32> CVarTSRThinGeometryDetectionWeightRelaxationSky(
-	TEXT("r.TSR.ThinGeometryDetection.WeightRelaxation.Sky"), 1,
+	TEXT("r.TSR.ThinGeometryDetection.WeightRelaxation.Sky"), 0,
 	TEXT("Apply weight adjust to deal with thin geometry shimmering against sky."),
 	ECVF_RenderThreadSafe);
 
