MeshPassProcessor - Raise severity of check (from checkSlow) to prevent out of bound access when using too much Uniform Buffers

#tests local test & preflight

diff --git a/Engine/Source/Runtime/Renderer/Private/MeshPassProcessor.cpp b/Engine/Source/Runtime/Renderer/Private/MeshPassProcessor.cpp
--- a/Engine/Source/Runtime/Renderer/Private/MeshPassProcessor.cpp
+++ b/Engine/Source/Runtime/Renderer/Private/MeshPassProcessor.cpp
@@ -159,7 +159,7 @@ void FReadOnlyMeshDrawSingleShaderBindings::SetShaderBindings(
 	for (int32 UniformBufferIndex = 0; UniformBufferIndex < NumUniformBuffers; UniformBufferIndex++)
 	{
 		FShaderUniformBufferParameterInfo Parameter = UniformBufferParameters[UniformBufferIndex];
-		checkSlow(Parameter.BaseIndex < UE_ARRAY_COUNT(ShaderBindingState.UniformBuffers));
+		check(Parameter.BaseIndex < FShaderBindingState::MAX_UNIFORM_BUFFERS_PER_STAGE);
 		FRHIUniformBuffer* UniformBuffer = UniformBufferBindings[UniformBufferIndex];
 
 		if (UniformBuffer != ShaderBindingState.UniformBuffers[Parameter.BaseIndex])
