[PostProcess] Add compute pass to upscaler, with optimized sharpening filter. (r.upscale.sharpening)

diff --git a/Engine/Shaders/Private/PostProcessUpscale.usf b/Engine/Shaders/Private/PostProcessUpscale.usf
--- a/Engine/Shaders/Private/PostProcessUpscale.usf
+++ b/Engine/Shaders/Private/PostProcessUpscale.usf
@@ -6,16 +6,45 @@
 #include "PaniniProjection.ush"
 #include "TextureSampling.ush"
 #include "LensDistortion.ush"
+#include "PixelQuadMessagePassing.ush"
+
+#define UPSCALE_METHOD_NEAREST 1
+#define UPSCALE_METHOD_BILINEAR 2
+#define UPSCALE_METHOD_DIRECTIONAL 3
+#define UPSCALE_METHOD_CATMULL_ROM 4
+#define UPSCALE_METHOD_LANCZOS 5
+#define UPSCALE_METHOD_GAUSSIAN 6
+#define UPSCALE_METHOD_SMOOTHSTEP 7
+
+#define METHOD_SET_SIMPLE 1
+#define METHOD_SET_COMPLEX 2
+
+#if DIM_METHOD == UPSCALE_METHOD_NEAREST
+	#define INCLUDE_UPSCALE_METHOD_NEAREST 1
+#elif DIM_METHOD == UPSCALE_METHOD_BILINEAR
+	#define INCLUDE_UPSCALE_METHOD_BILINEAR 1
+#elif DIM_METHOD == UPSCALE_METHOD_DIRECTIONAL
+	#define INCLUDE_UPSCALE_METHOD_DIRECTIONAL 1
+#elif DIM_METHOD == UPSCALE_METHOD_CATMULL_ROM
+	#define INCLUDE_UPSCALE_METHOD_CATMULL_ROM 1
+#elif DIM_METHOD == UPSCALE_METHOD_LANCZOS
+	#define INCLUDE_UPSCALE_METHOD_LANCZOS 1
+#elif DIM_METHOD == UPSCALE_METHOD_GAUSSIAN
+	#define INCLUDE_UPSCALE_METHOD_GAUSSIAN 1
+#elif DIM_METHOD == UPSCALE_METHOD_SMOOTHSTEP
+	#define INCLUDE_UPSCALE_METHOD_SMOOTHSTEP 1
+// see GetMethodSet
+#elif DIM_METHODSET == METHOD_SET_SIMPLE
+	#define INCLUDE_UPSCALE_METHOD_NEAREST 1
+	#define INCLUDE_UPSCALE_METHOD_BILINEAR 1
+	#define INCLUDE_UPSCALE_METHOD_SMOOTHSTEP 1
+#elif DIM_METHODSET == METHOD_SET_COMPLEX
+	#define INCLUDE_UPSCALE_METHOD_DIRECTIONAL 1
+	#define INCLUDE_UPSCALE_METHOD_CATMULL_ROM 1
+	#define INCLUDE_UPSCALE_METHOD_LANCZOS 1
+	#define INCLUDE_UPSCALE_METHOD_GAUSSIAN 1
+#endif
 
-#define UPSCALE_METHOD_NEAREST 0
-#define UPSCALE_METHOD_BILINEAR 1
-#define UPSCALE_METHOD_DIRECTIONAL 2
-#define UPSCALE_METHOD_CATMULL_ROM 3
-#define UPSCALE_METHOD_LANCZOS 4
-#define UPSCALE_METHOD_GAUSSIAN 5
-#define UPSCALE_METHOD_SMOOTHSTEP 6
-
-// only for MainPS
 float UpscaleSoftness;
 
 SCREEN_PASS_TEXTURE_VIEWPORT(Input)
@@ -23,6 +52,8 @@ SCREEN_PASS_TEXTURE_VIEWPORT(Output)
 
 Texture2D<float2> DistortingDisplacementTexture;
 SamplerState DistortingDisplacementSampler;
+Texture2D<float2> UndistortingDisplacementTexture;
+SamplerState UndistortingDisplacementSampler;
 
 Texture2D SceneColorTexture;
 SamplerState SceneColorSampler;
@@ -37,6 +68,8 @@ Texture2DArray PointSceneColorTextureArray;
 uint2 GridDimensions;
 uint bInvertAlpha;
 
+float Sharpening;
+
 // Do a smoothstep(x) = 3 x^2 - 2 x^3
 float2 GetSmoothstepUV(float2 LinearUV, float2 TextureSize, float2 TextureInvSize)
 {
@@ -68,31 +101,33 @@ float Luma(float3 Color)
 	return dot(Color, LuminanceFactors);
 }
 
-float Gaussian(float Scale, float2 Offset) 
+float3 RGBToYCoCg( float3 RGB )
 {
-	return exp2(Scale * dot(Offset, Offset));
+	float Y  = dot( RGB, float3(  1, 2,  1 ) );
+	float Co = dot( RGB, float3(  2, 0, -2 ) );
+	float Cg = dot( RGB, float3( -1, 2, -1 ) );
+	
+	float3 YCoCg = float3( Y, Co, Cg );
+	return YCoCg;
 }
 
-// vertex shader entry point
-void MainVS(
-	in uint VertexId : SV_VertexID,
-	in FStereoVSInput StereoInput,
-	out noperspective float4 OutTexCoord : TEXCOORD0,
-	out FStereoVSOutput StereoOutput,
-	out float4 OutPosition : SV_POSITION)
+float3 YCoCgToRGB( float3 YCoCg )
 {
-	StereoSetupVS(StereoInput, StereoOutput);
-	
-	float2 TexCoord = float2(VertexId % (GridDimensions.x + 1), VertexId / (GridDimensions.x + 1)) / float2(GridDimensions);
+	float Y  = YCoCg.x * 0.25;
+	float Co = YCoCg.y * 0.25;
+	float Cg = YCoCg.z * 0.25;
 
-	// still in 0..1 range
-	float4 Position = float4(TexCoord.x, TexCoord.y, 0, 1);
+	float R = Y + Co - Cg;
+	float G = Y + Cg;
+	float B = Y - Co - Cg;
 
-	// distort pos
-	Position.xy = ApplyLensDistortionOnViewportUV(DistortingDisplacementTexture, DistortingDisplacementSampler, TexCoord);
+	float3 RGB = float3( R, G, B );
+	return RGB;
+}
 
-	DrawRectangle(Position, TexCoord, OutPosition, OutTexCoord.xy);
-	OutTexCoord.zw =  OutPosition.xy;
+float Gaussian(float Scale, float2 Offset) 
+{
+	return exp2(Scale * dot(Offset, Offset));
 }
 
 float4 SampleSceneColorRGBA(float2 BufferUV)
@@ -108,76 +143,94 @@ float4 AccumulateAndApplyWeight(in float4 DataRGBA, in float Weight, inout float
 	return DataRGBA * Weight;
 }
 
-void MainPS(noperspective float4 UVAndScreenPos : TEXCOORD0, FStereoPSInput StereoInput, float4 SvPosition : SV_POSITION, out float4 OutColor : SV_Target0)
+
+
+float4 Upsample(float2 InputUV, uint EyeIndex, uint Method)
 {
-	StereoSetupPS(StereoInput);
-	const uint EyeIndex = GetEyeIndex(StereoInput);
-	
-	OutColor = 0;
+	float4 OutColor = 0;
 
-#if METHOD == UPSCALE_METHOD_NEAREST
+	switch(Method)
+	{
+#if INCLUDE_UPSCALE_METHOD_NEAREST
+	case UPSCALE_METHOD_NEAREST:
+	{
 	// Nearest sampling (not blurry but blocky, more for testing)
 	#if ES3_1_PROFILE
 		#if MOBILE_MULTI_VIEW
-			OutColor = Texture2DArraySample(PointSceneColorTextureArray, PointSceneColorSampler, float3(UVAndScreenPos.xy,EyeIndex));
+			OutColor = Texture2DArraySample(PointSceneColorTextureArray, PointSceneColorSampler, float3(InputUV,EyeIndex));
 		#else
-			OutColor = Texture2DSample(PointSceneColorTexture, PointSceneColorSampler, UVAndScreenPos.xy);
+			OutColor = Texture2DSample(PointSceneColorTexture, PointSceneColorSampler, InputUV);
 		#endif
 	#else
 		#if MOBILE_MULTI_VIEW
-			OutColor = PointSceneColorTextureArray.SampleLevel(PointSceneColorSampler, vec3(UVAndScreenPos.xy,EyeIndex), 0, int2(0, 0));
+			OutColor = PointSceneColorTextureArray.SampleLevel(PointSceneColorSampler, vec3(InputUV,EyeIndex), 0, int2(0, 0));
 		#else
-			OutColor = PointSceneColorTexture.SampleLevel(PointSceneColorSampler, UVAndScreenPos.xy, 0, int2(0, 0));
+			OutColor = PointSceneColorTexture.SampleLevel(PointSceneColorSampler, InputUV, 0, int2(0, 0));
 		#endif
 	#endif
-
-#elif METHOD == UPSCALE_METHOD_BILINEAR
-	// Bilinear (fast, aliasing)
-	OutColor = SampleSceneColorRGBA(UVAndScreenPos.xy);
-
-#elif METHOD == UPSCALE_METHOD_DIRECTIONAL
-	// Directional blur with unsharp mask upsample.
-	float2 UV = UVAndScreenPos.xy;
-	float X = 0.5;
-	float4 ColorNW = SampleSceneColorRGBA(UV + float2(-X, -X) * Input_ExtentInverse);
-	float4 ColorNE = SampleSceneColorRGBA(UV + float2( X, -X) * Input_ExtentInverse);
-	float4 ColorSW = SampleSceneColorRGBA(UV + float2(-X,  X) * Input_ExtentInverse);
-	float4 ColorSE = SampleSceneColorRGBA(UV + float2( X,  X) * Input_ExtentInverse);
-	OutColor = (ColorNW * 0.25) + (ColorNE * 0.25) + (ColorSW * 0.25) + (ColorSE * 0.25);
-	
-	float LumaNW = Luma(ColorNW.rgb);
-	float LumaNE = Luma(ColorNE.rgb);
-	float LumaSW = Luma(ColorSW.rgb);
-	float LumaSE = Luma(ColorSE.rgb);
-
-	float2 IsoBrightnessDir;
-	float DirSWMinusNE = LumaSW - LumaNE;
-	float DirSEMinusNW = LumaSE - LumaNW;
-	IsoBrightnessDir.x = DirSWMinusNE + DirSEMinusNW;
-	IsoBrightnessDir.y = DirSWMinusNE - DirSEMinusNW;
-
-	// avoid NaN on zero vectors by adding 2^-24 (float ulp when length==1, and also minimum representable half)
-	IsoBrightnessDir = IsoBrightnessDir * (0.125 * rsqrt(dot(IsoBrightnessDir, IsoBrightnessDir) + 6e-8));
-
-	float4 ColorN = SampleSceneColorRGBA(UV - IsoBrightnessDir * Input_ExtentInverse);
-	float4 ColorP = SampleSceneColorRGBA(UV + IsoBrightnessDir * Input_ExtentInverse);
-
-	float UnsharpMask = 0.25;
-	OutColor = (ColorN + ColorP) * ((UnsharpMask + 1.0) * 0.5) - (OutColor * UnsharpMask);
-
-#elif METHOD == UPSCALE_METHOD_CATMULL_ROM
-	// Bicubic Catmull-Rom in five samples
-	FCatmullRomSamples Samples = GetBicubic2DCatmullRomSamples(UVAndScreenPos.xy, Input_Extent, Input_ExtentInverse);
-	for (uint i = 0; i < Samples.Count; i++)
+	}
+	break;
+#endif
+#if INCLUDE_UPSCALE_METHOD_BILINEAR
+	case UPSCALE_METHOD_BILINEAR:
 	{
-		OutColor += SampleSceneColorRGBA(Samples.UV[i]) * Samples.Weight[i];
+		// Bilinear (fast, aliasing)
+		OutColor = SampleSceneColorRGBA(InputUV);
 	}
-	OutColor *= Samples.FinalMultiplier;
+	break;
+#endif
+#if INCLUDE_UPSCALE_METHOD_DIRECTIONAL
+	case UPSCALE_METHOD_DIRECTIONAL:
+	{
+		// Directional blur with unsharp mask upsample.
+		float2 UV = InputUV;
+		float X = 0.5;
+		float4 ColorNW = SampleSceneColorRGBA(UV + float2(-X, -X) * Input_ExtentInverse);
+		float4 ColorNE = SampleSceneColorRGBA(UV + float2( X, -X) * Input_ExtentInverse);
+		float4 ColorSW = SampleSceneColorRGBA(UV + float2(-X,  X) * Input_ExtentInverse);
+		float4 ColorSE = SampleSceneColorRGBA(UV + float2( X,  X) * Input_ExtentInverse);
+		OutColor = (ColorNW * 0.25) + (ColorNE * 0.25) + (ColorSW * 0.25) + (ColorSE * 0.25);
+	
+		float LumaNW = Luma(ColorNW.rgb);
+		float LumaNE = Luma(ColorNE.rgb);
+		float LumaSW = Luma(ColorSW.rgb);
+		float LumaSE = Luma(ColorSE.rgb);
 
-#elif METHOD == UPSCALE_METHOD_LANCZOS
+		float2 IsoBrightnessDir;
+		float DirSWMinusNE = LumaSW - LumaNE;
+		float DirSEMinusNW = LumaSE - LumaNW;
+		IsoBrightnessDir.x = DirSWMinusNE + DirSEMinusNW;
+		IsoBrightnessDir.y = DirSWMinusNE - DirSEMinusNW;
+
+		// avoid NaN on zero vectors by adding 2^-24 (float ulp when length==1, and also minimum representable half)
+		IsoBrightnessDir = IsoBrightnessDir * (0.125 * rsqrt(dot(IsoBrightnessDir, IsoBrightnessDir) + 6e-8));
+
+		float4 ColorN = SampleSceneColorRGBA(UV - IsoBrightnessDir * Input_ExtentInverse);
+		float4 ColorP = SampleSceneColorRGBA(UV + IsoBrightnessDir * Input_ExtentInverse);
+
+		float UnsharpMask = 0.25;
+		OutColor = (ColorN + ColorP) * ((UnsharpMask + 1.0) * 0.5) - (OutColor * UnsharpMask);
+	}
+	break;
+#endif
+#if INCLUDE_UPSCALE_METHOD_CATMULL_ROM
+	case UPSCALE_METHOD_CATMULL_ROM:
+	{
+		// Bicubic Catmull-Rom in five samples
+		FCatmullRomSamples Samples = GetBicubic2DCatmullRomSamples(InputUV, Input_Extent, Input_ExtentInverse);
+		for (uint i = 0; i < Samples.Count; i++)
+		{
+			OutColor += SampleSceneColorRGBA(Samples.UV[i]) * Samples.Weight[i];
+		}
+		OutColor *= Samples.FinalMultiplier;
+	}
+	break;
+#endif
+#if INCLUDE_UPSCALE_METHOD_LANCZOS
+	case UPSCALE_METHOD_LANCZOS:
 	{
 		// Lanczos 3
-		float2 UV = UVAndScreenPos.xy * Input_Extent;
+		float2 UV = InputUV * Input_Extent;
 		float2 tc = floor(UV - 0.5) + 0.5;
 		float2 f = UV - tc + 2;
 
@@ -231,11 +284,13 @@ void MainPS(noperspective float4 UVAndScreenPos : TEXCOORD0, FStereoPSInput Ster
 
 		OutColor /= WeightsSum;
 	}
-
-#elif METHOD == UPSCALE_METHOD_GAUSSIAN
+	break;
+#endif
+#if INCLUDE_UPSCALE_METHOD_GAUSSIAN
+	case UPSCALE_METHOD_GAUSSIAN:
 	{
 		// Gaussian filtered unsharp mask
-		float2 UV = UVAndScreenPos.xy * Input_Extent;
+		float2 UV = InputUV * Input_Extent;
 		float2 tc = floor(UV) + 0.5;
 
 		// estimate pixel value and derivatives
@@ -269,13 +324,23 @@ void MainPS(noperspective float4 UVAndScreenPos : TEXCOORD0, FStereoPSInput Ster
 		float UnsharpScale = UpscaleSoftness * (1 - Input_Extent.x * Input_Extent.y * Output_ViewportSizeInverse.x * Output_ViewportSizeInverse.y);
 		OutColor -= UnsharpScale * Laplacian;
 	}
-
-#elif METHOD == UPSCALE_METHOD_SMOOTHSTEP
-	// smooth step.
+	break;
+#endif
+#if INCLUDE_UPSCALE_METHOD_SMOOTHSTEP
+	case UPSCALE_METHOD_SMOOTHSTEP:
 	{
-		OutColor = SampleSceneColorRGBA(GetSmoothstepUV(UVAndScreenPos.xy, Input_Extent, Input_ExtentInverse));
+		OutColor = SampleSceneColorRGBA(GetSmoothstepUV(InputUV, Input_Extent, Input_ExtentInverse));
 	}
+	break;
 #endif
+	}
+
+	return OutColor;
+}
+
+float4 ProcessAlpha(float4 InColor)
+{
+	float4 OutColor = InColor;
 
 #if defined(DIM_ALPHA_CHANNEL) && DIM_ALPHA_CHANNEL == 1
 	FLATTEN
@@ -293,4 +358,285 @@ void MainPS(noperspective float4 UVAndScreenPos : TEXCOORD0, FStereoPSInput Ster
 #else
 	OutColor.a = 0; // Skip all computations related to alpha
 #endif
-}
\ No newline at end of file
+
+	return OutColor;
+}
+
+#if COMPUTESHADER
+
+RWTexture2D<float4> RWOutputTexture;
+
+groupshared uint LDS[(THREADGROUP_SIZEX) * (THREADGROUP_SIZEY)];
+
+uint LDSGetIndex(uint2 GroupThreadId)
+{
+	return (GroupThreadId.y * THREADGROUP_SIZEX) + GroupThreadId.x;
+}
+
+void LDSWriteFloat(uint2 GroupThreadId, float LocalValue)
+{
+	LDS[LDSGetIndex(GroupThreadId)] = asuint(LocalValue);
+}
+
+float LDSReadFloatAtIndex(uint Index)
+{
+	return asfloat(LDS[Index]);
+}
+
+float LDSReadFloat(uint2 GroupThreadId, int2 Direction, out bool bOutOfBounds)
+{
+	int2 ReadPosition = GroupThreadId + Direction;
+	int2 ClampedReadPosition = clamp(ReadPosition, int2(0, 0), int2(THREADGROUP_SIZEX-1, THREADGROUP_SIZEY-1));
+	bOutOfBounds = any(ClampedReadPosition != ReadPosition);
+
+	uint LDSValue = LDS[LDSGetIndex(ClampedReadPosition)];
+	float Result = asfloat(LDSValue);
+
+	return Result;
+}
+
+uint GetTotalNumBoundarySamples()
+{
+	return 2*(THREADGROUP_SIZEX + 2) + 2*(THREADGROUP_SIZEY);
+}
+
+int2 IndexToBoundaryPosition(int Index)
+{
+	// Generates positions on the boundary as shown below.
+	//
+	// e.g. for threadgroup size of 4x4:
+	//        Y
+	//        |
+	//   E F G H I J
+	//   D         0
+	//   C         1
+	//   B         2 -X
+	//   A         3
+	//   9 8 7 6 5 4
+	int2 Result;
+	if (Index < THREADGROUP_SIZEY)
+	{
+		Result.x = THREADGROUP_SIZEX;
+		Result.y = Index;
+	}
+	else if (Index < THREADGROUP_SIZEY + THREADGROUP_SIZEX + 1)
+	{
+		Result.x = THREADGROUP_SIZEX - (Index - THREADGROUP_SIZEY);
+		Result.y = THREADGROUP_SIZEY;
+	}
+	else if (Index < THREADGROUP_SIZEY + THREADGROUP_SIZEX + THREADGROUP_SIZEY + 2)
+	{
+		Result.x = -1;
+		Result.y = THREADGROUP_SIZEY - (Index - (THREADGROUP_SIZEY + THREADGROUP_SIZEX + 1));
+	}
+	else
+	{
+		Result.x = -1 + (Index - (THREADGROUP_SIZEY + THREADGROUP_SIZEX + THREADGROUP_SIZEY + 2));
+		Result.y = -1;
+	}
+	return Result;
+}
+
+uint BoundaryPositionToIndex(int2 Position)
+{
+	if (Position.x == THREADGROUP_SIZEX && Position.y >= 0)
+	{
+		return Position.y;
+	}
+	else if (Position.y == THREADGROUP_SIZEY && Position.x >= 0)
+	{
+		return THREADGROUP_SIZEY + THREADGROUP_SIZEX - Position.x;
+	}
+	else if (Position.x == -1 && Position.y >= 0)
+	{
+		return THREADGROUP_SIZEY + THREADGROUP_SIZEX + THREADGROUP_SIZEY + 1 - Position.y;
+	}
+	else// if (Position.y == -1)
+	{
+		return (Position.x + 1) + THREADGROUP_SIZEY + THREADGROUP_SIZEX + THREADGROUP_SIZEY + 2;
+	}
+}
+
+float2 ApplyLensDistortion(float2 UV, uint2 GroupId, float2 PixelGroupUV)
+{
+	float2 DistortedUV = UV;
+
+#if DIM_LENS_DISTORTION
+
+#if 1 //subsample distortion and interpolate
+	const uint2 GroupSize = uint2(THREADGROUP_SIZEX, THREADGROUP_SIZEY);
+	float2 GroupSizeUV = GroupSize * Output_ViewportSizeInverse;
+
+	uint2 TopLeftPixelPosition = GroupId * GroupSize;
+	float2 TopLeftUV     = ((float2)TopLeftPixelPosition + 0.5f) * Output_ViewportSizeInverse;
+	float2 TopRightUV    = TopLeftUV + float2(GroupSizeUV.x, 0);
+	float2 BottomLeftUV  = TopLeftUV + float2(0, GroupSizeUV.y);
+	float2 BottomRightUV = TopLeftUV + float2(GroupSizeUV.x, GroupSizeUV.y);
+
+	float4x2 Samples;
+	Samples[0] = ApplyLensDistortionOnViewportUV(UndistortingDisplacementTexture, UndistortingDisplacementSampler, TopLeftUV);
+	Samples[1] = ApplyLensDistortionOnViewportUV(UndistortingDisplacementTexture, UndistortingDisplacementSampler, TopRightUV);
+	Samples[2] = ApplyLensDistortionOnViewportUV(UndistortingDisplacementTexture, UndistortingDisplacementSampler, BottomLeftUV);
+	Samples[3] = ApplyLensDistortionOnViewportUV(UndistortingDisplacementTexture, UndistortingDisplacementSampler, BottomRightUV);
+
+	float4 BilinearSampleWeights;
+	BilinearSampleWeights[0] = (1 - PixelGroupUV.x) * (1 - PixelGroupUV.y);
+	BilinearSampleWeights[1] = (PixelGroupUV.x) * (1 - PixelGroupUV.y);
+	BilinearSampleWeights[2] = (1 - PixelGroupUV.x) * (PixelGroupUV.y);
+	BilinearSampleWeights[3] = (PixelGroupUV.x) * (PixelGroupUV.y);
+
+	DistortedUV = mul(BilinearSampleWeights, Samples);
+#else
+	DistortedUV = ApplyLensDistortionOnViewportUV(UndistortingDisplacementTexture, UndistortingDisplacementSampler, UV);
+#endif
+
+#endif
+
+	return DistortedUV;
+}
+
+/*half GrainFromUV(float2 GrainUV) 
+{
+	half Grain = frac(sin(GrainUV.x + GrainUV.y * 543.31) *  493013.0);
+	return Grain;
+}
+
+float3 CalculateFilmGrain(float2 GrainUV)
+{
+	// Compute uniform grain on [0;1]
+	half UniformGrainOP = GrainFromUV(GrainUV.zw); //TODO: verify the quality of GrainFromUV()
+
+	// Compute uniform grain on [-1;1]
+	half UniformGrainNP = UniformGrainOP * 2.0 - 1.0;
+	half UniformGrainNPSign = clamp(UniformGrainNP * POSITIVE_INFINITY, half(-1.0), half(1.0));
+		
+	// Compute triangular grain on [-1;1]
+	half TriangularGrainOP = UniformGrainNPSign - UniformGrainNPSign * sqrt(saturate(half(1.0) - abs(UniformGrainNP)));
+		
+	// Selects the grain to use.
+	half FinalGrain = TriangularGrainOP;
+
+	float BackbufferQuantizationDithering = 0;
+	return FinalGrain * BackbufferQuantizationDithering;
+}*/
+
+float2 PixelPositionToInputUV(uint2 PixelPosition, uint2 GroupId, float2 PixelGroupUV)
+{
+	float2 OutputUV = ((float2)PixelPosition + 0.5f) * Output_ViewportSizeInverse;
+	float2 DistortedOutputUV = ApplyLensDistortion(OutputUV, GroupId, PixelGroupUV);
+	float2 InputBufferUV = (Input_ViewportMin + DistortedOutputUV * Input_ViewportSize) * Input_ExtentInverse;
+	return InputBufferUV;
+}
+
+uint UpscaleMethod;
+
+[numthreads(THREADGROUP_SIZEX, THREADGROUP_SIZEY, 1)]
+void MainCS(uint2 DispatchThreadId : SV_DispatchThreadID, uint2 GroupId : SV_GroupID, uint2 GroupThreadId : SV_GroupThreadID, uint GroupThreadIndex : SV_GroupIndex)
+{
+	const uint2 GroupSize = uint2(THREADGROUP_SIZEX, THREADGROUP_SIZEY);
+
+	const uint EyeIndex = 0; // Unsupported
+	int2 PixelPosition = DispatchThreadId;
+	int2 BufferPixelPosition = PixelPosition + Output_ViewportMin;
+	float2 PixelGroupUV = float2(GroupThreadId) / GroupSize;
+
+	float2 InputBufferUV = PixelPositionToInputUV(PixelPosition, GroupId, PixelGroupUV);
+
+	float4 OutColor = Upsample(InputBufferUV, EyeIndex, UpscaleMethod);
+
+#if DIM_SHARPENING_QUALITY > 0
+	float3 ColorYCoCg = RGBToYCoCg(OutColor.rgb);
+	float LocalLuminance = ColorYCoCg.x;
+
+	LDSWriteFloat(GroupThreadId, LocalLuminance);
+	GroupMemoryBarrierWithGroupSync();
+
+	bool bOutOfBoundsLeft, bOutOfBoundsTop, bOutOfBoundsRight, bOutOfBoundsBottom;
+	float LuminanceLeft   = LDSReadFloat(GroupThreadId, uint2(-1,  0), bOutOfBoundsLeft);
+	float LuminanceTop    = LDSReadFloat(GroupThreadId, uint2( 0, -1), bOutOfBoundsTop);
+	float LuminanceRight  = LDSReadFloat(GroupThreadId, uint2( 1,  0), bOutOfBoundsRight);
+	float LuminanceBottom = LDSReadFloat(GroupThreadId, uint2( 0,  1), bOutOfBoundsBottom);
+
+	bool bOutOfBoundsTopLeft, bOutOfBoundsTopRight, bOutOfBoundsBottomRight, bOutOfBoundsBottomLeft;
+	float LuminanceTopLeft     = LDSReadFloat(GroupThreadId, uint2(-1, -1), bOutOfBoundsTopLeft);
+	float LuminanceTopRight    = LDSReadFloat(GroupThreadId, uint2( 1, -1), bOutOfBoundsTopRight);
+	float LuminanceBottomRight = LDSReadFloat(GroupThreadId, uint2( 1,  1), bOutOfBoundsBottomRight);
+	float LuminanceBottomLeft  = LDSReadFloat(GroupThreadId, uint2(-1,  1), bOutOfBoundsBottomLeft);
+
+#if DIM_SHARPENING_QUALITY > 1
+	GroupMemoryBarrierWithGroupSync();
+
+	if (GroupThreadIndex < GetTotalNumBoundarySamples())
+	{
+		int2 TopLeftPixelPosition = GroupId * GroupSize;
+		int2 BoundaryPixelPosRelativeToGroup = IndexToBoundaryPosition(GroupThreadIndex);
+		float2 BoundaryPixelGroupUV = BoundaryPixelPosRelativeToGroup / GroupSize;
+		int2 BoundaryPixelPosition = clamp(TopLeftPixelPosition + BoundaryPixelPosRelativeToGroup, int2(0,0), Output_ViewportSize-1);
+
+		float2 BoundaryPixelInputUV = PixelPositionToInputUV(BoundaryPixelPosition, GroupId, BoundaryPixelGroupUV); // approximates distortion
+		//float4 BoundaryPixel = Upsample(BoundaryPixelInputUV, EyeIndex, UpscaleMethod);
+		float4 BoundaryPixel = SampleSceneColorRGBA(BoundaryPixelInputUV); // use simple bilinear instead of full upscaler for speed
+		float3 BoundaryPixelYCoCg = RGBToYCoCg(BoundaryPixel.rgb);
+		float BoundaryPixelLuminance = BoundaryPixelYCoCg.x;
+		LDSWriteFloat(GroupThreadId, BoundaryPixelLuminance);
+	}
+
+	GroupMemoryBarrierWithGroupSync();
+
+	if (bOutOfBoundsLeft)        { LuminanceLeft   = LDSReadFloatAtIndex(BoundaryPositionToIndex(GroupThreadId + uint2(-1,  0))); }
+	else if (bOutOfBoundsRight)  { LuminanceRight  = LDSReadFloatAtIndex(BoundaryPositionToIndex(GroupThreadId + uint2( 1,  0))); }
+	if (bOutOfBoundsTop)         { LuminanceTop    = LDSReadFloatAtIndex(BoundaryPositionToIndex(GroupThreadId + uint2( 0, -1))); }
+	else if (bOutOfBoundsBottom) { LuminanceBottom = LDSReadFloatAtIndex(BoundaryPositionToIndex(GroupThreadId + uint2( 0,  1))); }
+#endif
+
+	float DeltaLuminance = ( -(LuminanceLeft + LuminanceTop + LuminanceRight + LuminanceBottom) + (LocalLuminance * 4.0) );
+	LocalLuminance += DeltaLuminance * Sharpening;
+
+	ColorYCoCg.x = LocalLuminance;
+	OutColor.rgb = YCoCgToRGB(ColorYCoCg);
+#endif
+
+	OutColor = ProcessAlpha(OutColor);
+
+	if (any(PixelPosition >= Output_ViewportSize))
+	{
+		return;
+	}
+
+	RWOutputTexture[BufferPixelPosition] = OutColor;
+}
+
+#else // !COMPUTESHADER
+
+// vertex shader entry point
+void MainVS(
+	in uint VertexId : SV_VertexID,
+	in FStereoVSInput StereoInput,
+	out noperspective float4 OutTexCoord : TEXCOORD0,
+	out FStereoVSOutput StereoOutput,
+	out float4 OutPosition : SV_POSITION)
+{
+	StereoSetupVS(StereoInput, StereoOutput);
+	
+	float2 TexCoord = float2(VertexId % (GridDimensions.x + 1), VertexId / (GridDimensions.x + 1)) / float2(GridDimensions);
+
+	// still in 0..1 range
+	float4 Position = float4(TexCoord.x, TexCoord.y, 0, 1);
+
+	// distort pos
+	Position.xy = ApplyLensDistortionOnViewportUV(DistortingDisplacementTexture, DistortingDisplacementSampler, TexCoord);
+
+	DrawRectangle(Position, TexCoord, OutPosition, OutTexCoord.xy);
+	OutTexCoord.zw =  OutPosition.xy;
+}
+
+void MainPS(noperspective float4 UVAndScreenPos : TEXCOORD0, FStereoPSInput StereoInput, float4 SvPosition : SV_POSITION, out float4 OutColor : SV_Target0)
+{
+	StereoSetupPS(StereoInput);
+	const uint EyeIndex = GetEyeIndex(StereoInput);
+
+	OutColor = Upsample(UVAndScreenPos.xy, EyeIndex, DIM_METHOD);
+	OutColor = ProcessAlpha(OutColor);
+}
+
+#endif // !COMPUTESHADER
\ No newline at end of file
diff --git a/Engine/Source/Runtime/Renderer/Private/PostProcess/PostProcessUpscale.cpp b/Engine/Source/Runtime/Renderer/Private/PostProcess/PostProcessUpscale.cpp
--- a/Engine/Source/Runtime/Renderer/Private/PostProcess/PostProcessUpscale.cpp
+++ b/Engine/Source/Runtime/Renderer/Private/PostProcess/PostProcessUpscale.cpp
@@ -30,6 +30,26 @@ TAutoConsoleVariable<int32> CVarUpscaleQuality(
 	TEXT(" 5: 36-tap Gaussian-filtered unsharp mask (very expensive, but good for extreme upsampling).\n"),
 	ECVF_Scalability | ECVF_RenderThreadSafe);
 
+TAutoConsoleVariable<int32> CVarUpscaleComputeEnabled(
+	TEXT("r.Upscale.ComputeEnabled"),
+	1,
+	TEXT("Allow running the upscaler as a compute pass. \n"),
+	ECVF_Scalability | ECVF_RenderThreadSafe);
+
+TAutoConsoleVariable<int32> CVarUpscaleSharpeningQuality(
+	TEXT("r.Upscale.SharpeningQuality"),
+	1,
+	TEXT("0: off\n")
+	TEXT("1: cheaper\n")
+	TEXT("2: higher quality"),
+	ECVF_Scalability | ECVF_RenderThreadSafe);
+
+TAutoConsoleVariable<float> CVarUpscaleSharpening(
+	TEXT("r.Upscale.Sharpening"),
+	0.0f,
+	TEXT("Increase to get more sharpening on the final upscale. Requires ComputeEnabled.\n"),
+	ECVF_Scalability | ECVF_RenderThreadSafe);
+
 } //! namespace
 
 BEGIN_SHADER_PARAMETER_STRUCT(FUpscaleParameters, )
@@ -38,6 +58,8 @@ BEGIN_SHADER_PARAMETER_STRUCT(FUpscaleParameters, )
 	SHADER_PARAMETER_STRUCT(FScreenPassTextureViewportParameters, Output)
 	SHADER_PARAMETER_RDG_TEXTURE(Texture2D, DistortingDisplacementTexture)
 	SHADER_PARAMETER_SAMPLER(SamplerState, DistortingDisplacementSampler)
+	SHADER_PARAMETER_RDG_TEXTURE(Texture2D, UndistortingDisplacementTexture)
+	SHADER_PARAMETER_SAMPLER(SamplerState, UndistortingDisplacementSampler)
 	SHADER_PARAMETER(FIntPoint, GridDimensions)
 	SHADER_PARAMETER(uint32, bInvertAlpha)
 	SHADER_PARAMETER_RDG_TEXTURE(Texture2D, SceneColorTexture)
@@ -46,6 +68,11 @@ BEGIN_SHADER_PARAMETER_STRUCT(FUpscaleParameters, )
 	SHADER_PARAMETER_RDG_TEXTURE(Texture2DArray, PointSceneColorTextureArray)
 	SHADER_PARAMETER_SAMPLER(SamplerState, PointSceneColorSampler)
 	SHADER_PARAMETER(float, UpscaleSoftness)
+	SHADER_PARAMETER(float, Sharpening)
+END_SHADER_PARAMETER_STRUCT()
+
+BEGIN_SHADER_PARAMETER_STRUCT(FUpscaleRasterParameters, )
+	SHADER_PARAMETER_STRUCT_INCLUDE(FUpscaleParameters, UpscaleParameters)
 	RENDER_TARGET_BINDING_SLOTS()
 END_SHADER_PARAMETER_STRUCT()
 
@@ -54,10 +81,10 @@ class FUpscalePS : public FGlobalShader
 public:
 	DECLARE_GLOBAL_SHADER(FUpscalePS);
 	SHADER_USE_PARAMETER_STRUCT(FUpscalePS, FGlobalShader);
-	using FParameters = FUpscaleParameters;
+	using FParameters = FUpscaleRasterParameters;
 
 	class FAlphaChannelDim : SHADER_PERMUTATION_BOOL("DIM_ALPHA_CHANNEL");
-	class FMethodDimension : SHADER_PERMUTATION_ENUM_CLASS("METHOD", EUpscaleMethod);
+	class FMethodDimension : SHADER_PERMUTATION_ENUM_CLASS("DIM_METHOD", EUpscaleMethod);
 	using FPermutationDomain = TShaderPermutationDomain<FAlphaChannelDim, FMethodDimension>;
 
 	static bool ShouldCompilePermutation(const FGlobalShaderPermutationParameters& Parameters)
@@ -65,6 +92,11 @@ public:
 		const FPermutationDomain PermutationVector(Parameters.PermutationId);
 		const EUpscaleMethod UpscaleMethod = PermutationVector.Get<FMethodDimension>();
 
+		if (UpscaleMethod == EUpscaleMethod::None)
+		{
+			return false;
+		}
+
 		// Always allow point and bilinear upscale. (Provides upscaling for mobile emulation)
 		if (UpscaleMethod == EUpscaleMethod::Nearest || UpscaleMethod == EUpscaleMethod::Bilinear)
 		{
@@ -83,7 +115,7 @@ public:
 	DECLARE_GLOBAL_SHADER(FUpscaleVS);
 	// FDrawRectangleParameters is filled by DrawScreenPass.
 	SHADER_USE_PARAMETER_STRUCT_WITH_LEGACY_BASE(FUpscaleVS, FScreenPassVS);
-	using FParameters = FUpscaleParameters;
+	using FParameters = FUpscaleRasterParameters;
 
 	static void ModifyCompilationEnvironment(const FGlobalShaderPermutationParameters& Parameters, FShaderCompilerEnvironment& OutEnvironment)
 	{
@@ -93,6 +125,109 @@ public:
 
 IMPLEMENT_GLOBAL_SHADER(FUpscaleVS, "/Engine/Private/PostProcessUpscale.usf", "MainVS", SF_Vertex);
 
+const int32 GUpscaleComputeTileSizeX = 16;
+const int32 GUpscaleComputeTileSizeY = 16;
+
+enum class ESharpeningQuality : uint8
+{
+	Off,
+	Low,
+	High,
+	MAX
+};
+
+enum class EMethodSet : uint8
+{
+	None,
+	Simple,
+	Complex,
+	MAX
+};
+
+EMethodSet GetMethodSet(EUpscaleMethod Method)
+{
+	switch (Method)
+	{
+		case EUpscaleMethod::Nearest:
+		case EUpscaleMethod::Bilinear:
+		case EUpscaleMethod::SmoothStep:
+			return EMethodSet::Simple;
+		case EUpscaleMethod::Directional:
+		case EUpscaleMethod::CatmullRom:
+		case EUpscaleMethod::Lanczos:
+		case EUpscaleMethod::Gaussian:
+			return EMethodSet::Complex;
+		default:
+			return EMethodSet::MAX;
+	};
+}
+
+class FUpscaleCS : public FGlobalShader
+{
+public:
+	DECLARE_GLOBAL_SHADER(FUpscaleCS);
+	SHADER_USE_PARAMETER_STRUCT(FUpscaleCS, FGlobalShader);
+	
+	BEGIN_SHADER_PARAMETER_STRUCT(FParameters, )
+		SHADER_PARAMETER_STRUCT_INCLUDE(FUpscaleParameters, UpscaleParameters)
+		SHADER_PARAMETER(uint32, UpscaleMethod)
+		SHADER_PARAMETER_RDG_TEXTURE_UAV(RWTexture<float4>, RWOutputTexture)
+	END_SHADER_PARAMETER_STRUCT()
+
+	class FMethodDim : SHADER_PERMUTATION_ENUM_CLASS("DIM_METHOD", EUpscaleMethod);
+	class FMethodSetDim : SHADER_PERMUTATION_ENUM_CLASS("DIM_METHODSET", EMethodSet);
+	class FAlphaChannelDim : SHADER_PERMUTATION_BOOL("DIM_ALPHA_CHANNEL");
+	class FLensDistortionDim : SHADER_PERMUTATION_BOOL("DIM_LENS_DISTORTION");
+	class FSharpeningQualityDim : SHADER_PERMUTATION_ENUM_CLASS("DIM_SHARPENING_QUALITY", ESharpeningQuality);
+	using FPermutationDomain = TShaderPermutationDomain<FMethodDim, FMethodSetDim, FAlphaChannelDim, FLensDistortionDim, FSharpeningQualityDim>;
+
+	static bool ShouldCompilePermutation(const FGlobalShaderPermutationParameters& Parameters)
+	{
+		const FPermutationDomain PermutationVector(Parameters.PermutationId);
+
+		const EUpscaleMethod UpscaleMethod = PermutationVector.Get<FMethodDim>();
+		const EMethodSet UpscaleMethodSet = PermutationVector.Get<FMethodSetDim>();
+
+		if (UpscaleMethodSet != EMethodSet::None && UpscaleMethod != EUpscaleMethod::None)
+		{
+			return false;
+		}
+		else if (UpscaleMethodSet == EMethodSet::None && UpscaleMethod == EUpscaleMethod::None)
+		{
+			return false;
+		}
+
+		if (!IsFeatureLevelSupported(Parameters.Platform, ERHIFeatureLevel::SM5))
+		{
+			return false;
+		}
+
+		return IsFeatureLevelSupported(Parameters.Platform, ERHIFeatureLevel::SM5);
+	}
+
+	static FPermutationDomain RemapPermutation(FPermutationDomain PermutationVector)
+	{
+		const EUpscaleMethod UpscaleMethod = PermutationVector.Get<FMethodDim>();
+
+		// Special path for SmoothStep, as that is used in secondary upscaler
+		if (UpscaleMethod != EUpscaleMethod::SmoothStep)
+		{
+			PermutationVector.Set<FMethodDim>(EUpscaleMethod::None);
+			PermutationVector.Set<FMethodSetDim>(GetMethodSet(UpscaleMethod));
+		}
+
+		return PermutationVector;
+	}
+
+	static void ModifyCompilationEnvironment(const FGlobalShaderPermutationParameters& Parameters, FShaderCompilerEnvironment& OutEnvironment)
+	{
+		OutEnvironment.SetDefine(TEXT("THREADGROUP_SIZEX"), GUpscaleComputeTileSizeX);
+		OutEnvironment.SetDefine(TEXT("THREADGROUP_SIZEY"), GUpscaleComputeTileSizeY);
+	}
+};
+
+IMPLEMENT_GLOBAL_SHADER(FUpscaleCS, "/Engine/Private/PostProcessUpscale.usf", "MainCS", SF_Compute);
+
 EUpscaleMethod GetUpscaleMethod()
 {
 	const int32 Value = CVarUpscaleQuality.GetValueOnRenderThread();
@@ -109,6 +244,7 @@ FScreenPassTexture ISpatialUpscaler::AddDefaultUpscalePass(
 	FLensDistortionLUT LensDistortionLUT)
 {
 	check(Inputs.SceneColor.IsValid());
+	check(Method != EUpscaleMethod::None);
 	check(Method != EUpscaleMethod::MAX);
 	check(Inputs.Stage != EUpscaleStage::MAX);
 
@@ -145,135 +281,198 @@ FScreenPassTexture ISpatialUpscaler::AddDefaultUpscalePass(
 	const FScreenPassTextureViewport OutputViewport(Output);
 
 	const bool bApplyLensDistortion = LensDistortionLUT.IsEnabled();
+
+	bool bIsUpscaleToOutput = Inputs.Stage == EUpscaleStage::PrimaryToOutput || Inputs.Stage == EUpscaleStage::SecondaryToOutput;
+	float Sharpening = bIsUpscaleToOutput ? CVarUpscaleSharpening.GetValueOnRenderThread() : 0.0f;
 	
-	FUpscaleParameters* PassParameters = GraphBuilder.AllocParameters<FUpscaleParameters>();
-	PassParameters->RenderTargets[0] = Output.GetRenderTargetBinding();
-	PassParameters->Input = GetScreenPassTextureViewportParameters(InputViewport);
-	PassParameters->Output = GetScreenPassTextureViewportParameters(OutputViewport);
-	PassParameters->DistortingDisplacementTexture = LensDistortionLUT.DistortingDisplacementTexture;
-	PassParameters->DistortingDisplacementSampler = TStaticSamplerState<SF_Bilinear, AM_Clamp, AM_Clamp, AM_Clamp>::GetRHI();
-	PassParameters->GridDimensions = LensDistortionLUT.DistortionGridDimensions;
-	PassParameters->bInvertAlpha = View.Family->EngineShowFlags.AlphaInvert;
-	PassParameters->SceneColorTexture = Inputs.SceneColor.Texture;
-	PassParameters->SceneColorSampler = TStaticSamplerState<SF_Bilinear, AM_Border, AM_Border, AM_Border>::GetRHI();
-	PassParameters->PointSceneColorTexture = Inputs.SceneColor.Texture;
-	PassParameters->PointSceneColorTextureArray = Inputs.SceneColor.Texture;
-	PassParameters->PointSceneColorSampler = TStaticSamplerState<SF_Point, AM_Border, AM_Border, AM_Border>::GetRHI();
-	PassParameters->UpscaleSoftness = FMath::Clamp(CVarUpscaleSoftness.GetValueOnRenderThread(), 0.0f, 1.0f);
-	PassParameters->View = View.GetShaderParameters();
-
-	FUpscalePS::FPermutationDomain PixelPermutationVector;
-	PixelPermutationVector.Set<FUpscalePS::FAlphaChannelDim>(IsPostProcessingWithAlphaChannelSupported());
-	PixelPermutationVector.Set<FUpscalePS::FMethodDimension>(Method);
-	TShaderMapRef<FUpscalePS> PixelShader(View.ShaderMap, PixelPermutationVector);
+	FUpscaleParameters PassParameters;
+	PassParameters.Input = GetScreenPassTextureViewportParameters(InputViewport);
+	PassParameters.Output = GetScreenPassTextureViewportParameters(OutputViewport);
+	PassParameters.DistortingDisplacementTexture = LensDistortionLUT.DistortingDisplacementTexture;
+	PassParameters.DistortingDisplacementSampler = TStaticSamplerState<SF_Bilinear, AM_Clamp, AM_Clamp, AM_Clamp>::GetRHI();
+	PassParameters.UndistortingDisplacementTexture = LensDistortionLUT.UndistortingDisplacementTexture;
+	PassParameters.UndistortingDisplacementSampler = TStaticSamplerState<SF_Bilinear, AM_Clamp, AM_Clamp, AM_Clamp>::GetRHI();
+	PassParameters.GridDimensions = LensDistortionLUT.DistortionGridDimensions;
+	PassParameters.bInvertAlpha = View.Family->EngineShowFlags.AlphaInvert;
+	PassParameters.SceneColorTexture = Inputs.SceneColor.Texture;
+	PassParameters.SceneColorSampler = TStaticSamplerState<SF_Bilinear, AM_Border, AM_Border, AM_Border>::GetRHI();
+	PassParameters.PointSceneColorTexture = Inputs.SceneColor.Texture;
+	PassParameters.PointSceneColorTextureArray = Inputs.SceneColor.Texture;
+	PassParameters.PointSceneColorSampler = TStaticSamplerState<SF_Point, AM_Border, AM_Border, AM_Border>::GetRHI();
+	PassParameters.UpscaleSoftness = FMath::Clamp(CVarUpscaleSoftness.GetValueOnRenderThread(), 0.0f, 1.0f);
+	PassParameters.Sharpening = Sharpening;
+	PassParameters.View = View.GetShaderParameters();
 
 	const TCHAR* const StageNames[] = { TEXT("PrimaryToSecondary"), TEXT("PrimaryToOutput"), TEXT("SecondaryToOutput") };
 	static_assert(UE_ARRAY_COUNT(StageNames) == static_cast<uint32>(EUpscaleStage::MAX), "StageNames does not match EUpscaleStage");
 	const TCHAR* StageName = StageNames[static_cast<uint32>(Inputs.Stage)];
 
-	GraphBuilder.AddPass(
-		RDG_EVENT_NAME("Upscale(%s Method=%d%s%s) %dx%d -> %dx%d",
-			StageName,
-			int32(Method),
-			PixelPermutationVector.Get<FUpscalePS::FAlphaChannelDim>() ? TEXT(" Alpha") : TEXT(""),
-			bApplyLensDistortion ? TEXT(" LensDistortion") : TEXT(""),
-			Inputs.SceneColor.ViewRect.Width(), Inputs.SceneColor.ViewRect.Height(),
-			Output.ViewRect.Width(), Output.ViewRect.Height()),
-		PassParameters,
-		ERDGPassFlags::Raster,
-		[&View, bApplyLensDistortion, PixelShader, PassParameters, InputViewport, OutputViewport](FRDGAsyncTask, FRHICommandList& RHICmdList)
+	bool bUseCompute = bIsUpscaleToOutput && CVarUpscaleComputeEnabled.GetValueOnRenderThread() > 0 && View.bUseComputePasses;
+	if (bUseCompute)
 	{
-		RHICmdList.SetViewport(OutputViewport.Rect.Min.X, OutputViewport.Rect.Min.Y, 0.0f, OutputViewport.Rect.Max.X, OutputViewport.Rect.Max.Y, 1.0f);
-
-		TShaderRef<FShader> VertexShader;
-		if (bApplyLensDistortion)
+		FRDGTextureRef ComputeRenderTarget{};
+		bool bOutputSupportsUAV = (Output.Texture->Desc.Flags & TexCreate_UAV) == TexCreate_UAV;
+		if (bOutputSupportsUAV)
 		{
-			TShaderMapRef<FUpscaleVS> TypedVertexShader(View.ShaderMap);
-			SetScreenPassPipelineState(RHICmdList, FScreenPassPipelineState(TypedVertexShader, PixelShader));
-			SetShaderParameters(RHICmdList, TypedVertexShader, TypedVertexShader.GetVertexShader(), *PassParameters);
-			VertexShader = TypedVertexShader;
+			ComputeRenderTarget = Output.Texture;
 		}
 		else
 		{
-			TShaderMapRef<FScreenPassVS> TypedVertexShader(View.ShaderMap);
-			SetScreenPassPipelineState(RHICmdList, FScreenPassPipelineState(TypedVertexShader, PixelShader));
-			VertexShader = TypedVertexShader;
+			const FRDGTextureDesc Desc(FRDGTextureDesc::Create2D(Output.Texture->Desc.Extent, Output.Texture->Desc.Format, FClearValueBinding::None, ETextureCreateFlags::UAV | ETextureCreateFlags::ShaderResource));
+			ComputeRenderTarget = GraphBuilder.CreateTexture(Desc, TEXT("SecondaryUpscalerOutput"));
 		}
-		check(VertexShader.IsValid());
 
-		SetShaderParameters(RHICmdList, PixelShader, PixelShader.GetPixelShader(), *PassParameters);
-
-		if (bApplyLensDistortion)
+		FUpscaleCS::FParameters* ComputePassParameters = GraphBuilder.AllocParameters<FUpscaleCS::FParameters>();
+
+		ComputePassParameters->UpscaleParameters = PassParameters;
+		ComputePassParameters->RWOutputTexture = GraphBuilder.CreateUAV(ComputeRenderTarget);
+		ComputePassParameters->UpscaleMethod = (uint32)Method;
+
+		ESharpeningQuality SharpeningQuality = (ESharpeningQuality)FMath::Clamp(CVarUpscaleSharpeningQuality.GetValueOnRenderThread(), 0, (int32)ESharpeningQuality::MAX); 
+		SharpeningQuality = Sharpening != 0.0f ? SharpeningQuality : ESharpeningQuality::Off;
+
+		FUpscaleCS::FPermutationDomain PermutationVector;
+		PermutationVector.Set<FUpscaleCS::FAlphaChannelDim>(IsPostProcessingWithAlphaChannelSupported());
+		PermutationVector.Set<FUpscaleCS::FLensDistortionDim>(bApplyLensDistortion);
+		PermutationVector.Set<FUpscaleCS::FMethodDim>(Method);
+		PermutationVector.Set<FUpscaleCS::FMethodSetDim>(EMethodSet::None); // RemapPermutation sets this to the correct one.
+		PermutationVector.Set<FUpscaleCS::FSharpeningQualityDim>(SharpeningQuality);
+		TShaderMapRef<FUpscaleCS> ComputeShader(View.ShaderMap, PermutationVector);
+
+		FComputeShaderUtils::AddPass(
+			GraphBuilder,
+			RDG_EVENT_NAME("Upscale(CS %s Method=%d%s%s) %dx%d -> %dx%d",
+				StageName,
+				int32(Method),
+				PermutationVector.Get<FUpscaleCS::FAlphaChannelDim>() ? TEXT(" Alpha") : TEXT(""),
+				bApplyLensDistortion ? TEXT(" LensDistortion") : TEXT(""),
+				Inputs.SceneColor.ViewRect.Width(), Inputs.SceneColor.ViewRect.Height(),
+				Output.ViewRect.Width(), Output.ViewRect.Height()),
+			ComputeShader,
+			ComputePassParameters,
+			FComputeShaderUtils::GetGroupCount(OutputViewport.Rect.Size(), FIntPoint(GUpscaleComputeTileSizeX, GUpscaleComputeTileSizeY)));
+		
+		FRHICopyTextureInfo CopyInfo {};
+		CopyInfo.SourcePosition = FIntVector{ OutputViewport.Rect.Min.X, OutputViewport.Rect.Min.Y, 0 };
+		CopyInfo.DestPosition = CopyInfo.SourcePosition;
+		CopyInfo.Size = FIntVector{ OutputViewport.Rect.Size().X, OutputViewport.Rect.Size().Y, 1 };
+		AddCopyTexturePass(GraphBuilder, ComputeRenderTarget, Output.Texture, CopyInfo);
+	}
+	else
+	{
+		FUpscaleRasterParameters* RasterPassParameters = GraphBuilder.AllocParameters<FUpscaleRasterParameters>();
+		RasterPassParameters->UpscaleParameters = PassParameters;
+		RasterPassParameters->RenderTargets[0] = Output.GetRenderTargetBinding();
+
+		FUpscalePS::FPermutationDomain PixelPermutationVector;
+		PixelPermutationVector.Set<FUpscalePS::FAlphaChannelDim>(IsPostProcessingWithAlphaChannelSupported());
+		PixelPermutationVector.Set<FUpscalePS::FMethodDimension>(Method);
+		TShaderMapRef<FUpscalePS> PixelShader(View.ShaderMap, PixelPermutationVector);
+
+		GraphBuilder.AddPass(
+			RDG_EVENT_NAME("Upscale(%s Method=%d%s%s) %dx%d -> %dx%d",
+				StageName,
+				int32(Method),
+				PixelPermutationVector.Get<FUpscalePS::FAlphaChannelDim>() ? TEXT(" Alpha") : TEXT(""),
+				bApplyLensDistortion ? TEXT(" LensDistortion") : TEXT(""),
+				Inputs.SceneColor.ViewRect.Width(), Inputs.SceneColor.ViewRect.Height(),
+				Output.ViewRect.Width(), Output.ViewRect.Height()),
+			RasterPassParameters,
+			ERDGPassFlags::Raster,
+			[&View, bApplyLensDistortion, PixelShader, RasterPassParameters, InputViewport, OutputViewport](FRDGAsyncTask, FRHICommandList& RHICmdList)
 		{
-			TArray<uint32> IndexBuffer;
+			RHICmdList.SetViewport(OutputViewport.Rect.Min.X, OutputViewport.Rect.Min.Y, 0.0f, OutputViewport.Rect.Max.X, OutputViewport.Rect.Max.Y, 1.0f);
 
-			const uint32 Width = PassParameters->GridDimensions.X;
-			const uint32 Height = PassParameters->GridDimensions.Y;
-			const uint32 NumVertices = (Width + 1) * (Height + 1);
-			const uint32 NumTriangles = Width * Height * 2;
-			const uint32 NumIndices = NumTriangles * 3;
-			
-			IndexBuffer.AddUninitialized(NumIndices);
+			TShaderRef<FShader> VertexShader;
+			if (bApplyLensDistortion)
+			{
+				TShaderMapRef<FUpscaleVS> TypedVertexShader(View.ShaderMap);
+				SetScreenPassPipelineState(RHICmdList, FScreenPassPipelineState(TypedVertexShader, PixelShader));
+				SetShaderParameters(RHICmdList, TypedVertexShader, TypedVertexShader.GetVertexShader(), *RasterPassParameters);
+				VertexShader = TypedVertexShader;
+			}
+			else
+			{
+				TShaderMapRef<FScreenPassVS> TypedVertexShader(View.ShaderMap);
+				SetScreenPassPipelineState(RHICmdList, FScreenPassPipelineState(TypedVertexShader, PixelShader));
+				VertexShader = TypedVertexShader;
+			}
+			check(VertexShader.IsValid());
 
-			uint32* Out = (uint32*)IndexBuffer.GetData();
+			SetShaderParameters(RHICmdList, PixelShader, PixelShader.GetPixelShader(), *RasterPassParameters);
 
-			for(uint32 y = 0; y < Height; ++y)
+			if (bApplyLensDistortion)
 			{
-				for(uint32 x = 0; x < Width; ++x)
+				TArray<uint32> IndexBuffer;
+
+				const uint32 Width = RasterPassParameters->UpscaleParameters.GridDimensions.X;
+				const uint32 Height = RasterPassParameters->UpscaleParameters.GridDimensions.Y;
+				const uint32 NumVertices = (Width + 1) * (Height + 1);
+				const uint32 NumTriangles = Width * Height * 2;
+				const uint32 NumIndices = NumTriangles * 3;
+			
+				IndexBuffer.AddUninitialized(NumIndices);
+
+				uint32* Out = (uint32*)IndexBuffer.GetData();
+
+				for(uint32 y = 0; y < Height; ++y)
 				{
-					// left top to bottom right in reading order
-					uint32 Index00 = x  + y * (Width + 1);
-					uint32 Index10 = Index00 + 1;
-					uint32 Index01 = Index00 + (Width + 1);
-					uint32 Index11 = Index01 + 1;
+					for(uint32 x = 0; x < Width; ++x)
+					{
+						// left top to bottom right in reading order
+						uint32 Index00 = x  + y * (Width + 1);
+						uint32 Index10 = Index00 + 1;
+						uint32 Index01 = Index00 + (Width + 1);
+						uint32 Index11 = Index01 + 1;
 					
-					// triangle A
-					*Out++ = Index00; *Out++ = Index01; *Out++ = Index10;
+						// triangle A
+						*Out++ = Index00; *Out++ = Index01; *Out++ = Index10;
 
-					// triangle B
-					*Out++ = Index11; *Out++ = Index10; *Out++ = Index01;
+						// triangle B
+						*Out++ = Index11; *Out++ = Index10; *Out++ = Index01;
+					}
 				}
-			}
 
-			// Create index buffer. Fill buffer with initial data upon creation
-			FBufferRHIRef IndexBufferRHI = UE::RHIResourceUtils::CreateIndexBufferFromArray(RHICmdList, TEXT("LensDistortionIndexBuffer"), EBufferUsageFlags::Static, MakeConstArrayView(IndexBuffer));
+				// Create index buffer. Fill buffer with initial data upon creation
+				FBufferRHIRef IndexBufferRHI = UE::RHIResourceUtils::CreateIndexBufferFromArray(RHICmdList, TEXT("LensDistortionIndexBuffer"), EBufferUsageFlags::Static, MakeConstArrayView(IndexBuffer));
 
 			
-			FRHIBatchedShaderParameters& BatchedParameters = RHICmdList.GetScratchShaderParameters();
-			UE::Renderer::PostProcess::SetDrawRectangleParameters(BatchedParameters, VertexShader.GetShader(),
-				0, 0, OutputViewport.Rect.Width(), OutputViewport.Rect.Height(),
-				InputViewport.Rect.Min.X, InputViewport.Rect.Min.Y, InputViewport.Rect.Width(), InputViewport.Rect.Height(),
-				OutputViewport.Rect.Size(),
-				InputViewport.Extent);
-			RHICmdList.SetBatchedShaderParameters(VertexShader.GetVertexShader(), BatchedParameters);
-
-			// no vertex buffer needed as we compute it in VS
-			RHICmdList.SetStreamSource(0, nullptr, 0);
-
-			RHICmdList.DrawIndexedPrimitive(
-				IndexBufferRHI,
-				/*BaseVertexIndex=*/ 0,
-				/*MinIndex=*/ 0,
-				/*NumVertices=*/ NumVertices,
-				/*StartIndex=*/ 0,
-				/*NumPrimitives=*/ NumTriangles,
-				/*NumInstances=*/ 1
-				);
-		}
-		else
-		{
-			DrawRectangle(
-				RHICmdList,
-				// Output Rect (RHI viewport relative).
-				0, 0, OutputViewport.Rect.Width(), OutputViewport.Rect.Height(),
-				// Input Rect
-				InputViewport.Rect.Min.X, InputViewport.Rect.Min.Y, InputViewport.Rect.Width(), InputViewport.Rect.Height(),
-				OutputViewport.Rect.Size(),
-				InputViewport.Extent,
-				VertexShader,
-				EDRF_UseTriangleOptimization);
-		}
-	});
+				FRHIBatchedShaderParameters& BatchedParameters = RHICmdList.GetScratchShaderParameters();
+				UE::Renderer::PostProcess::SetDrawRectangleParameters(BatchedParameters, VertexShader.GetShader(),
+					0, 0, OutputViewport.Rect.Width(), OutputViewport.Rect.Height(),
+					InputViewport.Rect.Min.X, InputViewport.Rect.Min.Y, InputViewport.Rect.Width(), InputViewport.Rect.Height(),
+					OutputViewport.Rect.Size(),
+					InputViewport.Extent);
+				RHICmdList.SetBatchedShaderParameters(VertexShader.GetVertexShader(), BatchedParameters);
+
+				// no vertex buffer needed as we compute it in VS
+				RHICmdList.SetStreamSource(0, nullptr, 0);
+
+				RHICmdList.DrawIndexedPrimitive(
+					IndexBufferRHI,
+					/*BaseVertexIndex=*/ 0,
+					/*MinIndex=*/ 0,
+					/*NumVertices=*/ NumVertices,
+					/*StartIndex=*/ 0,
+					/*NumPrimitives=*/ NumTriangles,
+					/*NumInstances=*/ 1
+					);
+			}
+			else
+			{
+				DrawRectangle(
+					RHICmdList,
+					// Output Rect (RHI viewport relative).
+					0, 0, OutputViewport.Rect.Width(), OutputViewport.Rect.Height(),
+					// Input Rect
+					InputViewport.Rect.Min.X, InputViewport.Rect.Min.Y, InputViewport.Rect.Width(), InputViewport.Rect.Height(),
+					OutputViewport.Rect.Size(),
+					InputViewport.Extent,
+					VertexShader,
+					EDRF_UseTriangleOptimization);
+			}
+		});
+	}
 
 	return MoveTemp(Output);
 }
diff --git a/Engine/Source/Runtime/Renderer/Private/PostProcess/PostProcessUpscale.h b/Engine/Source/Runtime/Renderer/Private/PostProcess/PostProcessUpscale.h
--- a/Engine/Source/Runtime/Renderer/Private/PostProcess/PostProcessUpscale.h
+++ b/Engine/Source/Runtime/Renderer/Private/PostProcess/PostProcessUpscale.h
@@ -7,6 +7,7 @@
 
 enum class EUpscaleMethod : uint8
 {
+	None,
 	Nearest,
 	Bilinear,
 	Directional,
