Added support for switching all GPU-Scene buffers to reserved resources (r.GPUScene.UseReservedResources) as resizing can cause sustained load in cases where the size is oscillating between two powers of two as well as spikes when larger size changes happen. To avoid that, also track peak GPU-scene allocation between successive frames.

diff --git a/Engine/Source/Runtime/Renderer/Private/GPUScene.cpp b/Engine/Source/Runtime/Renderer/Private/GPUScene.cpp
--- a/Engine/Source/Runtime/Renderer/Private/GPUScene.cpp
+++ b/Engine/Source/Runtime/Renderer/Private/GPUScene.cpp
@@ -143,6 +143,12 @@ static TAutoConsoleVariable<int32> CVarGPUSceneInstanceDataTileSizeLog2(
 	ECVF_RenderThreadSafe
 );
 
+static TAutoConsoleVariable<bool> CVarGPUSceneUseReservedResources(
+	TEXT("r.GPUScene.UseReservedResources"),
+	false,
+	TEXT("Turning this on makes all GPU-Scene buffers default to using reserved resource allocations."),
+	ECVF_ReadOnly);
+
 LLM_DECLARE_TAG_API(GPUScene, RENDERER_API);
 DECLARE_LLM_MEMORY_STAT(TEXT("GPUScene"), STAT_GPUSceneLLM, STATGROUP_LLMFULL);
 DECLARE_LLM_MEMORY_STAT(TEXT("GPUScene"), STAT_GPUSceneSummaryLLM, STATGROUP_LLM);
@@ -733,6 +739,13 @@ void FGPUScene::EndRender()
 	ensure(bInBeginEndBlock);
 	ensure(CurrentDynamicContext != nullptr);
 
+	// record the previous max primitives to avoid oscillating allocations.
+	PrevDynamicPrimitivesMaxSize = DynamicPrimitivesOffset;
+
+	// record number of dynamic allocs for next frame to avoid oscillations
+	PrevDynamicPrimitivePayloadMaxSize = InstancePayloadDataAllocator.GetMaxSize();
+	PrevDynamicPrimitiveInstanceMaxSize = InstanceSceneDataAllocator.GetMaxSize();
+
 	// Pop all dynamic primitives off the stack
 	DynamicPrimitivesOffset = Scene.GetMaxPersistentPrimitiveIndex();
 	bInBeginEndBlock = false;
@@ -1016,11 +1029,27 @@ FGPUScene::FRegisteredBuffers FGPUScene::UpdateBufferAllocations(FRDGBuilder& Gr
 
 	FRegisteredBuffers BufferState;
 
-	const uint32 SizeReserve = FMath::RoundUpToPowerOfTwo(FMath::Max(DynamicPrimitivesOffset, InitialBufferSize));
-	BufferState.PrimitiveBuffer = ResizeStructuredBufferIfNeeded(GraphBuilder, PrimitiveBuffer, SizeReserve * sizeof(FPrimitiveSceneShaderData::Data), TEXT("GPUScene.PrimitiveData"));
+	bool bUseReservedResources = CVarGPUSceneUseReservedResources.GetValueOnRenderThread() && GRHIGlobals.ReservedResources.Supported;
+
+	auto AllocOrResizeBufferIfNeeded = [bUseReservedResources](FRDGBuilder& GraphBuilder, TRefCountPtr<FRDGPooledBuffer>& Buffer, uint32 NumBytes, const TCHAR* DebugName)
+	{
+		if (bUseReservedResources && !Buffer.IsValid())
+		{
+			constexpr uint64 ReservedAllocationSize = uint64(2048) << 20;
+			const uint32 BytesPerElement = 16;
+			FRDGBufferDesc BufferDescReserved = FRDGBufferDesc::CreateStructuredDesc(BytesPerElement, uint32(ReservedAllocationSize / uint64(BytesPerElement)));
+			BufferDescReserved.Usage |= EBufferUsageFlags::ReservedResource;
+			Buffer = AllocatePooledBuffer(BufferDescReserved, DebugName);
+		}
+		return ResizeStructuredBufferIfNeeded(GraphBuilder, Buffer, NumBytes, DebugName);
+	};
+
 
+	const uint32 SizeReserve = FMath::RoundUpToPowerOfTwo(FMath::Max(PrevDynamicPrimitivesMaxSize, FMath::Max(DynamicPrimitivesOffset, InitialBufferSize)));
+	BufferState.PrimitiveBuffer = AllocOrResizeBufferIfNeeded(GraphBuilder, PrimitiveBuffer, SizeReserve * sizeof(FPrimitiveSceneShaderData::Data), TEXT("GPUScene.PrimitiveData"));
+	
 	// clamped instance count / high watermark (any overflowing will not render!)
-	const uint32 ClampedMaxInstanceId = FMath::Min(MAX_INSTANCE_ID, uint32(FMath::Max(InstanceSceneDataAllocator.GetMaxSize(), InitialBufferSize)));
+	const uint32 ClampedMaxInstanceId = FMath::Min(MAX_INSTANCE_ID, uint32(FMath::Max(PrevDynamicPrimitiveInstanceMaxSize, FMath::Max(InstanceSceneDataAllocator.GetMaxSize(), InitialBufferSize))));
 
 	// emulate the old functioning, deriving the stride from the instance count
 	if (!bUseTiledInstanceDataLayout)
@@ -1048,12 +1077,12 @@ FGPUScene::FRegisteredBuffers FGPUScene::UpdateBufferAllocations(FRDGBuilder& Gr
 		BufferState.InstanceSceneDataBuffer = ResizeStructuredBufferIfNeeded(GraphBuilder, InstanceSceneDataBuffer, InstanceSceneDataSizeReserve, TEXT("GPUScene.InstanceSceneData"));
 	}
 
-	const uint32 PayloadFloat4Count = FMath::Max(InstancePayloadDataAllocator.GetMaxSize(), InitialBufferSize);
+	const uint32 PayloadFloat4Count = FMath::Max(PrevDynamicPrimitivePayloadMaxSize, FMath::Max(InstancePayloadDataAllocator.GetMaxSize(), InitialBufferSize));
 	const uint32 InstancePayloadDataSizeReserve = FMath::RoundUpToPowerOfTwo(PayloadFloat4Count * sizeof(FVector4f));
-	BufferState.InstancePayloadDataBuffer = ResizeStructuredBufferIfNeeded(GraphBuilder, InstancePayloadDataBuffer, InstancePayloadDataSizeReserve, TEXT("GPUScene.InstancePayloadData"));
+	BufferState.InstancePayloadDataBuffer = AllocOrResizeBufferIfNeeded(GraphBuilder, InstancePayloadDataBuffer, InstancePayloadDataSizeReserve, TEXT("GPUScene.InstancePayloadData"));
 
 	const uint32 LightMapDataBufferSize = FMath::RoundUpToPowerOfTwo(FMath::Max(LightmapDataAllocator.GetMaxSize(), InitialBufferSize));
-	BufferState.LightmapDataBuffer = ResizeStructuredBufferIfNeeded(GraphBuilder, LightmapDataBuffer, LightMapDataBufferSize * sizeof(FLightmapSceneShaderData::Data), TEXT("GPUScene.LightmapData"));
+	BufferState.LightmapDataBuffer = AllocOrResizeBufferIfNeeded(GraphBuilder, LightmapDataBuffer, LightMapDataBufferSize * sizeof(FLightmapSceneShaderData::Data), TEXT("GPUScene.LightmapData"));
 
 	BufferState.LightDataBuffer = GraphBuilder.RegisterExternalBuffer(LightDataBuffer.GetPooledBuffer());
 
diff --git a/Engine/Source/Runtime/Renderer/Private/GPUScene.h b/Engine/Source/Runtime/Renderer/Private/GPUScene.h
--- a/Engine/Source/Runtime/Renderer/Private/GPUScene.h
+++ b/Engine/Source/Runtime/Renderer/Private/GPUScene.h
@@ -458,6 +458,10 @@ private:
 	uint32 SceneFrameNumber = 0xFFFFFFFF;
 
 	int32 DynamicPrimitivesOffset = 0;
+	
+	int32 PrevDynamicPrimitivesMaxSize = 0;
+	int32 PrevDynamicPrimitivePayloadMaxSize = 0;
+	int32 PrevDynamicPrimitiveInstanceMaxSize = 0;
 
 	bool bIsEnabled = false;
 	bool bInBeginEndBlock = false;
