Fix shader parameter type
#rnx

diff --git a/Engine/Source/Runtime/Renderer/Private/PostProcess/PostProcessUpscale.cpp b/Engine/Source/Runtime/Renderer/Private/PostProcess/PostProcessUpscale.cpp
--- a/Engine/Source/Runtime/Renderer/Private/PostProcess/PostProcessUpscale.cpp
+++ b/Engine/Source/Runtime/Renderer/Private/PostProcess/PostProcessUpscale.cpp
@@ -171,7 +171,7 @@ public:
 	BEGIN_SHADER_PARAMETER_STRUCT(FParameters, )
 		SHADER_PARAMETER_STRUCT_INCLUDE(FUpscaleParameters, UpscaleParameters)
 		SHADER_PARAMETER(uint32, UpscaleMethod)
-		SHADER_PARAMETER_RDG_TEXTURE_UAV(RWTexture<float4>, RWOutputTexture)
+		SHADER_PARAMETER_RDG_TEXTURE_UAV(RWTexture2D<float4>, RWOutputTexture)
 	END_SHADER_PARAMETER_STRUCT()
 
 	class FMethodDim : SHADER_PERMUTATION_ENUM_CLASS("DIM_METHOD", EUpscaleMethod);
