Heterogeneous Volumes: Create dummy UAV to overcome binding issues when rendering on DX11.

diff --git a/Engine/Source/Runtime/Renderer/Private/HeterogeneousVolumes/HeterogeneousVolumes.cpp b/Engine/Source/Runtime/Renderer/Private/HeterogeneousVolumes/HeterogeneousVolumes.cpp
--- a/Engine/Source/Runtime/Renderer/Private/HeterogeneousVolumes/HeterogeneousVolumes.cpp
+++ b/Engine/Source/Runtime/Renderer/Private/HeterogeneousVolumes/HeterogeneousVolumes.cpp
@@ -1078,14 +1078,17 @@ void FDeferredShadingSceneRenderer::RenderHeterogeneousVolumes(
 			AddClearUAVPass(GraphBuilder, GraphBuilder.CreateUAV(HeterogeneousVolumeHoldout), FLinearColor::Black);
 		}
 
-		// Satisfy binding UAV binding requirement
-		HeterogeneousVolumeBeerShadowMap = HeterogeneousVolumeRadiance;
 		if (HeterogeneousVolumes::GetTranslucencyCompositingMode() == HeterogeneousVolumes::ETranslucencyCompositingMode::BeerShadowMap)
 		{
 			Desc.Format = PF_FloatRGBA;
 			HeterogeneousVolumeBeerShadowMap = GraphBuilder.CreateTexture(Desc, TEXT("HeterogeneousVolume.BeerShadowMap"));
 			AddClearUAVPass(GraphBuilder, GraphBuilder.CreateUAV(HeterogeneousVolumeBeerShadowMap), FVector4d(HeterogeneousVolumes::GetMaxTraceDistance(), 0.0, 1.0, 1.0));
 		}
+		else
+		{
+			Desc.Extent = FIntPoint(1);
+			HeterogeneousVolumeBeerShadowMap = GraphBuilder.CreateTexture(Desc, TEXT("HeterogeneousVolume.BeerShadowMap"));
+		}
 #if 0
 		// Create velocity texture if it doesn't already exist
 		if (!HasBeenProduced(HeterogeneousVolumeVelocity))
