Remove deprecated code from FXRenderingUtils.

diff --git a/Engine/Source/Runtime/Renderer/Private/FXRenderingUtils.cpp b/Engine/Source/Runtime/Renderer/Private/FXRenderingUtils.cpp
--- a/Engine/Source/Runtime/Renderer/Private/FXRenderingUtils.cpp
+++ b/Engine/Source/Runtime/Renderer/Private/FXRenderingUtils.cpp
@@ -261,11 +261,6 @@ FRHIRayTracingScene* UE::FXRenderingUtils::RayTracing::GetRayTracingScene(const
 	return nullptr;
 }
 
-FRHIShaderResourceView* UE::FXRenderingUtils::RayTracing::GetRayTracingSceneView(const FSceneInterface* InScene)
-{
-	return GetRayTracingSceneView(FRHICommandListImmediate::Get(), InScene);
-}
-
 FRHIShaderResourceView* UE::FXRenderingUtils::RayTracing::GetRayTracingSceneView(FRHICommandListBase& RHICmdList, const FSceneInterface* InScene)
 {
 	if (const FScene* Scene = InScene->GetRenderScene())
diff --git a/Engine/Source/Runtime/Renderer/Public/FXRenderingUtils.h b/Engine/Source/Runtime/Renderer/Public/FXRenderingUtils.h
--- a/Engine/Source/Runtime/Renderer/Public/FXRenderingUtils.h
+++ b/Engine/Source/Runtime/Renderer/Public/FXRenderingUtils.h
@@ -75,30 +75,6 @@ namespace UE::FXRenderingUtils
 
 		RENDERER_API FShaderBindingTableRHIRef CreateShaderBindingTable(FRHICommandListBase& RHICmdList, const FSceneInterface* Scene,  uint32 LocalBindingDataSize);		
 		RENDERER_API TConstArrayView<FRayTracingShaderBindingData> GetVisibleRayTracingShaderBindings(const FSceneView& View);
-
-		UE_DEPRECATED(5.3, "GetRayTracingSceneView now requires a command list.")
-		RENDERER_API FRHIShaderResourceView* GetRayTracingSceneView(const FSceneInterface* Scene);
 	}
 #endif
 }
-
-/**
- * This class exposes methods required by FX rendering that must access rendering internals.
- */ 
-class FFXRenderingUtils
-{
-public:
-	FFXRenderingUtils() = delete;
-	FFXRenderingUtils(const FFXRenderingUtils&) = delete;
-	FFXRenderingUtils& operator=(const FFXRenderingUtils&) = delete;
-
-	/** Utility to determine if a material might render before the FXSystem's PostRenderOpaque is called for the view family */
-	UE_DEPRECATED(5.3, "Use UE::FXRenderingUtils::CanMaterialRenderBeforeFXPostOpaque")
-	static bool CanMaterialRenderBeforeFXPostOpaque(
-		const FSceneViewFamily& ViewFamily,
-		const FPrimitiveSceneProxy& SceneProxy,
-		const FMaterial& Material)
-	{
-		return UE::FXRenderingUtils::CanMaterialRenderBeforeFXPostOpaque(ViewFamily, SceneProxy, Material);
-	}
-};
