Also bind FLumenHardwareRayTracingUniformBufferParameters for FLightFunctionRayTracingShader

diff --git a/Engine/Source/Runtime/Renderer/Private/RayTracing/RayTracingLighting.cpp b/Engine/Source/Runtime/Renderer/Private/RayTracing/RayTracingLighting.cpp
--- a/Engine/Source/Runtime/Renderer/Private/RayTracing/RayTracingLighting.cpp
+++ b/Engine/Source/Runtime/Renderer/Private/RayTracing/RayTracingLighting.cpp
@@ -615,6 +615,9 @@ public:
 
 		// LightFunctions can use primitive data, set identity so we do not crash on a missing binding
 		ShaderBindings.Add(GetUniformBufferParameter<FPrimitiveUniformShaderParameters>(), GIdentityPrimitiveUniformBuffer);
+
+		// bind Lumen Hardware RT buffer
+		ShaderBindings.Add(GetUniformBufferParameter<FLumenHardwareRayTracingUniformBufferParameters>(), View.LumenHardwareRayTracingUniformBuffer);
 	}
 
 	static void ModifyCompilationEnvironment(const FMaterialShaderPermutationParameters& Parameters, FShaderCompilerEnvironment& OutEnvironment)
