Removed some deprecated rendering code.

diff --git a/Engine/Source/Runtime/Engine/Public/PrimitiveSceneDesc.h b/Engine/Source/Runtime/Engine/Public/PrimitiveSceneDesc.h
--- a/Engine/Source/Runtime/Engine/Public/PrimitiveSceneDesc.h
+++ b/Engine/Source/Runtime/Engine/Public/PrimitiveSceneDesc.h
@@ -105,9 +105,6 @@ struct FPrimitiveSceneDesc
 	const IPrimitiveComponent* GetPrimitiveComponentInterface() const { return PrimitiveComponentInterface; }
 
 	UPackage* GetOutermost() const { return PrimitiveUObject->GetOutermost(); }
-
-	UE_DEPRECATED(5.5, "GetRegistrationSerialNumber is no longer used.")
-	int32 GetRegistrationSerialNumber() const { return -1; }
 };
 
 struct FInstancedStaticMeshSceneDesc
diff --git a/Engine/Source/Runtime/Engine/Public/PrimitiveSceneInfoData.h b/Engine/Source/Runtime/Engine/Public/PrimitiveSceneInfoData.h
--- a/Engine/Source/Runtime/Engine/Public/PrimitiveSceneInfoData.h
+++ b/Engine/Source/Runtime/Engine/Public/PrimitiveSceneInfoData.h
@@ -79,9 +79,6 @@ protected:
 
 public:
 
-	UE_DEPRECATED(5.5, "GetNextRegistrationSerialNumber is no longer used")
-	static int32 GetNextRegistrationSerialNumber() { return -1; }
-
 	FPrimitiveSceneInfoData()
 		: bAlwaysVisible(false)
 	{
diff --git a/Engine/Source/Runtime/Engine/Public/RayTracingSkinnedGeometry.h b/Engine/Source/Runtime/Engine/Public/RayTracingSkinnedGeometry.h
--- a/Engine/Source/Runtime/Engine/Public/RayTracingSkinnedGeometry.h
+++ b/Engine/Source/Runtime/Engine/Public/RayTracingSkinnedGeometry.h
@@ -32,9 +32,6 @@ class FRayTracingSkinnedGeometryUpdateQueue
 	{
 		return 0;
 	}
-	
-	UE_DEPRECATED(5.5, "This function has been deprecated. Use the version that takes GraphBuilder as input.")
-	void Commit(FRHICommandList& RHICmdList, FRHIBuffer* ScratchBuffer) {}
 
 	UE_DEPRECATED(5.6, "This function has been deprecated because logic is merged with the dynamic geometry manager and is not needed anymore.")
 	void Commit(FRDGBuilder& GraphBuilder, ERDGPassFlags ComputePassFlags = ERDGPassFlags::Compute) {}
diff --git a/Engine/Source/Runtime/RenderCore/Public/RayTracingGeometry.h b/Engine/Source/Runtime/RenderCore/Public/RayTracingGeometry.h
--- a/Engine/Source/Runtime/RenderCore/Public/RayTracingGeometry.h
+++ b/Engine/Source/Runtime/RenderCore/Public/RayTracingGeometry.h
@@ -33,17 +33,6 @@ public:
 	FRayTracingGeometryOfflineDataHeader RawDataHeader;
 	TResourceArray<uint8> RawData;
 
-	//disable deprecation warnings for default constructors
-	PRAGMA_DISABLE_DEPRECATION_WARNINGS
-	FRayTracingGeometry() = default;
-	virtual ~FRayTracingGeometry() = default;
-
-	FRayTracingGeometry(const FRayTracingGeometry&) = delete;
-	FRayTracingGeometry& operator=(const FRayTracingGeometry&) = delete;
-	FRayTracingGeometry(FRayTracingGeometry&&) = delete;
-	FRayTracingGeometry& operator=(FRayTracingGeometry&&) = delete;
-	PRAGMA_ENABLE_DEPRECATION_WARNINGS
-
 #if RHI_RAYTRACING
 
 	/** When set to NonSharedVertexBuffers, then shared vertex buffers are not used  */
@@ -64,14 +53,9 @@ public:
 
 	FRayTracingGeometryInitializer Initializer;
 
-	UE_DEPRECATED(5.5, "Use GetRHI() instead.")
-	FRayTracingGeometryRHIRef RayTracingGeometryRHI;
-
 	FRHIRayTracingGeometry* GetRHI() const
 	{
-		PRAGMA_DISABLE_DEPRECATION_WARNINGS
 		return RayTracingGeometryRHI;
-		PRAGMA_ENABLE_DEPRECATION_WARNINGS
 	}
 
 	RayTracing::FGeometryGroupHandle GroupHandle = INDEX_NONE;
@@ -181,6 +165,8 @@ public:
 protected:
 	RENDERCORE_API void RemoveBuildRequest();
 
+	FRayTracingGeometryRHIRef RayTracingGeometryRHI;
+
 	friend class FRayTracingGeometryManager;
 	EGeometryStateFlags GeometryState = EGeometryStateFlags::Invalid;
 	int32 RayTracingBuildRequestIndex = INDEX_NONE;
diff --git a/Engine/Source/Runtime/Renderer/Private/PrimitiveSceneInfo.cpp b/Engine/Source/Runtime/Renderer/Private/PrimitiveSceneInfo.cpp
--- a/Engine/Source/Runtime/Renderer/Private/PrimitiveSceneInfo.cpp
+++ b/Engine/Source/Runtime/Renderer/Private/PrimitiveSceneInfo.cpp
@@ -473,31 +473,6 @@ FRayTracingGeometry* FPrimitiveSceneInfo::GetValidStaticRayTracingGeometry(int8&
 
 	return nullptr;
 }
-
-FRHIRayTracingGeometry* FPrimitiveSceneInfo::GetStaticRayTracingGeometryInstance(int LodLevel) const
-{
-	if (StaticRayTracingGeometries.Num() > LodLevel)
-	{
-		// TODO: Select different LOD, when build is still pending for this LOD?
-		if (StaticRayTracingGeometries[LodLevel]->HasPendingBuildRequest())
-		{
-			StaticRayTracingGeometries[LodLevel]->BoostBuildPriority();
-			return nullptr;
-		}
-		else if (StaticRayTracingGeometries[LodLevel]->IsValid() && !StaticRayTracingGeometries[LodLevel]->IsEvicted())
-		{
-			return StaticRayTracingGeometries[LodLevel]->GetRHI();
-		}
-		else
-		{
-			return nullptr;
-		}
-	}
-	else
-	{
-		return nullptr;
-	}
-}
 #endif
 
 
@@ -1636,20 +1611,6 @@ void FPrimitiveSceneInfo::AllocateGPUSceneInstances(FScene* Scene, const TArrayV
 	}
 }
 
-void FPrimitiveSceneInfo::ReallocateGPUSceneInstances(FScene* Scene, const TArrayView<FPrimitiveSceneInfo*>& SceneInfos)
-{
-	SCOPED_NAMED_EVENT(ReallocateGPUSceneInstances, FColor::Emerald);
-
-	// Free each scene info.
-	for (FPrimitiveSceneInfo* SceneInfo : SceneInfos)
-	{
-		SceneInfo->FreeGPUSceneInstances();
-	}
-
-	// Allocate them all.
-	AllocateGPUSceneInstances(Scene, SceneInfos);
-}
-
 void FPrimitiveSceneInfo::FreeGPUSceneInstances()
 {
 	if (!Scene->GPUScene.IsEnabled())
diff --git a/Engine/Source/Runtime/Renderer/Public/PrimitiveSceneInfo.h b/Engine/Source/Runtime/Renderer/Public/PrimitiveSceneInfo.h
--- a/Engine/Source/Runtime/Renderer/Public/PrimitiveSceneInfo.h
+++ b/Engine/Source/Runtime/Renderer/Public/PrimitiveSceneInfo.h
@@ -278,13 +278,6 @@ public:
 	 */
 	FPrimitiveComponentId PrimitiveComponentId;
 
-	/**
-	 * Number assigned to this component when it was registered with the world.
-	 * This will only ever be updated if the object is re-registered.
-	 */
-	UE_DEPRECATED(5.5, "RegistrationSerialNumber is no longer used")
-	int32 RegistrationSerialNumber = -1;
-
 	/** 
 	 * The root attachment component id for use with lighting, if valid.
 	 * If the root id is not valid, this is a parent primitive.
@@ -385,9 +378,6 @@ public:
 	static void AllocateGPUSceneInstances(FScene* Scene, const TArrayView<FPrimitiveSceneInfo*>& SceneInfos);
 	void FreeGPUSceneInstances();
 
-	UE_DEPRECATED(5.3, "ReallocateGPUSceneInstances has been deprecated.")
-	static void ReallocateGPUSceneInstances(FScene* Scene, const TArrayView<FPrimitiveSceneInfo*>& SceneInfos);
-
 	/** return true if we need to call LazyUpdateForRendering */
 	FORCEINLINE bool NeedsIndirectLightingCacheBufferUpdate()
 	{
@@ -541,9 +531,6 @@ public:
 	RENDERER_API FRayTracingGeometry* GetStaticRayTracingGeometry(int8 LODIndex) const;
 	RENDERER_API FRayTracingGeometry* GetValidStaticRayTracingGeometry(int8& InOutLODIndex) const;
 
-	UE_DEPRECATED(5.4, "GetStaticRayTracingGeometryInstance has been deprecated. Use GetStaticRayTracingGeometry / GetValidStaticRayTracingGeometry instead.")
-	RENDERER_API FRHIRayTracingGeometry* GetStaticRayTracingGeometryInstance(int LodLevel) const;
-
 	int GetStaticRayTracingGeometryNum() const { return StaticRayTracingGeometries.Num(); }
 
 	const FRayTracingGeometry* GetCachedRayTracingGeometry() const { return CachedRayTracingGeometry; }
@@ -563,27 +550,6 @@ public:
 	UPrimitiveComponent* GetComponentForDebugOnly() const;
 	IPrimitiveComponent* GetComponentInterfaceForDebugOnly() const;
 
-	UE_DEPRECATED(5.3, "NeedsUpdateStaticMeshes has been deprecated.")
-	bool NeedsUpdateStaticMeshes() { return false; }
-
-	UE_DEPRECATED(5.3, "NeedsUniformBufferUpdate has been deprecated.")
-	bool NeedsUniformBufferUpdate() const { return false; }
-	
-	UE_DEPRECATED(5.3, "UpdateUniformBuffer has been deprecated.")
-	void UpdateUniformBuffer(FRHICommandListImmediate& RHICmdList) {}
-	
-	UE_DEPRECATED(5.3, "ConditionalUpdateUniformBuffer has been deprecated.")
-	void ConditionalUpdateUniformBuffer(FRHICommandListImmediate& RHICmdList) {}
-	
-	UE_DEPRECATED(5.3, "BeginDeferredUpdateStaticMeshes has been deprecated. Use RequestStaticMeshUpdate instead.")
-	void BeginDeferredUpdateStaticMeshes() { RequestStaticMeshUpdate(); }
-	
-	UE_DEPRECATED(5.3, "BeginDeferredUpdateStaticMeshesWithoutVisibilityCheck has been deprecated. Use RequestStaticMeshUpdate instead.")
-	void BeginDeferredUpdateStaticMeshesWithoutVisibilityCheck() { RequestStaticMeshUpdate(); }
-	
-	UE_DEPRECATED(5.3, "SetNeedsUniformBufferUpdate is deprecated. Use RequestUniformBufferUpdate instead.")
-	void SetNeedsUniformBufferUpdate(bool bInNeedsUniformBufferUpdate) { RequestUniformBufferUpdate(); }
-
 	bool HasInstanceDataBuffers() const { return InstanceSceneDataBuffersInternal != nullptr; }
 	
 	/**
@@ -786,10 +752,6 @@ private:
 	/** Removes cached ray tracing representations for all meshes. */
 	void RemoveCachedRayTracingPrimitives();
 
-	/** Updates cached world bounds in CachedRayTracingInstance */
-	UE_DEPRECATED(5.4, "UpdateCachedRayTracingInstanceWorldBounds has been deprecated.")
-	void UpdateCachedRayTracingInstanceWorldBounds(const FMatrix& NewPrimitiveLocalToWorld) {};
-
 	/** Updates cached ray tracing instances. Utility closely mirrors CacheRayTracingPrimitives(..) */
 	static void UpdateCachedRayTracingInstance(FPrimitiveSceneInfo* SceneInfo, const FRayTracingInstance& RayTracingInstance, const ERayTracingPrimitiveFlags Flags);
 #endif
diff --git a/Engine/Source/Runtime/Renderer/Public/RayTracingInstance.h b/Engine/Source/Runtime/Renderer/Public/RayTracingInstance.h
--- a/Engine/Source/Runtime/Renderer/Public/RayTracingInstance.h
+++ b/Engine/Source/Runtime/Renderer/Public/RayTracingInstance.h
@@ -117,10 +117,6 @@ struct FRayTracingInstance
 		}
 	}
 
-	/** When instance transforms are only available in GPU, this SRV holds them. */
-	UE_DEPRECATED(5.5, "InstanceGPUTransformsSRV has been deprecated. GPU Scene should be used instead.")
-	FShaderResourceViewRHIRef InstanceGPUTransformsSRV;
-
 	//disable deprecation warnings for default constructors
 	PRAGMA_DISABLE_DEPRECATION_WARNINGS
 	FRayTracingInstance() = default;
diff --git a/Engine/Source/Runtime/Renderer/Public/RayTracingInstanceBufferUtil.h b/Engine/Source/Runtime/Renderer/Public/RayTracingInstanceBufferUtil.h
--- a/Engine/Source/Runtime/Renderer/Public/RayTracingInstanceBufferUtil.h
+++ b/Engine/Source/Runtime/Renderer/Public/RayTracingInstanceBufferUtil.h
@@ -59,9 +59,6 @@ struct UE_DEPRECATED(5.6, "Use FRayTracingInstanceBufferBuilder instead.") FRayT
 	// Any referenced geometry is kept alive while the scene is alive.
 	TArray<FRHIRayTracingGeometry*> ReferencedGeometries;
 
-	UE_DEPRECATED(5.5, "No longer used as 'inline SBT' is managed using FShaderBindingTable.")
-	TArray<FRHIRayTracingGeometry*> PerInstanceGeometries;
-
 	FRayTracingSceneInitializationData() = default;
 	FRayTracingSceneInitializationData(const FRayTracingSceneInitializationData&) = default;
 	FRayTracingSceneInitializationData& operator=(const FRayTracingSceneInitializationData&) = default;
