Make sure ViewLocation is also set on FSceneCaptureViewInfo (used by InstancedStaticMesh for culling calculations)

diff --git a/Engine/Source/Runtime/Renderer/Private/SceneCaptureRendering.cpp b/Engine/Source/Runtime/Renderer/Private/SceneCaptureRendering.cpp
--- a/Engine/Source/Runtime/Renderer/Private/SceneCaptureRendering.cpp
+++ b/Engine/Source/Runtime/Renderer/Private/SceneCaptureRendering.cpp
@@ -840,9 +840,10 @@ static FSceneRenderer* CreateSceneRendererForSceneCapture(
 	const AActor* ViewActor,
 	int32 CubemapFaceIndex = INDEX_NONE)
 {
-	FSceneCaptureViewInfo SceneCaptureViewInfo;
+	FSceneCaptureViewInfo SceneCaptureViewInfo{};
 	SceneCaptureViewInfo.ViewRotationMatrix = ViewRotationMatrix;
 	SceneCaptureViewInfo.ViewOrigin = ViewLocation;
+	SceneCaptureViewInfo.ViewLocation = ViewLocation;
 	SceneCaptureViewInfo.ProjectionMatrix = ProjectionMatrix;
 	SceneCaptureViewInfo.StereoPass = EStereoscopicPass::eSSP_FULL;
 	SceneCaptureViewInfo.StereoViewIndex = INDEX_NONE;
