Unify virtual texture mip color cvars for streaming and runtime virtual textures.
"r.vt.rvt.mipcolors" is deprecated and runtime virtual textures now respect "r.vt.borders" which streaming virtual textures have always used.
"r.vt.borders" is no longer on/off and instead the value sets the debug border texel width.
Added "r.vt.borders.mip" to isolate debug mip colors to a single mip level. This can be useful for debugging specific issues with mip sampling, such as showing where mip 0 is always used.
Added flushing when these cvars are set so that the effect can be seen immediately.

diff --git a/Engine/Shaders/Private/VirtualTextureMaterial.usf b/Engine/Shaders/Private/VirtualTextureMaterial.usf
--- a/Engine/Shaders/Private/VirtualTextureMaterial.usf
+++ b/Engine/Shaders/Private/VirtualTextureMaterial.usf
@@ -114,8 +114,20 @@ void FPixelShaderInOut_MainPS(
 #endif
 
 	// Apply any fixed color.
-	BaseColor =  lerp(BaseColor, RuntimeVirtualTexturePassParameters.DebugParams.xyz, RuntimeVirtualTexturePassParameters.DebugParams.w);
-	Displacement = lerp(Displacement, RuntimeVirtualTexturePassParameters.DebugParams.x, RuntimeVirtualTexturePassParameters.DebugParams.w);
+	if (RuntimeVirtualTexturePassParameters.FixedColor.w > 0)
+	{
+		BaseColor = lerp(BaseColor, RuntimeVirtualTexturePassParameters.FixedColor.xyz, RuntimeVirtualTexturePassParameters.FixedColor.w);
+		Displacement = lerp(Displacement, RuntimeVirtualTexturePassParameters.FixedColor.x, RuntimeVirtualTexturePassParameters.FixedColor.w);
+	}
+
+	// Apply mip debug color.
+	if (RuntimeVirtualTexturePassParameters.MipColorAndSize.w < 1)
+	{
+		if (any(abs(ScreenPosition.xy) > RuntimeVirtualTexturePassParameters.MipColorAndSize.w))
+		{
+			BaseColor = RuntimeVirtualTexturePassParameters.MipColorAndSize.xyz;
+		}
+	}
 
 #if OUT_BASECOLOR
 	Out.MRT[0] = float4(BaseColor, 1.f) * Opacity;
diff --git a/Engine/Source/Runtime/Engine/Private/VT/VirtualTextureTranscodeCache.cpp b/Engine/Source/Runtime/Engine/Private/VT/VirtualTextureTranscodeCache.cpp
--- a/Engine/Source/Runtime/Engine/Private/VT/VirtualTextureTranscodeCache.cpp
+++ b/Engine/Source/Runtime/Engine/Private/VT/VirtualTextureTranscodeCache.cpp
@@ -40,16 +40,25 @@ static FAutoConsoleCommand CVarVTShowDecodeErrors(
 		})
 );
 
-namespace TextureBorderGenerator
-{
-	static int32 Enabled = 0;
-	static FAutoConsoleVariableRef CVarEnableDebugBorders(
-		TEXT("r.VT.Borders"),
-		Enabled,
-		TEXT("If > 0, debug borders will enabled\n")
-		/*,ECVF_ReadOnly*/
-	);
-}
+/** Fill virtual texture borders (or full tiles) with mip colors. */
+static TAutoConsoleVariable<int32> CVarVirtualTextureBorders(
+	TEXT("r.VT.Borders"),
+	0,
+	TEXT("Sets the virtual texture tile border size, in texels, that is filled with debug mip colors.\n")
+	TEXT("Size may be rounded up for block compressed textures."),
+	FConsoleVariableDelegate::CreateLambda([](IConsoleVariable* InVariable) { GetRendererModule().FlushVirtualTextureCache(); }),
+	ECVF_RenderThreadSafe
+);
+
+/** Select a single mip level to apply mip colors to. This can be useful for highlighting mip 0 sampling issues. */
+static TAutoConsoleVariable<int32> CVarVirtualTextureBordersMip(
+	TEXT("r.VT.Borders.Mip"),
+	-1,
+	TEXT("Set a single mip level that r.VT.Borders will apply to.\n")
+	TEXT("-1 (default) applies to all mip levels."),
+	FConsoleVariableDelegate::CreateLambda([](IConsoleVariable* InVariable) { GetRendererModule().FlushVirtualTextureCache(); }),
+	ECVF_RenderThreadSafe
+);
 
 struct FTranscodeTask
 {
@@ -221,9 +230,14 @@ struct FTranscodeTask
 			}
 
 			// Bake debug borders directly into the tile pixels
-			if (TextureBorderGenerator::Enabled)
+			const int32 DebugBorderSize = CVarVirtualTextureBorders.GetValueOnAnyThread();
+			if (DebugBorderSize > 0)
 			{
-				BakeDebugInfo(StagingBufferForLayer, TilePixelSize, TilePixelSize, TileBorderPixelSize + 4, LayerFormat, vLevel);
+				const int32 DebugBorderMip = CVarVirtualTextureBordersMip.GetValueOnAnyThread();
+				if (DebugBorderMip < 0 || DebugBorderMip == vLevel)
+				{
+					BakeDebugInfo(StagingBufferForLayer, TilePixelSize, TilePixelSize, TileBorderPixelSize + DebugBorderSize, LayerFormat, vLevel);
+				}
 			}
 		}
 
diff --git a/Engine/Source/Runtime/Renderer/Private/VT/RuntimeVirtualTextureRender.cpp b/Engine/Source/Runtime/Renderer/Private/VT/RuntimeVirtualTextureRender.cpp
--- a/Engine/Source/Runtime/Renderer/Private/VT/RuntimeVirtualTextureRender.cpp
+++ b/Engine/Source/Runtime/Renderer/Private/VT/RuntimeVirtualTextureRender.cpp
@@ -34,12 +34,7 @@ DECLARE_DWORD_COUNTER_STAT(TEXT("Num pages rendered"), STAT_RenderedPages, STATG
 
 namespace RuntimeVirtualTexture
 {
-	static TAutoConsoleVariable<int32> CVarVTMipColors(
-		TEXT("r.VT.RVT.MipColors"),
-		0,
-		TEXT("Render mip colors to RVT BaseColor.\n 0 off. 1 shows all mip colors. 2 shows only mip 0."),
-		FConsoleVariableDelegate::CreateLambda([](IConsoleVariable* InVariable) { FGlobalComponentRecreateRenderStateContext Context; }),
-		ECVF_RenderThreadSafe);
+	static FAutoConsoleVariableDeprecated CVarVTMipColors_Deprecated(TEXT("r.VT.RVT.MipColors"), TEXT("r.VT.Borders"), TEXT("5.7"));
 
 	static TAutoConsoleVariable<int32> CVarVTHighQualityPerPixelHeight(
 		TEXT("r.VT.RVT.HighQualityPerPixelHeight"),
@@ -549,7 +544,8 @@ namespace RuntimeVirtualTexture
 	BEGIN_UNIFORM_BUFFER_STRUCT(FRuntimeVirtualTexturePassParameters, )
 		SHADER_PARAMETER(FVector4f, MipLevel)
 		SHADER_PARAMETER(FVector4f, CustomMaterialData)
-		SHADER_PARAMETER(FVector4f, DebugParams)
+		SHADER_PARAMETER(FVector4f, FixedColor)
+		SHADER_PARAMETER(FVector4f, MipColorAndSize)
 		SHADER_PARAMETER(FVector2f, PackHeight)
 	END_UNIFORM_BUFFER_STRUCT()
 
@@ -1352,6 +1348,7 @@ namespace RuntimeVirtualTexture
 		FRenderGraphSetup GraphSetup;
 		FRenderPageBatchDesc BatchDesc;
 		bool bAllowCachedMeshDrawCommands = true;
+		bool bAllowMipDebug = true;
 	};
 
 	/** Mesh processor for rendering static meshes to the virtual texture */
@@ -2116,24 +2113,40 @@ namespace RuntimeVirtualTexture
 	}
 
 	/** Get the debug color to use for a given mip level. Returns InDefaultColor if mip debugging is disabled. */
-	FLinearColor GetDebugMipLevelColor(uint32 InLevel, FLinearColor InDefaultColor)
+	FLinearColor GetDebugMipLevelColor(bool bInAllowMipDebug, uint32 InLevel, uint32 InTileAndBorderSize)
 	{
-		const int32 MipColorMode = CVarVTMipColors.GetValueOnRenderThread();
-		if (MipColorMode == 1)
+		if (!bInAllowMipDebug)
 		{
-			static const uint32 MipColors[] = {
-				0xC0FFFFFF, 0xC0FFFF00, 0xC000FFFF, 0xC000FF00, 0xC0FF00FF, 0xC0FF0000, 0xC00000FF,
-				0xC0808080, 0xC0808000, 0xC0008080, 0xC0008000, 0xC0800080, 0xC0800000, 0xC0000080 };
+			return FLinearColor::Black;
+		}
 
-			InLevel = FMath::Min<uint32>(InLevel, sizeof(MipColors) / sizeof(MipColors[0]) - 1);
-			return FLinearColor(FColor(MipColors[InLevel]));
+		static const auto CVarBorders = IConsoleManager::Get().FindTConsoleVariableDataInt(TEXT("r.VT.Borders"));
+		const int32 DebugBorderSize = CVarBorders ? CVarBorders->GetValueOnRenderThread() : 0;
+		if (DebugBorderSize <= 0)
+		{
+			return FLinearColor::Black;
 		}
-		else if (MipColorMode == 2 && InLevel == 0)
+
+		static const auto CVarBordersMip = IConsoleManager::Get().FindTConsoleVariableDataInt(TEXT("r.VT.Borders.Mip"));
+		const int32 DebugBorderMip = CVarBordersMip ? CVarBordersMip->GetValueOnRenderThread() : -1;
+		if (DebugBorderMip >= 0 && DebugBorderMip != InLevel)
 		{
-			return FLinearColor(1.f, 0.f, 1.f, 0.5f);
+			return FLinearColor::Black;
 		}
 
-		return InDefaultColor;
+		static const uint32 MipColors[] = {
+			0x00FFFFFF, 0x00FFFF00, 0x0000FFFF, 0x0000FF00, 0x00FF00FF, 0x00FF0000, 0x000000FF,
+			0x00808080, 0x00808000, 0x00008080, 0x00008000, 0x00800080, 0x00800000, 0x00000080 };
+
+		InLevel = FMath::Min<uint32>(InLevel, sizeof(MipColors) / sizeof(MipColors[0]) - 1);
+		FLinearColor Color = FLinearColor(FColor(MipColors[InLevel]));
+		
+		// Mip border size is stored in alpha channel.
+		// Need to cover requested border size plus the tile border. Assume tile border is the difference above some power of 2 tile size.
+		const uint32 TileBorderSize = (InTileAndBorderSize - (1 << FMath::FloorLog2(InTileAndBorderSize)));
+		Color.A = 1.0f - float(TileBorderSize + DebugBorderSize * 2) / float(InTileAndBorderSize);
+
+		return Color;
 	}
 
 	/** 
@@ -2218,7 +2231,8 @@ namespace RuntimeVirtualTexture
 			FRuntimeVirtualTexturePassParameters* RuntimeVirtualTexturePassParameters = GraphBuilder.AllocParameters<FRuntimeVirtualTexturePassParameters>();
 			RuntimeVirtualTexturePassParameters->MipLevel = MipLevelParameter;
 			RuntimeVirtualTexturePassParameters->CustomMaterialData = BatchDesc.CustomMaterialData;
-			RuntimeVirtualTexturePassParameters->DebugParams = GetDebugMipLevelColor(vLevel, BatchDesc.FixedColor);
+			RuntimeVirtualTexturePassParameters->FixedColor = BatchDesc.FixedColor;
+			RuntimeVirtualTexturePassParameters->MipColorAndSize = GetDebugMipLevelColor(BatchRenderContext.bAllowMipDebug, vLevel, TextureSize.X);
 			RuntimeVirtualTexturePassParameters->PackHeight = FVector2f(WorldHeightPackParameter);	// LWC_TODO: Precision loss
  			PassParameters->RuntimeVirtualTexturePassParameters = GraphBuilder.CreateUniformBuffer(RuntimeVirtualTexturePassParameters);
 
@@ -2423,6 +2437,8 @@ namespace RuntimeVirtualTexture
 		// Disable MDC caching for this standalone path because we can't guarantee that primitives 
 		// associated with the scene have been recreated (e.g. by World->SendAllEndOfFrameUpdates()).
 		Context.bAllowCachedMeshDrawCommands = false;
+		// Disable mip color debugging for standalone path used for baking streaming low mips.
+		Context.bAllowMipDebug = false;
 
 		RenderPageBatch(GraphBuilder, Context);
 		FinalizePageBatch(GraphBuilder, Context);
