Fixed: RHI validation warning around binding View UB for a occlusion shader in a mobile renderer

diff --git a/Engine/Source/Runtime/Renderer/Private/SceneOcclusion.h b/Engine/Source/Runtime/Renderer/Private/SceneOcclusion.h
--- a/Engine/Source/Runtime/Renderer/Private/SceneOcclusion.h
+++ b/Engine/Source/Runtime/Renderer/Private/SceneOcclusion.h
@@ -103,7 +103,12 @@ public:
 private:
 	void SetParametersInternal(FRHIBatchedShaderParameters& BatchedParameters, const FViewInfo& View, const FVector4f StencilingSpherePosAndScale[2])
 	{
-		FGlobalShader::SetParameters<FViewUniformShaderParameters>(BatchedParameters, View.ViewUniformBuffer);
+		// Mobile renderer sets View UB through static binding
+		if (!IsMobilePlatform(GetShaderPlatform()))
+		{
+			FGlobalShader::SetParameters<FViewUniformShaderParameters>(BatchedParameters, View.ViewUniformBuffer);
+		}
+				
 		if (!View.bShouldBindInstancedViewUB)
 		{
 			SetShaderValue(BatchedParameters, StencilingGeometryParameters, StencilingSpherePosAndScale[0]);
