Fixed: Enabling r.OcclusionFeedback.Enable crashes on desktop when running with feature level ES3.1

diff --git a/Engine/Source/Runtime/Renderer/Private/SceneOcclusion.cpp b/Engine/Source/Runtime/Renderer/Private/SceneOcclusion.cpp
--- a/Engine/Source/Runtime/Renderer/Private/SceneOcclusion.cpp
+++ b/Engine/Source/Runtime/Renderer/Private/SceneOcclusion.cpp
@@ -1403,7 +1403,8 @@ static void BeginOcclusionTests(
 				}
 			}
 
-			if (View.ViewState && View.ViewState->OcclusionFeedback.IsInitialized())
+			if (FeatureLevel == ERHIFeatureLevel::ES3_1 && 
+				View.ViewState && View.ViewState->OcclusionFeedback.IsInitialized())
 			{
 				SCOPED_DRAW_EVENT(RHICmdList, IndividualQueries);
 				View.ViewState->OcclusionFeedback.SubmitOcclusionDraws(RHICmdList, View);
diff --git a/Engine/Source/Runtime/Renderer/Private/SceneVisibility.cpp b/Engine/Source/Runtime/Renderer/Private/SceneVisibility.cpp
--- a/Engine/Source/Runtime/Renderer/Private/SceneVisibility.cpp
+++ b/Engine/Source/Runtime/Renderer/Private/SceneVisibility.cpp
@@ -2578,7 +2578,7 @@ FGPUOcclusionPacket::FGPUOcclusionPacket(FVisibilityViewPacket& InViewPacket, co
 	, PrimitiveProbablyVisibleTime(GEngine->PrimitiveProbablyVisibleTime)
 	, CurrentRealTime(View.Family->Time.GetRealTimeSeconds())
 	, NeverOcclusionTestDistanceSquared(GNeverOcclusionTestDistance* GNeverOcclusionTestDistance)
-	, bUseOcclusionFeedback(OcclusionFeedback.IsInitialized())
+	, bUseOcclusionFeedback(View.FeatureLevel == ERHIFeatureLevel::ES3_1 && OcclusionFeedback.IsInitialized())
 	, bNewlyConsideredBBoxExpandActive(GExpandNewlyOcclusionTestedBBoxesAmount > 0.0f && GFramesToExpandNewlyOcclusionTestedBBoxes > 0 && GFramesNotOcclusionTestedToExpandBBoxes > 0)
 {}
 
@@ -3061,7 +3061,8 @@ void FGPUOcclusion::Map(FRHICommandListImmediate& RHICmdListImmediate)
 {
 	SCOPED_NAMED_EVENT(MapOcclusionResults, FColor::Magenta);
 
-	if (ViewState.OcclusionFeedback.IsInitialized())
+	if (View.FeatureLevel == ERHIFeatureLevel::ES3_1 &&
+		ViewState.OcclusionFeedback.IsInitialized())
 	{
 		QUICK_SCOPE_CYCLE_COUNTER(STAT_OcclusionFeedback_ReadbackResults);
 		ViewState.OcclusionFeedback.ReadbackResults(RHICmdListImmediate);
