Stochastic Lighting - added missing pragma once causing compilation issues

diff --git a/Engine/Source/Runtime/Renderer/Private/StochasticLighting/StochasticLighting.h b/Engine/Source/Runtime/Renderer/Private/StochasticLighting/StochasticLighting.h
--- a/Engine/Source/Runtime/Renderer/Private/StochasticLighting/StochasticLighting.h
+++ b/Engine/Source/Runtime/Renderer/Private/StochasticLighting/StochasticLighting.h
@@ -1,5 +1,7 @@
 // Copyright Epic Games, Inc. All Rights Reserved.
 
+#pragma once
+
 #include "CoreTypes.h"
 #include "RenderGraphFwd.h"
 #include "Lumen/LumenFrontLayerTranslucency.h"
