Path Tracer: Fix crash when enabling decal normal projection

This cvar leads to the g-buffer being accessed which is not correct when the path tracer is active.

diff --git a/Engine/Source/Runtime/Renderer/Private/DeferredShadingRenderer.cpp b/Engine/Source/Runtime/Renderer/Private/DeferredShadingRenderer.cpp
--- a/Engine/Source/Runtime/Renderer/Private/DeferredShadingRenderer.cpp
+++ b/Engine/Source/Runtime/Renderer/Private/DeferredShadingRenderer.cpp
@@ -2833,7 +2833,7 @@ void FDeferredShadingSceneRenderer::Render(FRDGBuilder& GraphBuilder, const FSce
 			}
 		}
 
-		if (bUseGBuffer)
+		if (bRenderDeferredLighting)
 		{
 			// mark GBufferA for saving for next frame if it's needed
 			ExtractNormalsForNextFrameReprojection(GraphBuilder, SceneTextures, Views);
