Stochastic Lighting - added a dedicated GPU scope in order to move this out from the "unnacounted" category

diff --git a/Engine/Source/Runtime/Renderer/Private/StochasticLighting/StochasticLighting.cpp b/Engine/Source/Runtime/Renderer/Private/StochasticLighting/StochasticLighting.cpp
--- a/Engine/Source/Runtime/Renderer/Private/StochasticLighting/StochasticLighting.cpp
+++ b/Engine/Source/Runtime/Renderer/Private/StochasticLighting/StochasticLighting.cpp
@@ -145,12 +145,17 @@ void StochasticLighting::TileClassificationMark(
 	}
 }
 
+DECLARE_GPU_STAT(StochasticLighting);
+
 /**
  * Load GBuffer data once and transform it for subsequent lighting passes
  * This includes full res depth and normal copy for opaque before it gets overwritten by water or other translucency writing depth
  */
 void FDeferredShadingSceneRenderer::StochasticLightingTileClassificationMark(FRDGBuilder& GraphBuilder, FLumenSceneFrameTemporaries& FrameTemporaries, const FSceneTextures& SceneTextures)
 {
+	RDG_EVENT_SCOPE_STAT(GraphBuilder, StochasticLighting, "StochasticLighting");
+	RDG_GPU_STAT_SCOPE(GraphBuilder, StochasticLighting);
+
 	const ERDGPassFlags ComputePassFlags = ERDGPassFlags::Compute;
 
 	for (int32 ViewIndex = 0; ViewIndex < Views.Num(); ++ViewIndex)
