Stochastic Lighting - moved Stochastic Lighting to a broaded RenderDeferredLightingScope

diff --git a/Engine/Source/Runtime/Renderer/Private/DeferredShadingRenderer.cpp b/Engine/Source/Runtime/Renderer/Private/DeferredShadingRenderer.cpp
--- a/Engine/Source/Runtime/Renderer/Private/DeferredShadingRenderer.cpp
+++ b/Engine/Source/Runtime/Renderer/Private/DeferredShadingRenderer.cpp
@@ -2897,7 +2897,12 @@ void FDeferredShadingSceneRenderer::Render(FRDGBuilder& GraphBuilder, const FSce
 			Substrate::AddSubstrateSampleMaterialPass(GraphBuilder, Scene, SceneTextures, Views);
 		}
 
-		StochasticLightingTileClassificationMark(GraphBuilder, LumenFrameTemporaries, SceneTextures);
+		{
+			RDG_EVENT_SCOPE_STAT(GraphBuilder, RenderDeferredLighting, "StochasticLighting");
+			RDG_GPU_STAT_SCOPE(GraphBuilder, RenderDeferredLighting);
+
+			StochasticLightingTileClassificationMark(GraphBuilder, LumenFrameTemporaries, SceneTextures);
+		}		
 
 		// Copy lighting channels out of stencil before deferred decals which overwrite those values
 		TArray<FRDGTextureRef, TInlineAllocator<2>> NaniteShadingMask;
diff --git a/Engine/Source/Runtime/Renderer/Private/StochasticLighting/StochasticLighting.cpp b/Engine/Source/Runtime/Renderer/Private/StochasticLighting/StochasticLighting.cpp
--- a/Engine/Source/Runtime/Renderer/Private/StochasticLighting/StochasticLighting.cpp
+++ b/Engine/Source/Runtime/Renderer/Private/StochasticLighting/StochasticLighting.cpp
@@ -145,17 +145,12 @@ void StochasticLighting::TileClassificationMark(
 	}
 }
 
-DECLARE_GPU_STAT(StochasticLighting);
-
 /**
  * Load GBuffer data once and transform it for subsequent lighting passes
  * This includes full res depth and normal copy for opaque before it gets overwritten by water or other translucency writing depth
  */
 void FDeferredShadingSceneRenderer::StochasticLightingTileClassificationMark(FRDGBuilder& GraphBuilder, FLumenSceneFrameTemporaries& FrameTemporaries, const FSceneTextures& SceneTextures)
 {
-	RDG_EVENT_SCOPE_STAT(GraphBuilder, StochasticLighting, "StochasticLighting");
-	RDG_GPU_STAT_SCOPE(GraphBuilder, StochasticLighting);
-
 	const ERDGPassFlags ComputePassFlags = ERDGPassFlags::Compute;
 
 	for (int32 ViewIndex = 0; ViewIndex < Views.Num(); ++ViewIndex)
