Fixed blue noise texture nullptr access

#rnx

diff --git a/Engine/Source/Runtime/Renderer/Private/BlueNoise.cpp b/Engine/Source/Runtime/Renderer/Private/BlueNoise.cpp
--- a/Engine/Source/Runtime/Renderer/Private/BlueNoise.cpp
+++ b/Engine/Source/Runtime/Renderer/Private/BlueNoise.cpp
@@ -58,7 +58,8 @@ FBlueNoiseParameters GetBlueNoiseParameters()
 FBlueNoise GetBlueNoiseGlobalParameters()
 {
 	FBlueNoise Out;
-	if (GEngine->BlueNoiseScalarTexture != nullptr && GEngine->BlueNoiseScalarTexture->GetResource() != nullptr)		
+	if (GEngine->BlueNoiseScalarTexture != nullptr && GEngine->BlueNoiseScalarTexture->GetResource() != nullptr && 
+		GEngine->BlueNoiseVec2Texture != nullptr && GEngine->BlueNoiseVec2Texture->GetResource() != nullptr)
 	{
 		Out.BlueNoise = GetBlueNoiseParameters();
 	}
