SkyAtmosphere - added transmittance and multiple scattering versioning to avoid recomputing those luts when they would be unchanged.

diff --git a/Engine/Source/Runtime/Engine/Public/Rendering/SkyAtmosphereCommonData.cpp b/Engine/Source/Runtime/Engine/Public/Rendering/SkyAtmosphereCommonData.cpp
--- a/Engine/Source/Runtime/Engine/Public/Rendering/SkyAtmosphereCommonData.cpp
+++ b/Engine/Source/Runtime/Engine/Public/Rendering/SkyAtmosphereCommonData.cpp
@@ -25,6 +25,30 @@ FTentDistribution GetTentDistribution(const FSkyAtmosphereDynamicState& Ds)
 	return D;
 }
 
+
+void FAtmosphereSetup::ComputeAtmosphereVersion()
+{
+	uint32 Crc = 0;
+	Crc = FCrc::MemCrc32((void*)&BottomRadiusKm,					sizeof(float),		Crc);
+	Crc = FCrc::MemCrc32((void*)&TopRadiusKm,						sizeof(float),		Crc);
+	Crc = FCrc::MemCrc32((void*)&MultiScatteringFactor,				sizeof(float),		Crc);
+	Crc = FCrc::MemCrc32((void*)&RayleighScattering.R,				4 * sizeof(float),	Crc);
+	Crc = FCrc::MemCrc32((void*)&RayleighDensityExpScale,			sizeof(float),		Crc);
+	Crc = FCrc::MemCrc32((void*)&MieScattering.R,					4 * sizeof(float),	Crc);
+	Crc = FCrc::MemCrc32((void*)&MieExtinction.R,					4 * sizeof(float),	Crc);
+	Crc = FCrc::MemCrc32((void*)&MieAbsorption.R,					4 * sizeof(float),	Crc);
+	Crc = FCrc::MemCrc32((void*)&MieDensityExpScale,				sizeof(float),		Crc);
+	Crc = FCrc::MemCrc32((void*)&MiePhaseG,							sizeof(float),		Crc);
+	Crc = FCrc::MemCrc32((void*)&AbsorptionExtinction.R,			4 * sizeof(float),	Crc);
+	Crc = FCrc::MemCrc32((void*)&AbsorptionDensity0LayerWidth,		sizeof(float),		Crc);
+	Crc = FCrc::MemCrc32((void*)&AbsorptionDensity0ConstantTerm,	sizeof(float),		Crc);
+	Crc = FCrc::MemCrc32((void*)&AbsorptionDensity0LinearTerm,		sizeof(float),		Crc);
+	Crc = FCrc::MemCrc32((void*)&AbsorptionDensity1ConstantTerm,	sizeof(float),		Crc);
+	Crc = FCrc::MemCrc32((void*)&AbsorptionDensity1LinearTerm,		sizeof(float),		Crc);
+	Crc = FCrc::MemCrc32((void*)&GroundAlbedo.R,					4 * sizeof(float),	Crc);
+	TransmittanceAndMultiScatteringLUTsVersion = Crc;
+}
+
 template<typename T> void FAtmosphereSetup::InternalInit(const T& SkyAtmosphereComponent)
 {
 	// Convert Tent distribution to linear curve coefficients.
@@ -117,6 +141,8 @@ template<typename T> void FAtmosphereSetup::InternalInit(const T& SkyAtmosphereC
 	TransmittanceMinLightElevationAngle = SkyAtmosphereComponent.TransmittanceMinLightElevationAngle;
 
 	UpdateTransform(SkyAtmosphereComponent.GetComponentTransform(), uint8(SkyAtmosphereComponent.TransformMode));
+
+	ComputeAtmosphereVersion();
 }
 
 
diff --git a/Engine/Source/Runtime/Engine/Public/Rendering/SkyAtmosphereCommonData.h b/Engine/Source/Runtime/Engine/Public/Rendering/SkyAtmosphereCommonData.h
--- a/Engine/Source/Runtime/Engine/Public/Rendering/SkyAtmosphereCommonData.h
+++ b/Engine/Source/Runtime/Engine/Public/Rendering/SkyAtmosphereCommonData.h
@@ -58,6 +58,16 @@ struct FAtmosphereSetup
 	ENGINE_API void ComputeViewData(
 		const FVector& WorldCameraOrigin, const FVector& PreViewTranslation, const FVector3f& ViewForward, const FVector3f& ViewRight,
 		FVector3f& SkyCameraTranslatedWorldOrigin, FVector4f& SkyPlanetTranslatedWorldCenterAndViewHeight, FMatrix44f& SkyViewLutReferential) const;
+
+	uint32 GetTransmittanceAndMultiScatteringLUTsVersion() const
+	{
+		return TransmittanceAndMultiScatteringLUTsVersion;
+	}
+
+private:
+	uint32 TransmittanceAndMultiScatteringLUTsVersion; // A version to know if we have to recompute the transmittance and multiple scattering LUTs.
+
+	void ComputeAtmosphereVersion();
 };
 
 
diff --git a/Engine/Source/Runtime/Renderer/Private/RendererScene.cpp b/Engine/Source/Runtime/Renderer/Private/RendererScene.cpp
--- a/Engine/Source/Runtime/Renderer/Private/RendererScene.cpp
+++ b/Engine/Source/Runtime/Renderer/Private/RendererScene.cpp
@@ -1099,6 +1099,7 @@ FScene::FScene(UWorld* InWorld, bool bInRequiresHitProxies, bool bInIsEditorScen
 ,	DefaultLumenSceneData(nullptr)
 ,	PreshadowCacheLayout(0, 0, 0, 0, false)
 ,	SkyAtmosphere(NULL)
+,	CachedTransmittanceAndMultiScatteringLUTsVersion(0)
 ,	VolumetricCloud(NULL)
 ,	PrecomputedVisibilityHandler(NULL)
 ,	LocalShadowCastingLightOctree(FVector::ZeroVector, UE_OLD_HALF_WORLD_MAX)
diff --git a/Engine/Source/Runtime/Renderer/Private/ScenePrivate.h b/Engine/Source/Runtime/Renderer/Private/ScenePrivate.h
--- a/Engine/Source/Runtime/Renderer/Private/ScenePrivate.h
+++ b/Engine/Source/Runtime/Renderer/Private/ScenePrivate.h
@@ -3229,6 +3229,9 @@ public:
 	/** The sky/atmosphere components of the scene. */
 	FSkyAtmosphereRenderSceneInfo* SkyAtmosphere;
 
+	/** A value representing the state of the SkyAtmosphere to know if we need to update some psecific LUTs */
+	uint32 CachedTransmittanceAndMultiScatteringLUTsVersion;
+
 	/** Used to track the order that skylights were enabled in. */
 	TArray<FSkyAtmosphereSceneProxy*> SkyAtmosphereStack;
 
@@ -3451,6 +3454,15 @@ public:
 	virtual FSkyAtmosphereRenderSceneInfo* GetSkyAtmosphereSceneInfo() override { return SkyAtmosphere; }
 	virtual const FSkyAtmosphereRenderSceneInfo* GetSkyAtmosphereSceneInfo() const override { return SkyAtmosphere; }
 
+	inline uint32 GetCachedTransmittanceAndMultiScatteringLUTsVersion() const
+	{
+		return CachedTransmittanceAndMultiScatteringLUTsVersion;
+	}
+	inline void SetCachedTransmittanceAndMultiScatteringLUTsVersion(uint32 NewTransmittanceAndMultiScatteringLUTsVersion)
+	{
+		CachedTransmittanceAndMultiScatteringLUTsVersion = NewTransmittanceAndMultiScatteringLUTsVersion;
+	}
+
 	virtual void AddSparseVolumeTextureViewer(FSparseVolumeTextureViewerSceneProxy* SVTV) override;
 	virtual void RemoveSparseVolumeTextureViewer(FSparseVolumeTextureViewerSceneProxy* SVTV) override;
 
diff --git a/Engine/Source/Runtime/Renderer/Private/SkyAtmosphereRendering.cpp b/Engine/Source/Runtime/Renderer/Private/SkyAtmosphereRendering.cpp
--- a/Engine/Source/Runtime/Renderer/Private/SkyAtmosphereRendering.cpp
+++ b/Engine/Source/Runtime/Renderer/Private/SkyAtmosphereRendering.cpp
@@ -47,6 +47,11 @@ static TAutoConsoleVariable<int32> CVarSupportSkyAtmosphereAffectsHeightFog(
 	TEXT("Enables SkyAtmosphere affecting height fog. It requires r.SupportSkyAtmosphere to be true."),
 	ECVF_ReadOnly | ECVF_RenderThreadSafe);
 
+static TAutoConsoleVariable<bool> CVarSkyAtmosphereStateVersioning(
+	TEXT("r.SkyAtmosphere.StateVersioning"), true,
+	TEXT("SkyAtmosphere will only update the transmittance and multiple scattering LUT when its state changes in order to save some GPU time.\n"),
+	ECVF_RenderThreadSafe);
+
 ////////////////////////////////////////////////////////////////////////// Regular sky 
 
 static TAutoConsoleVariable<float> CVarSkyAtmosphereSampleCountMin(
@@ -1382,9 +1387,21 @@ void FSceneRenderer::RenderSkyAtmosphereLookUpTables(FRDGBuilder& GraphBuilder,
 
 	ERDGPassFlags PassFlag = CVarSkyAtmosphereAsyncCompute.GetValueOnAnyThread() ? ERDGPassFlags::AsyncCompute : ERDGPassFlags::Compute;
 
+	bool bEvaluateTransmittanceAndMultiScatteringLUTs = false;
+	const FAtmosphereSetup& AtmosphereSetup = SkyAtmosphereSceneProxy.GetAtmosphereSetup();
+	const uint32 IncomingTransmittanceAndMultiScatteringLUTsVersion = AtmosphereSetup.GetTransmittanceAndMultiScatteringLUTsVersion();
+	uint32 CachedTransmittanceAndMultiScatteringLUTsVersion = Scene->GetCachedTransmittanceAndMultiScatteringLUTsVersion();
+	if (CachedTransmittanceAndMultiScatteringLUTsVersion==0 
+		|| IncomingTransmittanceAndMultiScatteringLUTsVersion != CachedTransmittanceAndMultiScatteringLUTsVersion
+		|| CVarSkyAtmosphereStateVersioning.GetValueOnAnyThread() == false)
+	{
+		bEvaluateTransmittanceAndMultiScatteringLUTs = true;
+		Scene->SetCachedTransmittanceAndMultiScatteringLUTsVersion(IncomingTransmittanceAndMultiScatteringLUTsVersion);
+	}
+
 	// Transmittance LUT
 	FGlobalShaderMap* GlobalShaderMap = GetGlobalShaderMap(FeatureLevel);
-	if (CVarSkyAtmosphereTransmittanceLUT.GetValueOnRenderThread() > 0)
+	if (bEvaluateTransmittanceAndMultiScatteringLUTs && CVarSkyAtmosphereTransmittanceLUT.GetValueOnRenderThread() > 0)
 	{
 		RDG_EVENT_SCOPE(GraphBuilder, "SkyAtmosphere::TransmittanceLut");
 		TShaderMapRef<FRenderTransmittanceLutCS> ComputeShader(GlobalShaderMap);
@@ -1401,6 +1418,7 @@ void FSceneRenderer::RenderSkyAtmosphereLookUpTables(FRDGBuilder& GraphBuilder,
 	}
 
 	// Multi-Scattering LUT
+	if(bEvaluateTransmittanceAndMultiScatteringLUTs)
 	{
 		RDG_EVENT_SCOPE(GraphBuilder, "SkyAtmosphere::MultiScatteringLut");
 		FRenderMultiScatteredLuminanceLutCS::FPermutationDomain PermutationVector;
@@ -1497,7 +1515,6 @@ void FSceneRenderer::RenderSkyAtmosphereLookUpTables(FRDGBuilder& GraphBuilder,
 		const float AerialPerspectiveStartDepthInCm = GetValidAerialPerspectiveStartDepthInCm(ViewForReflection, SkyAtmosphereSceneProxy);
 
 		// Setup a constant referential for each of the faces of the dynamic reflection capture.
-		const FAtmosphereSetup& AtmosphereSetup = SkyAtmosphereSceneProxy.GetAtmosphereSetup();
 		const FVector3f SkyViewLutReferentialForward = FVector3f(1.0f, 0.0f, 0.0f);
 		const FVector3f SkyViewLutReferentialRight = FVector3f(0.0f, 1.0f, 0.0f);
 		FVector3f SkyCameraTranslatedWorldOrigin;
@@ -1760,6 +1777,7 @@ void FSceneRenderer::RenderSkyAtmosphereLookUpTables(FRDGBuilder& GraphBuilder,
 		PendingRDGResources.ViewResources[ViewIndex].SkyAtmosphereCameraAerialPerspectiveVolume = SkyAtmosphereCameraAerialPerspectiveVolume;
 	}
 
+	PendingRDGResources.bTransmittanceLutUpdated = bEvaluateTransmittanceAndMultiScatteringLUTs;
 	PendingRDGResources.TransmittanceLut = TransmittanceLut;
 }
 
@@ -1803,8 +1821,11 @@ void FSkyAtmospherePendingRDGResources::CommitToSceneAndViewUniformBuffers(FRDGB
 		View.SkyAtmosphereCameraAerialPerspectiveVolume = ConvertToExternalAccessTexture(GraphBuilder, ExternalAccessQueue, ViewResources[ViewIndex].SkyAtmosphereCameraAerialPerspectiveVolume, ERHIAccess::SRVMask, ERHIPipeline::All);
 	}
 
-	check(TransmittanceLut);
-	SkyInfo.GetTransmittanceLutTexture() = ConvertToExternalAccessTexture(GraphBuilder, ExternalAccessQueue, TransmittanceLut, ERHIAccess::SRVMask, ERHIPipeline::All);
+	if (bTransmittanceLutUpdated)
+	{
+		check(TransmittanceLut);
+		SkyInfo.GetTransmittanceLutTexture() = ConvertToExternalAccessTexture(GraphBuilder, ExternalAccessQueue, TransmittanceLut, ERHIAccess::SRVMask, ERHIPipeline::All);
+	}
 }
 
 FSkyAtmosphereRenderContext::FSkyAtmosphereRenderContext()
diff --git a/Engine/Source/Runtime/Renderer/Private/SkyAtmosphereRendering.h b/Engine/Source/Runtime/Renderer/Private/SkyAtmosphereRendering.h
--- a/Engine/Source/Runtime/Renderer/Private/SkyAtmosphereRendering.h
+++ b/Engine/Source/Runtime/Renderer/Private/SkyAtmosphereRendering.h
@@ -198,6 +198,8 @@ private:
 	FRDGBufferRef MobileDistantSkyLightLutBuffer = nullptr;
 	FRDGTextureRef RealTimeReflectionCaptureSkyAtmosphereViewLutTexture = nullptr;
 	FRDGTextureRef RealTimeReflectionCaptureCamera360APLutTexture = nullptr;
+
+	bool bTransmittanceLutUpdated = false;
 	FRDGTextureRef TransmittanceLut = nullptr;
 
 	friend class FSceneRenderer;
