Fix missing simple light lighting on hair strands geometry.

diff --git a/Engine/Source/Runtime/Renderer/Private/LightRendering.cpp b/Engine/Source/Runtime/Renderer/Private/LightRendering.cpp
--- a/Engine/Source/Runtime/Renderer/Private/LightRendering.cpp
+++ b/Engine/Source/Runtime/Renderer/Private/LightRendering.cpp
@@ -190,6 +190,12 @@ static TAutoConsoleVariable<int32> CVarHairStrandsAllowOneTransmittancePass(
 	TEXT("Allows one transmittance pass for hair strands lighting to have better performance (experimental).\n"),
 	ECVF_RenderThreadSafe);
 
+static TAutoConsoleVariable<int32> CVarHairStrandsSimpleLightEnable(
+	TEXT("r.HairStrands.Lighting.SimpleLights"),
+	1,
+	TEXT("Allows simple lights to lit hair strands geometry.\n"),
+	ECVF_RenderThreadSafe);
+
 #if ENABLE_DEBUG_DISCARD_PROP
 static float GDebugLightDiscardProp = 0.0f;
 static FAutoConsoleVariableRef CVarDebugLightDiscardProp(
@@ -3094,7 +3100,12 @@ BEGIN_SHADER_PARAMETER_STRUCT(FSimpleLightsStandardDeferredParameters, )
 	SHADER_PARAMETER_STRUCT_INCLUDE(FDeferredLightVS::FParameters, VS)
 END_SHADER_PARAMETER_STRUCT()
 
-static FSimpleLightsStandardDeferredParameters GetRenderLightSimpleParameters(
+BEGIN_SHADER_PARAMETER_STRUCT(FSimpleLightsForHairParameters, )
+	SHADER_PARAMETER_STRUCT_INCLUDE(FDeferredLightPS::FParameters, PS)
+	SHADER_PARAMETER_STRUCT_INCLUDE(FDeferredLightHairVS::FParameters, VS)
+END_SHADER_PARAMETER_STRUCT()
+
+static FDeferredLightPS::FParameters GetRenderLightSimplePSParameters(
 	FRDGBuilder& GraphBuilder,
 	const FScene* Scene,
 	const FViewInfo& View,
@@ -3102,48 +3113,167 @@ static FSimpleLightsStandardDeferredParameters GetRenderLightSimpleParameters(
 	const FSimpleLightEntry& SimpleLight,
 	const FVector& SimpleLightPosition)
 {
-	FSimpleLightsStandardDeferredParameters Out;
+	FDeferredLightPS::FParameters Out;
 
 	FRDGTextureRef WhiteDummy = GSystemTextures.GetWhiteDummy(GraphBuilder);
 	FRDGBufferRef BufferDummy = GSystemTextures.GetDefaultBuffer(GraphBuilder, 4, 0u);
 	FRDGBufferSRVRef BufferDummySRV = GraphBuilder.CreateSRV(BufferDummy, PF_R32_UINT);
 	
 	// PS - General parameters
-	Out.PS.SceneTextures = SceneTextures.UniformBuffer;
-	Out.PS.HairStrands = View.HairStrandsViewData.UniformBuffer;
-	Out.PS.Substrate = Substrate::BindSubstrateGlobalUniformParameters(View);
-	Out.PS.LightingChannelsTexture = WhiteDummy;
-	Out.PS.LightingChannelsSampler = TStaticSamplerState<SF_Point, AM_Clamp, AM_Clamp, AM_Clamp>::GetRHI();
-	Out.PS.CloudShadowAO = GetCloudShadowAOParameters(GraphBuilder, View, nullptr);
-	Out.PS.CloudShadowEnabled = 0;
-	SetupLightCloudTransmittanceParameters(GraphBuilder, nullptr, View, nullptr, Out.PS.CloudShadow);
-	Out.PS.LightAttenuationTexture = WhiteDummy;
-	Out.PS.LightAttenuationTextureSampler = TStaticSamplerState<SF_Point, AM_Wrap, AM_Wrap, AM_Wrap>::GetRHI();
-	Out.PS.View = View.ViewUniformBuffer;
-	Out.PS.DeferredLight = CreateDeferredLightUniformBuffer(GraphBuilder, View, SimpleLight, SimpleLightPosition);
+	Out.SceneTextures = SceneTextures.UniformBuffer;
+	Out.HairStrands = View.HairStrandsViewData.UniformBuffer;
+	Out.Substrate = Substrate::BindSubstrateGlobalUniformParameters(View);
+	Out.LightingChannelsTexture = WhiteDummy;
+	Out.LightingChannelsSampler = TStaticSamplerState<SF_Point, AM_Clamp, AM_Clamp, AM_Clamp>::GetRHI();
+	Out.CloudShadowAO = GetCloudShadowAOParameters(GraphBuilder, View, nullptr);
+	Out.CloudShadowEnabled = 0;
+	SetupLightCloudTransmittanceParameters(GraphBuilder, nullptr, View, nullptr, Out.CloudShadow);
+	Out.LightAttenuationTexture = WhiteDummy;
+	Out.LightAttenuationTextureSampler = TStaticSamplerState<SF_Point, AM_Wrap, AM_Wrap, AM_Wrap>::GetRHI();
+	Out.View = View.ViewUniformBuffer;
+	Out.DeferredLight = CreateDeferredLightUniformBuffer(GraphBuilder, View, SimpleLight, SimpleLightPosition);
 	// PS - Hair (default)
-	Out.PS.ScreenShadowMaskSubPixelTexture = WhiteDummy;
-	Out.PS.HairTransmittanceBuffer = BufferDummySRV;
-	Out.PS.HairTransmittanceBufferMaxCount = 0;
-	Out.PS.HairShadowMaskValid = false;
-	Out.PS.ShadowChannelMask = FVector4f(1, 1, 1, 1);
+	Out.ScreenShadowMaskSubPixelTexture = WhiteDummy;
+	Out.HairTransmittanceBuffer = BufferDummySRV;
+	Out.HairTransmittanceBufferMaxCount = 0;
+	Out.HairShadowMaskValid = false;
+	Out.ShadowChannelMask = FVector4f(1, 1, 1, 1);
 	// PS - RT/Depth
-	Out.PS.RenderTargets[0] = FRenderTargetBinding(SceneTextures.Color.Target, ERenderTargetLoadAction::ELoad);
+	Out.RenderTargets[0] = FRenderTargetBinding(SceneTextures.Color.Target, ERenderTargetLoadAction::ELoad);
 	if (SceneTextures.Depth.Target)
 	{
-		Out.PS.RenderTargets.DepthStencil = FDepthStencilBinding(SceneTextures.Depth.Target, ERenderTargetLoadAction::ELoad, ERenderTargetLoadAction::ELoad, FExclusiveDepthStencil::DepthRead_StencilWrite);
+		Out.RenderTargets.DepthStencil = FDepthStencilBinding(SceneTextures.Depth.Target, ERenderTargetLoadAction::ELoad, ERenderTargetLoadAction::ELoad, FExclusiveDepthStencil::DepthRead_StencilWrite);
 	}
-	Out.PS.AVSM = HeterogeneousVolumes::CreateEmptyAdaptiveVolumetricShadowMapUniformBuffer(GraphBuilder);
+	Out.AVSM = HeterogeneousVolumes::CreateEmptyAdaptiveVolumetricShadowMapUniformBuffer(GraphBuilder);
+
+	return Out;
+}
+
+static FSimpleLightsStandardDeferredParameters GetRenderLightSimpleParameters(
+	FRDGBuilder& GraphBuilder,
+	const FScene* Scene,
+	const FViewInfo& View,
+	const FMinimalSceneTextures& SceneTextures,
+	const FSimpleLightEntry& SimpleLight,
+	const FVector& SimpleLightPosition)
+{
+	FSimpleLightsStandardDeferredParameters Out;
+
+	Out.PS = GetRenderLightSimplePSParameters(GraphBuilder, Scene, View, SceneTextures, SimpleLight, SimpleLightPosition);
 
 	// VS - General parameters (dummy geometry, as the geometry is setup within the pass light loop)
 	FSphere SphereLight;
 	SphereLight.Center = SimpleLightPosition; // Should we account for LWC Position+Tile here?
 	SphereLight.W = SimpleLight.Radius;
 	Out.VS = FDeferredLightVS::GetParameters(View, SphereLight, false);
+	return Out;
+}
+static FSimpleLightsForHairParameters GetRenderLightSimpleForHairParameters(
+	FRDGBuilder& GraphBuilder,
+	const FScene* Scene,
+	const FViewInfo& View,
+	const FMinimalSceneTextures& SceneTextures,
+	const FSimpleLightEntry& SimpleLight,
+	const FVector& SimpleLightPosition)
+{
+	FSimpleLightsForHairParameters Out;
+
+	Out.PS = GetRenderLightSimplePSParameters(GraphBuilder, Scene, View, SceneTextures, SimpleLight, SimpleLightPosition);
+
+	// Hair VS
+	Out.VS.HairStrands = HairStrands::BindHairStrandsViewUniformParameters(View);
+
+	// Hair PS
+	Out.PS.RenderTargets[0] = FRenderTargetBinding(View.HairStrandsViewData.VisibilityData.SampleLightingTexture, ERenderTargetLoadAction::ELoad);
+	Out.PS.RenderTargets[1] = FRenderTargetBinding();
+	Out.PS.RenderTargets[2] = FRenderTargetBinding();
+	Out.PS.RenderTargets.DepthStencil = FDepthStencilBinding(nullptr, ERenderTargetLoadAction::ENoAction, ERenderTargetLoadAction::ENoAction, FExclusiveDepthStencil::DepthNop_StencilNop);
 
 	return Out;
 }
 
+static void InternalRenderSimpleLightsForHair(
+	FRDGBuilder& GraphBuilder,
+	const FScene* Scene,
+	const FViewInfo& View,
+	const uint32 ViewIndex,
+	const uint32 NumViews,
+	const FMinimalSceneTextures& SceneTextures,
+	const FSimpleLightArray& SimpleLights)
+{
+	FSimpleLightsForHairParameters* PassParameters = GraphBuilder.AllocParameters<FSimpleLightsForHairParameters>();
+	*PassParameters = GetRenderLightSimpleForHairParameters(
+		GraphBuilder,
+		Scene,
+		View,
+		SceneTextures,
+		SimpleLights.InstanceData[0], // Use a dummy light to create the PassParameter buffer. The light data will be
+		FVector(0, 0, 0));		  // update dynamically with the pass light loop for efficiency purpose
+
+	FDeferredLightPS::FPermutationDomain PermutationVector;
+	PermutationVector.Set< FDeferredLightPS::FSourceShapeDim >(ELightSourceShape::Capsule);
+	PermutationVector.Set< FDeferredLightPS::FIESProfileDim >(false);
+	PermutationVector.Set< FDeferredLightPS::FLightFunctionAtlasDim>(false);
+	PermutationVector.Set< FDeferredLightPS::FVisualizeCullingDim >(false);
+	PermutationVector.Set< FDeferredLightPS::FLightingChannelsDim >(false);
+	PermutationVector.Set< FDeferredLightPS::FAnistropicMaterials >(false);
+	PermutationVector.Set< FDeferredLightPS::FTransmissionDim >(false);
+	PermutationVector.Set< FDeferredLightPS::FHairLighting>(1);
+	PermutationVector.Set< FDeferredLightPS::FHairComplexTransmittance>(true /*bNeedComplexTransmittanceSupport*/);
+	PermutationVector.Set< FDeferredLightPS::FAtmosphereTransmittance >(false);
+	PermutationVector.Set< FDeferredLightPS::FCloudTransmittance >(false);
+	PermutationVector.Set< FDeferredLightPS::FSubstrateTileType>(Substrate::IsSubstrateEnabled() ? ESubstrateTileType::ESimple : 0);
+	TShaderMapRef<FDeferredLightPS> PixelShader(View.ShaderMap, PermutationVector);
+	TShaderMapRef<FDeferredLightHairVS> VertexShader(View.ShaderMap);
+
+	GraphBuilder.AddPass(
+		RDG_EVENT_NAME("Light::DeferredSimpleLightsForHair"),
+		PassParameters,
+		ERDGPassFlags::Raster,
+		[&View, &SimpleLights, ViewIndex, NumViews, PassParameters, PixelShader, VertexShader](FRDGAsyncTask, FRHICommandList& RHICmdList)
+	{
+		FGraphicsPipelineStateInitializer GraphicsPSOInit;
+		RHICmdList.ApplyCachedRenderTargets(GraphicsPSOInit);
+
+		SetupLightForHairGraphicsPSOState(GraphicsPSOInit);
+
+		// Use additive blending for color
+		GraphicsPSOInit.BlendState = TStaticBlendState<CW_RGBA, BO_Add, BF_One, BF_One, BO_Add, BF_One, BF_One>::GetRHI();
+		GraphicsPSOInit.PrimitiveType = PT_TriangleList;
+
+		const int32 GlobalPSOCollectorIndex = FGlobalPSOCollectorManager::GetIndex(DeferredLightGlobalPSOCollectorName);
+		const FIntPoint SampleLightingViewportResolution = View.HairStrandsViewData.VisibilityData.SampleLightingViewportResolution;
+		RHICmdList.SetViewport(0, 0, 0.0f, SampleLightingViewportResolution.X, SampleLightingViewportResolution.Y, 1.0f);
+
+		for (int32 LightIndex = 0; LightIndex < SimpleLights.InstanceData.Num(); LightIndex++)
+		{
+			const FSimpleLightEntry& SimpleLight = SimpleLights.InstanceData[LightIndex];
+
+			const FSimpleLightPerViewEntry& SimpleLightPerViewData = SimpleLights.GetViewDependentData(LightIndex, ViewIndex, NumViews);
+			GraphicsPSOInit.BoundShaderState.VertexShaderRHI = VertexShader.GetVertexShader();
+			GraphicsPSOInit.BoundShaderState.PixelShaderRHI = PixelShader.GetPixelShader();
+						
+			#if PSO_PRECACHING_VALIDATE
+			if (PSOCollectorStats::IsFullPrecachingValidationEnabled())
+			{
+				PSOCollectorStats::CheckGlobalGraphicsPipelineStateInCache(GraphicsPSOInit, GlobalPSOCollectorIndex);
+			}
+			#endif // PSO_PRECACHING_VALIDATE
+
+			SetGraphicsPipelineState(RHICmdList, GraphicsPSOInit, 0/*StencilRef*/);
+
+			// Update the light parameters with a custom uniform buffer
+			FDeferredLightUniformStruct DeferredLightUniformsValue = GetSimpleDeferredLightParameters(View, SimpleLight, SimpleLightPerViewData);
+
+			SetShaderParametersMixed(RHICmdList, PixelShader, PixelShader.GetPixelShader(), PassParameters->PS, DeferredLightUniformsValue);
+			SetShaderParameters(RHICmdList, VertexShader, VertexShader.GetVertexShader(), PassParameters->VS);
+
+			RHICmdList.SetStreamSource(0, nullptr, 0);
+			RHICmdList.DrawPrimitive(0, 1, 1);
+		}
+	});
+}
+
 static void InternalRenderSimpleLightsStandardDeferred(
 	FRDGBuilder& GraphBuilder,
 	const FScene* Scene,
@@ -3258,6 +3388,7 @@ void FDeferredShadingSceneRenderer::RenderSimpleLightsStandardDeferred(
 	SCOPE_CYCLE_COUNTER(STAT_DirectLightRenderingTime);
 	INC_DWORD_STAT_BY(STAT_NumLightsUsingStandardDeferred, SimpleLights.InstanceData.Num());
 
+	// GBufffer lighting
 	for (int32 ViewIndex = 0, NumViews = Views.Num(); ViewIndex < NumViews; ViewIndex++)
 	{
 		const FViewInfo& View = Views[ViewIndex];
@@ -3277,6 +3408,16 @@ void FDeferredShadingSceneRenderer::RenderSimpleLightsStandardDeferred(
 			InternalRenderSimpleLightsStandardDeferred(GraphBuilder, Scene, View, ViewIndex, NumViews, SceneTextures, SimpleLights, ESubstrateTileType::ECount);
 		}
 	}
+
+	// Hair lighting
+	if (HairStrands::HasViewHairStrandsData(Views) && CVarHairStrandsSimpleLightEnable.GetValueOnRenderThread() > 0)
+	{
+		for (int32 ViewIndex = 0, NumViews = Views.Num(); ViewIndex < NumViews; ViewIndex++)
+		{
+			const FViewInfo& View = Views[ViewIndex];
+			InternalRenderSimpleLightsForHair(GraphBuilder, Scene, View, ViewIndex, NumViews, SceneTextures, SimpleLights);
+		}
+	}
 }
 
 class FCopyStencilToLightingChannelsPS : public FGlobalShader
