Reduce DeferredLightPixelShader's shader permutation count when Substrate is enabled.

Trimmed tile permutation for hair rendering as only simple tile is relevant (1304 -> 1112 vs 454 for legacy).

diff --git a/Engine/Source/Runtime/Renderer/Private/LightRendering.cpp b/Engine/Source/Runtime/Renderer/Private/LightRendering.cpp
--- a/Engine/Source/Runtime/Renderer/Private/LightRendering.cpp
+++ b/Engine/Source/Runtime/Renderer/Private/LightRendering.cpp
@@ -984,6 +984,11 @@ class FDeferredLightPS : public FGlobalShader
 			return false;
 		}
 
+		if (Substrate::IsSubstrateEnabled() && PermutationVector.Get<FHairLighting>() == 1 && PermutationVector.Get<FSubstrateTileType>() != ESubstrateTileType::ESimple)
+		{
+			return false;
+		}
+
 		if (PermutationVector.Get<FDeferredLightPS::FAnistropicMaterials>())
 		{
 			if (Substrate::IsSubstrateEnabled())
@@ -2982,6 +2987,7 @@ void FDeferredShadingSceneRenderer::RenderLightForHair(
 	PermutationVector.Set< FDeferredLightPS::FHairLighting>(1);
 	PermutationVector.Set< FDeferredLightPS::FHairComplexTransmittance>(true);
 	PermutationVector.Set< FDeferredLightPS::FVirtualShadowMapMask >(bUseVirtualShadowMapMask);
+	PermutationVector.Set< FDeferredLightPS::FSubstrateTileType>(Substrate::IsSubstrateEnabled() ? ESubstrateTileType::ESimple : 0);
 	PermutationVector.Set< FDeferredLightPS::FLightFunctionAtlasDim >(
 		LightFunctionAtlas::IsEnabled(View, ELightFunctionAtlasSystem::DeferredLighting) && LightSceneInfo->Proxy->HasValidLightFunctionAtlasSlot() &&
 		LightSceneInfo->Proxy->GetLightFunctionMaterial() != nullptr && !View.Family->EngineShowFlags.VisualizeLightCulling && bCanLightUsesAtlasForUnbatchedLight);
@@ -3227,7 +3233,7 @@ static void InternalRenderSimpleLightsForHair(
 	TShaderMapRef<FDeferredLightHairVS> VertexShader(View.ShaderMap);
 
 	GraphBuilder.AddPass(
-		RDG_EVENT_NAME("Light::DeferredSimpleLightsForHair"),
+		RDG_EVENT_NAME("Light::DeferredSimpleLights(Hair)"),
 		PassParameters,
 		ERDGPassFlags::Raster,
 		[&View, &SimpleLights, ViewIndex, NumViews, PassParameters, PixelShader, VertexShader](FRDGAsyncTask, FRHICommandList& RHICmdList)
