Fixed crash in ShadowSceneRenderer caused by a race with the task updating the ShadowScene.

diff --git a/Engine/Source/Runtime/Renderer/Private/Shadows/ShadowSceneRenderer.cpp b/Engine/Source/Runtime/Renderer/Private/Shadows/ShadowSceneRenderer.cpp
--- a/Engine/Source/Runtime/Renderer/Private/Shadows/ShadowSceneRenderer.cpp
+++ b/Engine/Source/Runtime/Renderer/Private/Shadows/ShadowSceneRenderer.cpp
@@ -172,6 +172,9 @@ struct FHeapPair
 
 void FShadowSceneRenderer::PreInitViews(FRDGBuilder& GraphBuilder)
 {
+	// Wait for scene update task to finish
+	ShadowScene.SceneChangeUpdateTask.Wait();
+
 	// Clear the frame setups to indicate that nothing is allocated for this frame
 	CommonSetups.SetNumZeroed(ShadowScene.LightsCommonData.GetMaxIndex());
 
