Fixed: Missing Primitive UB for a Niagara mesh particles on platforms that use mobile "UniformView"
(submitting on behalf of Dmitriy.Dyomin))

diff --git a/Engine/Source/Runtime/Renderer/Private/ShaderBaseClasses.cpp b/Engine/Source/Runtime/Renderer/Private/ShaderBaseClasses.cpp
--- a/Engine/Source/Runtime/Renderer/Private/ShaderBaseClasses.cpp
+++ b/Engine/Source/Runtime/Renderer/Private/ShaderBaseClasses.cpp
@@ -493,13 +493,20 @@ void FMeshMaterialShader::GetElementShaderBindings(
 	}
 
 	EShaderPlatform ShaderPlatform = Scene ? Scene->GetShaderPlatform() : GMaxRHIShaderPlatform;
-	bool PlatformUsesPrimitiveUB = (BatchElement.bForceInstanceCulling && PlatformGPUSceneUsesUniformBufferView(ShaderPlatform));
+	bool PlatformUsesPrimitiveUB = false;
+	if (PlatformGPUSceneUsesUniformBufferView(ShaderPlatform))
+	{
+		// On platforms that use "UniformView" GPUScene, some VF might bind primitive UB
+		// this includes ISM VFs and mesh particles
+		PlatformUsesPrimitiveUB = (BatchElement.bForceInstanceCulling || (PrimitiveSceneProxy && PrimitiveSceneProxy->DoesVFRequirePrimitiveUniformBuffer()));
+	}
+		
 	if (UseGPUScene(ShaderPlatform, FeatureLevel) && VertexFactory->GetPrimitiveIdStreamIndex(FeatureLevel, InputStreamType) >= 0 && !PlatformUsesPrimitiveUB)
 	{
 		const FShaderType* ShaderType = GetType(PointerTable);
 		ensureMsgf(!GetUniformBufferParameter<FPrimitiveUniformShaderParameters>().IsBound(), TEXT("Shader %s attempted to bind the Primitive uniform buffer even though Vertex Factory computes a PrimitiveId per-instance.  This will break auto-instancing.  Shaders should use GetPrimitiveData(PrimitiveId).Member instead of Primitive.Member."), ShaderType->GetName());
 		// Some primitives may use several VFs with a mixed support for a GPUScene. In this case all mesh batches get Primitive UB assigned regardless of VF type 
-		ensureMsgf(!BatchElement.PrimitiveUniformBuffer || (PrimitiveSceneProxy->DoesVFRequirePrimitiveUniformBuffer()), TEXT("FMeshBatchElement was assigned a PrimitiveUniformBuffer even though Vertex Factory %s fetches primitive shader data through a Scene buffer.  The assigned PrimitiveUniformBuffer cannot be respected.  Use PrimitiveUniformBufferResource instead for dynamic primitive data."), ShaderType->GetName());
+		ensureMsgf(!BatchElement.PrimitiveUniformBuffer || (PrimitiveSceneProxy && PrimitiveSceneProxy->DoesVFRequirePrimitiveUniformBuffer()), TEXT("FMeshBatchElement was assigned a PrimitiveUniformBuffer even though Vertex Factory %s fetches primitive shader data through a Scene buffer.  The assigned PrimitiveUniformBuffer cannot be respected.  Use PrimitiveUniformBufferResource instead for dynamic primitive data."), ShaderType->GetName());
 	}
 	else
 	{
