Fix Lens_LensFlares test failed due to Ensure condition failed: !State.bInsideBeginRenderPass.

Add the SkipRenderPass to allow the RHITransition call inside the lambda.

#tests lens flare engin test

diff --git a/Engine/Source/Runtime/Renderer/Private/PostProcess/PostProcessLensFlares.cpp b/Engine/Source/Runtime/Renderer/Private/PostProcess/PostProcessLensFlares.cpp
--- a/Engine/Source/Runtime/Renderer/Private/PostProcess/PostProcessLensFlares.cpp
+++ b/Engine/Source/Runtime/Renderer/Private/PostProcess/PostProcessLensFlares.cpp
@@ -276,11 +276,14 @@ FScreenPassTexture AddLensFlaresPass(
 		TShaderMapRef<FLensFlareBlurVS> VertexShader(View.ShaderMap);
 		TShaderMapRef<FLensFlareBlurPS> PixelShader(View.ShaderMap);
 
+		ERDGPassFlags LensFlareBlurFlags = ERDGPassFlags::Raster;
+		LensFlareBlurFlags |= GRHINeedsSRVGraphicsNonPixelWorkaround ? ERDGPassFlags::SkipRenderPass : ERDGPassFlags::None;
+
 		GraphBuilder.AddPass(
 			RDG_EVENT_NAME("LensFlareBlur %dx%d", FlareViewSize.X, FlareViewSize.Y),
 			PassParameters,
-			ERDGPassFlags::Raster,
-			[VertexShader, PixelShader, VertexParameters, PixelParameters, AdditiveBlendState, FlareViewport, TileCount](FRDGAsyncTask, FRHICommandList& RHICmdList)
+			LensFlareBlurFlags,
+			[VertexShader, PixelShader, VertexParameters, PixelParameters, PassParameters, AdditiveBlendState, FlareViewport, TileCount](FRDGAsyncTask, FRHICommandList& RHICmdList)
 			{
 				// Viewport is the same as the input.
 				RHICmdList.SetViewport(
@@ -293,8 +296,12 @@ FScreenPassTexture AddLensFlaresPass(
 
 				// Force a transition to SRVGraphicsNonPixel on a dummy texture to account for VS accessing the texture (using NON_PIXEL_SRV in struct gets merged)
 				if(GRHINeedsSRVGraphicsNonPixelWorkaround)
-				{ 
+				{
 					RHICmdList.Transition(FRHITransitionInfo(GDummyTransitionTexture->GetTextureRHI(), ERHIAccess::Unknown, ERHIAccess::SRVGraphicsNonPixel));
+					
+					
+					FRHIRenderPassInfo RPInfo(PassParameters->RenderTargets[0].GetTexture()->GetRHI(), ERenderTargetActions::Clear_Store);
+					RHICmdList.BeginRenderPass(RPInfo, TEXT("LensFlareBlur"));
 				}
 			
 				// Apply additive blending pipeline state.
@@ -313,6 +320,11 @@ FScreenPassTexture AddLensFlaresPass(
 				// Emit an instanced quad draw call on the order of the number of pixels on the screen.
 				RHICmdList.SetStreamSource(0, nullptr, 0);
 				RHICmdList.DrawPrimitive(0, 2, FMath::DivideAndRoundUp(TileCount.X * TileCount.Y, GLensFlareQuadsPerInstance));
+
+				if (GRHINeedsSRVGraphicsNonPixelWorkaround)
+				{
+					RHICmdList.EndRenderPass();
+				}
 			});
 	}
 	else
