MegaLights - fixed split screen crash due to out of bounds memory reads

diff --git a/Engine/Source/Runtime/Renderer/Private/TranslucentLighting.cpp b/Engine/Source/Runtime/Renderer/Private/TranslucentLighting.cpp
--- a/Engine/Source/Runtime/Renderer/Private/TranslucentLighting.cpp
+++ b/Engine/Source/Runtime/Renderer/Private/TranslucentLighting.cpp
@@ -2407,9 +2407,9 @@ void FilterTranslucencyLightingVolume(
 				GraphBuilder.QueueTextureExtraction(Textures.Ambient[TextureIndex], &TranslucencyLightingViewState.HistoryAmbient[VolumeCascadeIndex]);
 				GraphBuilder.QueueTextureExtraction(Textures.Directional[TextureIndex], &TranslucencyLightingViewState.HistoryDirectional[VolumeCascadeIndex]);
 
-				if (View.TranslucencyVolumeMarkData[TextureIndex].MarkTexture)
+				if (View.TranslucencyVolumeMarkData[VolumeCascadeIndex].MarkTexture)
 				{
-					GraphBuilder.QueueTextureExtraction(View.TranslucencyVolumeMarkData[TextureIndex].MarkTexture, &TranslucencyLightingViewState.HistoryMark[VolumeCascadeIndex]);
+					GraphBuilder.QueueTextureExtraction(View.TranslucencyVolumeMarkData[VolumeCascadeIndex].MarkTexture, &TranslucencyLightingViewState.HistoryMark[VolumeCascadeIndex]);
 				}
 				else
 				{
