[Eye Adaptation] Fix incorrect eye adaptation when sequence ends.
- View.bCameraCut is not set in this case so also need to check View.bPrevTransformsReset.

diff --git a/Engine/Source/Runtime/Renderer/Private/PostProcess/PostProcessEyeAdaptation.cpp b/Engine/Source/Runtime/Renderer/Private/PostProcess/PostProcessEyeAdaptation.cpp
--- a/Engine/Source/Runtime/Renderer/Private/PostProcess/PostProcessEyeAdaptation.cpp
+++ b/Engine/Source/Runtime/Renderer/Private/PostProcess/PostProcessEyeAdaptation.cpp
@@ -571,7 +571,7 @@ FEyeAdaptationParameters GetEyeAdaptationParameters(const FViewInfo& View)
 	// if it is a camera cut we force the exposure to go all the way to the target exposure without blending.
 	// if it is manual mode, we also force the exposure to hit the target, which matters for HDR Visualization
 	// if we don't have a valid range (AutoExposureMinBrightness == AutoExposureMaxBrightness) then force it like Manual as well.
-	const float ForceTarget = (View.bCameraCut || AutoExposureMethod == EAutoExposureMethod::AEM_Manual || !bValidRange || !bValidSpeeds) ? 1.0f : 0.0f;
+	const float ForceTarget = (View.bCameraCut || View.bPrevTransformsReset || AutoExposureMethod == EAutoExposureMethod::AEM_Manual || !bValidRange || !bValidSpeeds) ? 1.0f : 0.0f;
 
 	FEyeAdaptationParameters Parameters;
 	Parameters.ExposureLowPercent = ExposureLowPercent;
