Substrate - do not allocate a Substrate GBuffer and other resources using clear to not have them allocate when using the legacy blendable gbuffer.

#rnx
#tests pc editor dx12 and console

diff --git a/Engine/Source/Runtime/Renderer/Private/PostProcess/PostProcessCompositeEditorPrimitives.cpp b/Engine/Source/Runtime/Renderer/Private/PostProcess/PostProcessCompositeEditorPrimitives.cpp
--- a/Engine/Source/Runtime/Renderer/Private/PostProcess/PostProcessCompositeEditorPrimitives.cpp
+++ b/Engine/Source/Runtime/Renderer/Private/PostProcess/PostProcessCompositeEditorPrimitives.cpp
@@ -290,7 +290,7 @@ FScreenPassTexture AddEditorPrimitivePass(
 	
 	// Substrate data might not be produced in certain case (e.g., path-tracer). In such a case we force generate 
 	// them with a simple clear to please validation.
-	if (Substrate::IsSubstrateEnabled() && Substrate::UsesSubstrateMaterialBuffer(View.GetShaderPlatform()) && !HasBeenProduced(View.SubstrateViewData.SceneData->TopLayerTexture))
+	if (Substrate::IsSubstrateEnabled() && Substrate::UsesSubstrateMaterialBuffer(View.GetShaderPlatform()) && View.SubstrateViewData.SceneData->TopLayerTexture && !HasBeenProduced(View.SubstrateViewData.SceneData->TopLayerTexture))
 	{
 		FRDGTextureClearInfo ClearInfo;
 		AddClearRenderTargetPass(GraphBuilder, View.SubstrateViewData.SceneData->TopLayerTexture, ClearInfo);
diff --git a/Engine/Source/Runtime/Renderer/Private/Substrate/Substrate.cpp b/Engine/Source/Runtime/Renderer/Private/Substrate/Substrate.cpp
--- a/Engine/Source/Runtime/Renderer/Private/Substrate/Substrate.cpp
+++ b/Engine/Source/Runtime/Renderer/Private/Substrate/Substrate.cpp
@@ -215,7 +215,7 @@ bool GetSubstrateUsesComplexSpecialPath(const FViewInfo& View)
 	return false;
 }
 
-static void BindSubstrateGlobalUniformParameters(FRDGBuilder& GraphBuilder, const FSubstrateViewData* SubstrateViewData, bool bNeedsMaterialBuffer, FSubstrateGlobalUniformParameters& OutSubstrateUniformParameters);
+static void BindSubstrateGlobalUniformParameters(FRDGBuilder& GraphBuilder, const FSubstrateViewData* SubstrateViewData, bool bNeedsMaterialBuffer, bool bBlendableGBufferEnabled, FSubstrateGlobalUniformParameters& OutSubstrateUniformParameters);
 
 bool SupportsCMask(const FStaticShaderPlatform InPlatform)
 {
@@ -232,10 +232,10 @@ static EPixelFormat GetClassificationTileFormat(const FIntPoint& InResolution, u
 	return InTileEncoding == SUBSTRATE_TILE_ENCODING_16BITS ? PF_R32_UINT : PF_R16_UINT;
 }
 
-static void InitialiseSubstrateViewData(FRDGBuilder& GraphBuilder, FViewInfo& View, const FSceneTexturesConfig& SceneTexturesConfig, bool bNeedsClosureOffets, bool bNeedsMaterialBuffer, FSubstrateSceneData& SceneData)
+static void InitialiseSubstrateViewData(FRDGBuilder& GraphBuilder, FViewInfo& View, const FSceneTexturesConfig& SceneTexturesConfig, bool bNeedsClosureOffets, bool bNeedsMaterialBuffer, bool bBlendableGBufferEnabled, FSubstrateSceneData& SceneData)
 {
 	// Sanity check: the scene data should already exist 
-	if (bNeedsMaterialBuffer)
+	if (bNeedsMaterialBuffer && !bBlendableGBufferEnabled)
 	{
 		check(SceneData.MaterialTextureArray != nullptr);
 	}
@@ -324,7 +324,8 @@ static void InitialiseSubstrateViewData(FRDGBuilder& GraphBuilder, FViewInfo& Vi
 		// Create the readable uniform buffers
 		{
 			FSubstrateGlobalUniformParameters* SubstrateUniformParameters = GraphBuilder.AllocParameters<FSubstrateGlobalUniformParameters>();
-			BindSubstrateGlobalUniformParameters(GraphBuilder, &Out, bNeedsMaterialBuffer, *SubstrateUniformParameters);
+
+			BindSubstrateGlobalUniformParameters(GraphBuilder, &Out, bNeedsMaterialBuffer, bBlendableGBufferEnabled, *SubstrateUniformParameters);
 			Out.SubstrateGlobalUniformParameters = GraphBuilder.CreateUniformBuffer(SubstrateUniformParameters);
 		}
 	}
@@ -420,6 +421,7 @@ void InitialiseSubstrateFrameSceneData(FRDGBuilder& GraphBuilder, FSceneRenderer
 
 	// Compute the max byte per pixels required by the views
 	bool bNeedsMaterialBuffer = UsesSubstrateMaterialBuffer(SceneRenderer.ShaderPlatform);
+	bool bBlendableGBufferEnabled = Substrate::IsSubstrateBlendableGBufferEnabled(SceneRenderer.ShaderPlatform);
 	bool bNeedsClosureOffsets = false;
 	bool bNeedsUAV = false;
 	bool bUseDBufferPass = false;
@@ -506,22 +508,13 @@ void InitialiseSubstrateFrameSceneData(FRDGBuilder& GraphBuilder, FSceneRenderer
 		UpdateMaterialBufferToTiledResolution(SceneTextureExtent, MaterialBufferSizeXY);
 
 		// Top layer texture
-		if (bNeedsMaterialBuffer)
+		if (bNeedsMaterialBuffer && !bBlendableGBufferEnabled)
 		{
-			if (Substrate::IsSubstrateBlendableGBufferEnabled(SceneRenderer.ShaderPlatform))
-			{
-				// Some passes cannot check the usage of TopLayer texture in the shader since the shader is selected later within the pass lambda. So we always allocate a dummy one that is cleared.
-				Out.TopLayerTexture = GraphBuilder.CreateTexture(FRDGTextureDesc::Create2D(FIntPoint(1,1), GetTopLayerTextureFormat(bUseDBufferPass), FClearValueBinding::Black, TexCreate_RenderTargetable | TexCreate_ShaderResource | TexCreate_FastVRAM | TexCreate_UAV), TEXT("Substrate.TopLayerTexture"));
-				AddClearUAVPass(GraphBuilder, GraphBuilder.CreateUAV(Out.TopLayerTexture), 0u);
-			}
-			else
-			{
-				Out.TopLayerTexture = GraphBuilder.CreateTexture(FRDGTextureDesc::Create2D(SceneTextureExtent, GetTopLayerTextureFormat(bUseDBufferPass), FClearValueBinding::Black, TexCreate_RenderTargetable | TexCreate_ShaderResource | TexCreate_FastVRAM | (bNeedsUAV ? TexCreate_UAV : TexCreate_None)), TEXT("Substrate.TopLayerTexture"));
-			}
+			Out.TopLayerTexture = GraphBuilder.CreateTexture(FRDGTextureDesc::Create2D(SceneTextureExtent, GetTopLayerTextureFormat(bUseDBufferPass), FClearValueBinding::Black, TexCreate_RenderTargetable | TexCreate_ShaderResource | TexCreate_FastVRAM | (bNeedsUAV ? TexCreate_UAV : TexCreate_None)), TEXT("Substrate.TopLayerTexture"));
 		}
 
 		// Separated subsurface and rough refraction textures
-		if (bNeedsMaterialBuffer)
+		if (bNeedsMaterialBuffer && !bBlendableGBufferEnabled)
 		{
 			const bool bIsSubstrateOpaqueMaterialRoughRefractionEnabled = IsOpaqueRoughRefractionEnabled(SceneRenderer.ShaderPlatform);
 			const FIntPoint OpaqueRoughRefractionSceneExtent		 = bIsSubstrateOpaqueMaterialRoughRefractionEnabled ? SceneTextureExtent : FIntPoint(4, 4);
@@ -566,7 +559,7 @@ void InitialiseSubstrateFrameSceneData(FRDGBuilder& GraphBuilder, FSceneRenderer
 	const uint32 SliceCountSSS = SUBSTRATE_SSS_DATA_UINT_COUNT;
 	const uint32 SliceCountAdvDebug = IsAdvancedVisualizationEnabled() ? 1 : 0;
 	const uint32 SliceCount = bNeedsMaterialBuffer ? FMath::DivideAndRoundUp(Out.EffectiveMaxBytesPerPixel, 4u) + SliceCountSSS + SliceCountAdvDebug : 1u;
-	if (bNeedsMaterialBuffer)
+	if (bNeedsMaterialBuffer && !bBlendableGBufferEnabled)
 	{
 		FRDGTextureDesc MaterialTextureDesc = FRDGTextureDesc::Create2DArray(SceneTextureExtent, PF_R32_UINT, FClearValueBinding::Transparent, TexCreate_TargetArraySlicesIndependently | TexCreate_DisableDCC | TexCreate_NoFastClear | TexCreate_RenderTargetable | TexCreate_ShaderResource | TexCreate_UAV | TexCreate_FastVRAM, SliceCount, 1, 1);
 		MaterialTextureDesc.FastVRAMPercentage = (1.0f / SliceCount) * 0xFF; // Only allocate the first slice into ESRAM
@@ -613,7 +606,7 @@ void InitialiseSubstrateFrameSceneData(FRDGBuilder& GraphBuilder, FSceneRenderer
 	// Initialized view data
 	for (int32 ViewIndex = 0; ViewIndex < SceneRenderer.Views.Num(); ViewIndex++)
 	{
-		Substrate::InitialiseSubstrateViewData(GraphBuilder, SceneRenderer.Views[ViewIndex], SceneRenderer.GetActiveSceneTexturesConfig(), bNeedsClosureOffsets, bNeedsMaterialBuffer, Out);
+		Substrate::InitialiseSubstrateViewData(GraphBuilder, SceneRenderer.Views[ViewIndex], SceneRenderer.GetActiveSceneTexturesConfig(), bNeedsClosureOffsets, bNeedsMaterialBuffer, bBlendableGBufferEnabled, Out);
 	}
 
 	if (IsSubstrateEnabled())
@@ -694,7 +687,7 @@ static FRDGTextureRef GetDefaultSubstrateMaterialTextureArray(FRDGBuilder& Graph
 	return DefaultSubstrateMaterialTextureArray;
 }
 
-static void BindSubstrateGlobalUniformParameters(FRDGBuilder& GraphBuilder, const FSubstrateViewData* SubstrateViewData, bool bNeedsMaterialBuffer, FSubstrateGlobalUniformParameters& OutSubstrateUniformParameters)
+static void BindSubstrateGlobalUniformParameters(FRDGBuilder& GraphBuilder, const FSubstrateViewData* SubstrateViewData, bool bNeedsMaterialBuffer, bool bBlendableGBufferEnabled, FSubstrateGlobalUniformParameters& OutSubstrateUniformParameters)
 {
 	FSubstrateSceneData* SubstrateSceneData = SubstrateViewData->SceneData;
 	if (IsSubstrateEnabled() && SubstrateSceneData)
@@ -722,7 +715,7 @@ static void BindSubstrateGlobalUniformParameters(FRDGBuilder& GraphBuilder, cons
 			OutSubstrateUniformParameters.ClosureTileBuffer = DefaultBuffer;
 		}
 
-		if (!bNeedsMaterialBuffer)
+		if (!bNeedsMaterialBuffer || bBlendableGBufferEnabled)
 		{
 			check(OutSubstrateUniformParameters.MaterialTextureArray == nullptr);
 			check(SubstrateSceneData->TopLayerTexture == nullptr);
diff --git a/Engine/Source/Runtime/Renderer/Private/Substrate/SubstrateVisualize.cpp b/Engine/Source/Runtime/Renderer/Private/Substrate/SubstrateVisualize.cpp
--- a/Engine/Source/Runtime/Renderer/Private/Substrate/SubstrateVisualize.cpp
+++ b/Engine/Source/Runtime/Renderer/Private/Substrate/SubstrateVisualize.cpp
@@ -1002,13 +1002,13 @@ static void AddVisualizeSystemInfoPasses(FRDGBuilder& GraphBuilder, const FViewI
 	{
 		return;
 	}
+	const bool bSubstrateBlendableGBufferEnabled = IsSubstrateBlendableGBufferEnabled(View.GetShaderPlatform());
 
 	// Force ShaderPrint on.
 	ShaderPrint::SetEnabled(true);
 	ShaderPrint::RequestSpaceForLines(1024);
 	ShaderPrint::RequestSpaceForCharacters(1024);
 
-	const FRDGTextureDesc MaterialBufferDesc = View.SubstrateViewData.SceneData->MaterialTextureArray->Desc;
 	const FShadingEnergyConservationData ShadingEnergyConservationData = ShadingEnergyConservation::GetData(View);
 
 	FSubstrateSystemInfoCS::FParameters* PassParameters = GraphBuilder.AllocParameters<FSubstrateSystemInfoCS::FParameters>();
@@ -1058,7 +1058,12 @@ static void AddVisualizeSystemInfoPasses(FRDGBuilder& GraphBuilder, const FViewI
 	} SystemInformation;
 	SystemInformation.CurrentBytesPerPixel = View.SubstrateViewData.SceneData->EffectiveMaxBytesPerPixel;
 	SystemInformation.ProjectMaxBytesPerPixel = GetBytePerPixel(View.GetShaderPlatform());
-	SystemInformation.MaterialBufferAllocationInBytes = MaterialBufferDesc.Extent.X * MaterialBufferDesc.Extent.Y * MaterialBufferDesc.ArraySize * sizeof(uint32);
+	SystemInformation.MaterialBufferAllocationInBytes = 0;
+	if (!bSubstrateBlendableGBufferEnabled)
+	{
+		const FRDGTextureDesc MaterialBufferDesc = View.SubstrateViewData.SceneData->MaterialTextureArray->Desc;
+		SystemInformation.MaterialBufferAllocationInBytes = MaterialBufferDesc.Extent.X * MaterialBufferDesc.Extent.Y * MaterialBufferDesc.ArraySize * sizeof(uint32);
+	}
 	SystemInformation.ViewsMaxBytesPerPixel = View.SubstrateViewData.SceneData->ViewsMaxBytesPerPixel;
 	SystemInformation.CurrentClosuresPerPixel = View.SubstrateViewData.SceneData->EffectiveMaxClosurePerPixel;
 	SystemInformation.ViewsMaxClosuresPerPixel = View.SubstrateViewData.SceneData->ViewsMaxClosurePerPixel;
@@ -1100,7 +1105,7 @@ static void AddVisualizeSystemInfoPasses(FRDGBuilder& GraphBuilder, const FViewI
 
 		////////////////////////////////////////////////////////////////////////////////////////////////////
 		AddDrawCanvasPass(GraphBuilder, {}, View, FScreenPassRenderTarget(ScreenPassSceneColor.Texture, View.ViewRect, ERenderTargetLoadAction::ELoad),
-			[&View, SystemInfoDebugDataToPrint, SystemInformation](FCanvas& Canvas)
+			[&View, SystemInfoDebugDataToPrint, SystemInformation, bSubstrateBlendableGBufferEnabled](FCanvas& Canvas)
 		{
 			Canvas.SetScaledToRenderTarget(true);
 
@@ -1145,7 +1150,7 @@ static void AddVisualizeSystemInfoPasses(FRDGBuilder& GraphBuilder, const FViewI
 			PrintS(DrawPosX, DrawPosY, TEXT("General"), Yellow);
 			NewLine();
 			{
-				if (IsSubstrateBlendableGBufferEnabled(View.GetShaderPlatform()))
+				if (bSubstrateBlendableGBufferEnabled)
 				{
 					PrintS(DrawPosX, DrawPosY, TEXT("  Using blendable GBuffer with fixed layout."), DarkYellow);
 					NewLine();
