Substrate - Fixes for some console to run.

#tests PC editor and console
#rnx

diff --git a/Engine/Source/Runtime/Renderer/Private/Substrate/SubstrateVisualize.cpp b/Engine/Source/Runtime/Renderer/Private/Substrate/SubstrateVisualize.cpp
--- a/Engine/Source/Runtime/Renderer/Private/Substrate/SubstrateVisualize.cpp
+++ b/Engine/Source/Runtime/Renderer/Private/Substrate/SubstrateVisualize.cpp
@@ -188,6 +188,7 @@ class FMaterialDebugSubstrateTreeCS : public FGlobalShader
 	static bool IsSupported(EShaderPlatform InPlatform, EShaderPermutationFlags InFlags=EShaderPermutationFlags::None)
 	{
 		return IsSubstrateDebugVisualizationSupported(InPlatform, true, InFlags)
+			&& IsAdvancedVisualizationEnabled()
 			&& !IsSubstrateBlendableGBufferEnabled(InPlatform)
 			&& !IsVulkanPlatform(InPlatform)  // SUBSTRATE_TODO Move to CPU debug visualisation and it should then work on all platforms
 			&& !IsMetalPlatform(InPlatform);
@@ -236,6 +237,7 @@ class FMaterialDebugSubstrateTreePS : public FGlobalShader
 	static bool IsSupported(EShaderPlatform InPlatform, EShaderPermutationFlags InFlags=EShaderPermutationFlags::None)
 	{
 		return IsSubstrateDebugVisualizationSupported(InPlatform, true, InFlags)
+			&& IsAdvancedVisualizationEnabled()
 			&& !IsSubstrateBlendableGBufferEnabled(InPlatform)
 			&& !IsVulkanPlatform(InPlatform)  // SUBSTRATE_TODO Move to CPU debug visualisation and it should then work on all platforms
 			&& !IsMetalPlatform(InPlatform);
