[HWRT] Set InstancedViewUniformBuffer when running dynamic geometry update.

#rnx

diff --git a/Engine/Source/Runtime/Renderer/Private/RayTracing/RayTracingDynamicGeometryUpdateManager.cpp b/Engine/Source/Runtime/Renderer/Private/RayTracing/RayTracingDynamicGeometryUpdateManager.cpp
--- a/Engine/Source/Runtime/Renderer/Private/RayTracing/RayTracingDynamicGeometryUpdateManager.cpp
+++ b/Engine/Source/Runtime/Renderer/Private/RayTracing/RayTracingDynamicGeometryUpdateManager.cpp
@@ -246,6 +246,7 @@ FRayTracingDynamicGeometryUpdateManager::~FRayTracingDynamicGeometryUpdateManage
 void FRayTracingDynamicGeometryUpdateManager::Clear()
 {
 	DispatchCommandsPerView = {};
+	InstancedViewUniformBuffers = {};
 
 	// Clear working arrays - keep max size allocated
 	BuildParams.Empty(BuildParams.Max());
@@ -260,6 +261,7 @@ void FRayTracingDynamicGeometryUpdateManager::Clear()
 int64 FRayTracingDynamicGeometryUpdateManager::BeginUpdate()
 {
 	check(DispatchCommandsPerView.IsEmpty());
+	check(InstancedViewUniformBuffers.IsEmpty());
 	check(BuildParams.IsEmpty());
 	check(Segments.IsEmpty());
 	check(ReferencedUniformBuffers.IsEmpty());
@@ -323,6 +325,7 @@ void FRayTracingDynamicGeometryUpdateManager::AddDynamicGeometryToUpdate(
 	FRayTracingDynamicGeometryBuildParams GeometryBuildParams;
 	GeometryBuildParams.ViewUniformBuffer = View->ViewUniformBuffer;
 	GeometryBuildParams.DispatchCommands.Reserve(UpdateParams.MeshBatches.Num());
+	InstancedViewUniformBuffers.FindOrAdd(View->ViewUniformBuffer, View->GetInstancedViewUniformBuffer());
 
 	// Only update when we have mesh batches
 	if (!UpdateParams.MeshBatches.IsEmpty())
@@ -805,7 +808,7 @@ void FRayTracingDynamicGeometryUpdateManager::AddDynamicGeometryUpdatePass(FRDGB
 
 		FRayTracingDynamicGeometryUpdatePassParams* PassParams = GraphBuilder.AllocParameters<FRayTracingDynamicGeometryUpdatePassParams>();
 		PassParams->View.View = TUniformBufferRef<FViewUniformShaderParameters>(ViewDispatchCommands.Key);
-		PassParams->View.InstancedView = TUniformBufferRef<FInstancedViewUniformShaderParameters>(nullptr); // TODO: Is this necessary?
+		PassParams->View.InstancedView = TUniformBufferRef<FInstancedViewUniformShaderParameters>(InstancedViewUniformBuffers[ViewDispatchCommands.Key]);
 		PassParams->Scene = SceneUB;
 
 		// DynamicGeometryScratchBuffer is not directly used in this pass but set so RDG orders passes correctly
diff --git a/Engine/Source/Runtime/Renderer/Public/RayTracingDynamicGeometryUpdateManager.h b/Engine/Source/Runtime/Renderer/Public/RayTracingDynamicGeometryUpdateManager.h
--- a/Engine/Source/Runtime/Renderer/Public/RayTracingDynamicGeometryUpdateManager.h
+++ b/Engine/Source/Runtime/Renderer/Public/RayTracingDynamicGeometryUpdateManager.h
@@ -110,6 +110,7 @@ private:
 
 	// Group dispatch commands per view uniform buffer since it is specified when creating the RDG passes
 	TMap<FRHIUniformBuffer*, TArray<FMeshComputeDispatchCommand>> DispatchCommandsPerView;
+	TMap<FRHIUniformBuffer*, FRHIUniformBuffer*> InstancedViewUniformBuffers;
 	TArray<FRayTracingGeometryBuildParams> BuildParams;
 	TArray<FRayTracingGeometrySegment> Segments;
 
