Substrate - scene color always go in fast ram now.

#tests Console
#rnx

diff --git a/Engine/Source/Runtime/Renderer/Private/SceneRendering.cpp b/Engine/Source/Runtime/Renderer/Private/SceneRendering.cpp
--- a/Engine/Source/Runtime/Renderer/Private/SceneRendering.cpp
+++ b/Engine/Source/Runtime/Renderer/Private/SceneRendering.cpp
@@ -743,12 +743,6 @@ void FFastVramConfig::Update()
 	bDirty |= UpdateBufferFlagFromCVar(CVarFastVRam_DistanceFieldAOScreenGridResources, DistanceFieldAOScreenGridResources);
 	bDirty |= UpdateBufferFlagFromCVar(CVarFastVRam_ForwardLightingCullingResources, ForwardLightingCullingResources);
 	bDirty |= UpdateBufferFlagFromCVar(CVarFastVRam_GlobalDistanceFieldCullGridBuffers, GlobalDistanceFieldCullGridBuffers);
-
-	// When Substrate is enable, remove Scene color from fast VRAM to leave space for material buffer which has more impact on performance
-	if (Substrate::IsSubstrateEnabled() && !IsForwardShadingEnabled(GMaxRHIShaderPlatform))
-	{
-		SceneColor = SceneColor & (~(TexCreate_FastVRAM | TexCreate_FastVRAMPartialAlloc));
-	}
 }
 
 bool FFastVramConfig::UpdateTextureFlagFromCVar(TAutoConsoleVariable<int32>& CVar, ETextureCreateFlags& InOutValue)
