Mobile: Optimizations for simple (particle) lights rendering
Use simple shading for particle lights
Do not use stencil culling for particle lights

diff --git a/Engine/Shaders/Private/MobileDeferredShading.usf b/Engine/Shaders/Private/MobileDeferredShading.usf
--- a/Engine/Shaders/Private/MobileDeferredShading.usf
+++ b/Engine/Shaders/Private/MobileDeferredShading.usf
@@ -316,11 +316,11 @@ void MobileRadialLightPS(
 	FLightAccumulator DirectLighting = (FLightAccumulator)0;
 	FDeferredLightData LightData = InitDeferredLightFromLightParameters(GetRootLightShaderParameters(), RADIAL_LIGHT_TYPE);
 		
+#if !SIMPLE_LIGHT
 	// Affect the light color by any other shadow/transmittance before the lighting is evaluated
-	{
-		half Attenuation = ComputeLightProfileMultiplier(TranslatedWorldPosition, LightData.TranslatedWorldPosition, -LightData.Direction, LightData.Tangent, LightData.IESAtlasIndex);
-		LightData.Color *= Attenuation;
-	}
+	half Attenuation = ComputeLightProfileMultiplier(TranslatedWorldPosition, LightData.TranslatedWorldPosition, -LightData.Direction, LightData.Tangent, LightData.IESAtlasIndex);
+	LightData.Color *= Attenuation;
+#endif
 
 	half Shadow = 1.0;
 #if SUPPORT_SPOTLIGHTS_SHADOW
@@ -363,12 +363,24 @@ void MobileRadialLightPS(
 	}
 #endif
 
+#if SIMPLE_LIGHT
+	FSimpleDeferredLightData SimpleLightData = (FSimpleDeferredLightData)0;
+	SimpleLightData.TranslatedWorldPosition = LightData.TranslatedWorldPosition;
+	SimpleLightData.InvRadius = LightData.InvRadius;
+	SimpleLightData.Color = LightData.Color;
+	SimpleLightData.FalloffExponent = LightData.FalloffExponent;
+	SimpleLightData.bInverseSquared = LightData.bInverseSquared;
+	half3 Color = GetSimpleDynamicLighting(TranslatedWorldPosition, CameraVector, GBuffer.WorldNormal, 1.0, 
+		GBuffer.DiffuseColor * LightData.DiffuseScale, 
+		GBuffer.SpecularColor * LightData.SpecularScale,
+		GBuffer.Roughness, SimpleLightData);
+#else
 	half4 LightAttenuation = half4(1, 1, Shadow, Shadow);
-
 	float SurfaceShadow = 0;
 	DirectLighting = AccumulateDynamicLighting(TranslatedWorldPosition, CameraVector, GBuffer, 1, LightData, LightAttenuation, 0, uint2(0, 0), SurfaceShadow);
-
 	half3 Color = DirectLighting.TotalLight * ComputeLightFunctionMultiplier(TranslatedWorldPosition);
+#endif
+
 	// MobileHDR applies PreExposure in tonemapper
 	Color *= View.PreExposure;
 	Color = SafeGetOutColor(Color);
diff --git a/Engine/Source/Runtime/Renderer/Private/MobileDeferredShadingPass.cpp b/Engine/Source/Runtime/Renderer/Private/MobileDeferredShadingPass.cpp
--- a/Engine/Source/Runtime/Renderer/Private/MobileDeferredShadingPass.cpp
+++ b/Engine/Source/Runtime/Renderer/Private/MobileDeferredShadingPass.cpp
@@ -216,7 +216,8 @@ public:
 	class FRadialLightTypeDim		: SHADER_PERMUTATION_RANGE_INT("RADIAL_LIGHT_TYPE", LIGHT_TYPE_POINT, LIGHT_TYPE_RECT);
 	class FIESProfileDim			: SHADER_PERMUTATION_BOOL("USE_IES_PROFILE");
 	class FSpotLightShadowDim		: SHADER_PERMUTATION_BOOL("SUPPORT_SPOTLIGHTS_SHADOW");
-	using FPermutationDomain = TShaderPermutationDomain<FEnableShadingModelSupport, FRadialLightTypeDim, FIESProfileDim, FSpotLightShadowDim>;
+	class FSimpleLightDim			: SHADER_PERMUTATION_BOOL("SIMPLE_LIGHT");
+	using FPermutationDomain = TShaderPermutationDomain<FEnableShadingModelSupport, FRadialLightTypeDim, FIESProfileDim, FSpotLightShadowDim, FSimpleLightDim>;
 
 	BEGIN_SHADER_PARAMETER_STRUCT(FParameters, )
 		SHADER_PARAMETER_STRUCT(FLightShaderParameters, Light)
@@ -237,6 +238,12 @@ public:
 		}
 
 		FPermutationDomain PermutationVector(Parameters.PermutationId);
+		
+		if (PermutationVector.Get<FSimpleLightDim>() == true && !Parameters.MaterialParameters.bIsDefaultMaterial)
+		{
+			return false;
+		}
+				
 		// Compile out the shader if this permutation gets remapped.
 		if (RemapPermutationVector(PermutationVector, Parameters.Platform) != PermutationVector)
 		{
@@ -258,6 +265,14 @@ public:
 			PermutationVector.Set<FEnableShadingModelSupport>(false);
 		}
 
+		if (PermutationVector.Get<FSimpleLightDim>() == true)
+		{
+			PermutationVector.Set<FEnableShadingModelSupport>(false);
+			PermutationVector.Set<FRadialLightTypeDim>(LIGHT_TYPE_POINT);
+			PermutationVector.Set<FIESProfileDim>(false);
+			PermutationVector.Set<FSpotLightShadowDim>(false);
+		}
+
 		return PermutationVector;
 	}
 
@@ -1038,97 +1053,54 @@ static void RenderSimpleLights(
 	PermutationVectorVS.Set<FDeferredLightVS::FRadialLight>(true);
 	TShaderMapRef<FDeferredLightVS> VertexShader(View.ShaderMap, PermutationVectorVS);
 
-	// Setup stencil mask PSO
-	FGraphicsPipelineStateInitializer GraphicsPSOLightMask;
+	FGraphicsPipelineStateInitializer GraphicsPSOLight;
+	RHICmdList.ApplyCachedRenderTargets(GraphicsPSOLight);
+	// Use additive blending for color
+	GraphicsPSOLight.BlendState = TStaticBlendState<CW_RGB, BO_Add, BF_One, BF_One, BO_Add, BF_One, BF_One>::GetRHI();
+	GraphicsPSOLight.PrimitiveType = PT_TriangleList;
+	GraphicsPSOLight.BoundShaderState.VertexDeclarationRHI = GetVertexDeclarationFVector4();
+	GraphicsPSOLight.BoundShaderState.VertexShaderRHI = VertexShader.GetVertexShader();
+	GraphicsPSOLight.RasterizerState = View.bReverseCulling ? TStaticRasterizerState<FM_Solid, CM_CCW>::GetRHI() : TStaticRasterizerState<FM_Solid, CM_CW>::GetRHI();
+	GraphicsPSOLight.DepthStencilState = TStaticDepthStencilState<
+		false, CF_DepthNearOrEqual,
+		true, CF_NotEqual, SO_Keep, SO_Keep, SO_Keep, // Render where ShadingModel Mask is not zero
+		false, CF_Always, SO_Keep, SO_Keep, SO_Keep,
+		GET_STENCIL_MOBILE_SM_MASK(0xff), 0x00>::GetRHI();
+
+	TShaderRef<FMobileRadialLightFunctionPS> PixelShader;
+	FMobileRadialLightFunctionPS::FPermutationDomain PermutationVector;
+	PermutationVector.Set<FMobileRadialLightFunctionPS::FSimpleLightDim>(true);
 	{
-		RHICmdList.ApplyCachedRenderTargets(GraphicsPSOLightMask);
-		GraphicsPSOLightMask.PrimitiveType = PT_TriangleList;
-		GraphicsPSOLightMask.BlendState = TStaticBlendStateWriteMask<CW_NONE, CW_NONE, CW_NONE, CW_NONE, CW_NONE, CW_NONE, CW_NONE, CW_NONE>::GetRHI();
-		GraphicsPSOLightMask.RasterizerState = View.bReverseCulling ? TStaticRasterizerState<FM_Solid, CM_CCW>::GetRHI() : TStaticRasterizerState<FM_Solid, CM_CW>::GetRHI();
-		// set stencil to 1 where depth test fails
-		GraphicsPSOLightMask.DepthStencilState = TStaticDepthStencilState<
-			false, CF_DepthNearOrEqual,
-			true, CF_Always, SO_Keep, SO_Replace, SO_Keep,		
-			false, CF_Always, SO_Keep, SO_Keep, SO_Keep,
-			0x00, STENCIL_SANDBOX_MASK>::GetRHI();
-		GraphicsPSOLightMask.BoundShaderState.VertexDeclarationRHI = GetVertexDeclarationFVector4();
-		GraphicsPSOLightMask.BoundShaderState.VertexShaderRHI = VertexShader.GetVertexShader();
-		GraphicsPSOLightMask.BoundShaderState.PixelShaderRHI = nullptr;
+		FCachedLightMaterial LightMaterial;
+		GetLightMaterial(DefaultMaterial, nullptr, PermutationVector.ToDimensionValueId(), LightMaterial, PixelShader);
 	}
+	GraphicsPSOLight.BoundShaderState.PixelShaderRHI = PixelShader.GetPixelShader();
+	const uint8 StencilRef = 0;
+	SetGraphicsPipelineState(RHICmdList, GraphicsPSOLight, StencilRef);
 
-	const FMaterialShaderMap* MaterialShaderMap = DefaultMaterial.Material->GetRenderingThreadShaderMap();
-
-	// Setup PSOs we going to use for light rendering 
-	// Do two passes, first masking DefautLit, second masking all other shading models
-	const bool bOnlyDefaultLitInView = IsOnlyDefaultLitShadingModel(View.ShadingModelMaskInView);
-	int32 NumPasses = !bOnlyDefaultLitInView && MobileUsesGBufferCustomData(Scene.GetShaderPlatform()) ? 2 : 1;
-	TShaderRef<FMobileRadialLightFunctionPS> PassPixelShaders[2];
-	FGraphicsPipelineStateInitializer GraphicsPSOLight[2];
-
-	for (int32 PassIndex = 0; PassIndex < NumPasses; ++PassIndex)
+	if (NumViews > 1)
 	{
-		const bool bEnableShadingModelSupport = (PassIndex > 0);
-		
-		RHICmdList.ApplyCachedRenderTargets(GraphicsPSOLight[PassIndex]);
-		// Use additive blending for color
-		GraphicsPSOLight[PassIndex].BlendState = TStaticBlendState<CW_RGB, BO_Add, BF_One, BF_One, BO_Add, BF_One, BF_One>::GetRHI();
-		GraphicsPSOLight[PassIndex].PrimitiveType = PT_TriangleList;
-		GraphicsPSOLight[PassIndex].BoundShaderState.VertexDeclarationRHI = GetVertexDeclarationFVector4();
-		GraphicsPSOLight[PassIndex].BoundShaderState.VertexShaderRHI = VertexShader.GetVertexShader();
-		if (bEnableShadingModelSupport)
-		{
-			SetLocalLightRasterizerAndDepthState<0u, true, /*bEnableShadingModelSupport=*/true>(GraphicsPSOLight[PassIndex], View, FSphere());
-		}
-		else
-		{
-			SetLocalLightRasterizerAndDepthState<0u, true, /*bEnableShadingModelSupport=*/false>(GraphicsPSOLight[PassIndex], View, FSphere());
-		}
-
-		TShaderRef<FMobileRadialLightFunctionPS> PixelShader;
-		FMobileRadialLightFunctionPS::FPermutationDomain PermutationVector;
-		PermutationVector.Set<FMobileRadialLightFunctionPS::FEnableShadingModelSupport>(bEnableShadingModelSupport);
-		PermutationVector.Set<FMobileRadialLightFunctionPS::FRadialLightTypeDim>(LightType_Point);
-		PermutationVector.Set<FMobileRadialLightFunctionPS::FIESProfileDim>(false);
-		PassPixelShaders[PassIndex] = MaterialShaderMap->GetShader<FMobileRadialLightFunctionPS>(PermutationVector);
-		GraphicsPSOLight[PassIndex].BoundShaderState.PixelShaderRHI = PassPixelShaders[PassIndex].GetPixelShader();
+		// set viewports only we we have more than one 
+		// otherwise it is set at the start of the pass
+		RHICmdList.SetViewport(View.ViewRect.Min.X, View.ViewRect.Min.Y, 0.0f, View.ViewRect.Max.X, View.ViewRect.Max.Y, 1.0f);
 	}
-		
+	
 	for (int32 LightIndex = 0; LightIndex < SimpleLights.InstanceData.Num(); LightIndex++)
 	{
 		const FSimpleLightEntry& SimpleLight = SimpleLights.InstanceData[LightIndex];
 		const FSimpleLightPerViewEntry& SimpleLightPerViewData = SimpleLights.GetViewDependentData(LightIndex, ViewIndex, NumViews);
 		const FSphere LightBounds(SimpleLightPerViewData.Position, SimpleLight.Radius);
-
-		if (NumViews > 1)
-		{
-			// set viewports only we we have more than one 
-			// otherwise it is set at the start of the pass
-			RHICmdList.SetViewport(View.ViewRect.Min.X, View.ViewRect.Min.Y, 0.0f, View.ViewRect.Max.X, View.ViewRect.Max.Y, 1.0f);
-		}
-
-		// Render light mask
-		SetGraphicsPipelineState(RHICmdList, GraphicsPSOLightMask, 1);
 		FDeferredLightVS::FParameters ParametersVS = FDeferredLightVS::GetParameters(View, LightBounds);
-		SetShaderParameters(RHICmdList, VertexShader, VertexShader.GetVertexShader(), ParametersVS);
-		StencilingGeometry::DrawSphere(RHICmdList);
 
 		// Render light
-		FMobileRadialLightFunctionPS::FParameters PassParameters;
-		PassParameters.Light = GetSimpleDeferredLightParameters(View, SimpleLight, SimpleLightPerViewData).LightParameters;
-
-		for (int32 PassIndex = 0; PassIndex < NumPasses; ++PassIndex)
-		{
-			const uint8 StencilRef = PassShadingModelStencilValue(PassIndex > 0);
-			SetGraphicsPipelineState(RHICmdList, GraphicsPSOLight[PassIndex], StencilRef);
-
-			SetShaderParameters(RHICmdList, VertexShader, VertexShader.GetVertexShader(), ParametersVS);
-
-			SetShaderParametersMixedPS(RHICmdList, PassPixelShaders[PassIndex], PassParameters, View, DefaultMaterial.MaterialProxy, *DefaultMaterial.Material);
+		FMobileRadialLightFunctionPS::FParameters ParametersPS;
+		ParametersPS.Light = GetSimpleDeferredLightParameters(View, SimpleLight, SimpleLightPerViewData).LightParameters;
+		SetShaderParameters(RHICmdList, VertexShader, VertexShader.GetVertexShader(), ParametersVS);
+		SetShaderParametersMixedPS(RHICmdList, PixelShader, ParametersPS, View, DefaultMaterial.MaterialProxy, *DefaultMaterial.Material);
 
-			// Apply the point or spot light with some approximately bounding geometry,
-			// So we can get speedups from depth testing and not processing pixels outside of the light's influence.
-			StencilingGeometry::DrawSphere(RHICmdList);
-		}
+		// Apply the point or spot light with some approximately bounding geometry,
+		// So we can get speedups from depth testing and not processing pixels outside of the light's influence.
+		StencilingGeometry::DrawSphere(RHICmdList);
 	}
 }
 
