SkyAtmosphere - always finish transmittance texture transition by calling ConvertToExternalAccessTexture when RegisterExternalTexture is always called before.
Fixes post 43064400.

#rnx
#tests PC engine skyatmospehre tests after event reproduced, then fixed.

diff --git a/Engine/Source/Runtime/Renderer/Private/SkyAtmosphereRendering.cpp b/Engine/Source/Runtime/Renderer/Private/SkyAtmosphereRendering.cpp
--- a/Engine/Source/Runtime/Renderer/Private/SkyAtmosphereRendering.cpp
+++ b/Engine/Source/Runtime/Renderer/Private/SkyAtmosphereRendering.cpp
@@ -1777,7 +1777,6 @@ void FSceneRenderer::RenderSkyAtmosphereLookUpTables(FRDGBuilder& GraphBuilder,
 		PendingRDGResources.ViewResources[ViewIndex].SkyAtmosphereCameraAerialPerspectiveVolume = SkyAtmosphereCameraAerialPerspectiveVolume;
 	}
 
-	PendingRDGResources.bTransmittanceLutUpdated = bEvaluateTransmittanceAndMultiScatteringLUTs;
 	PendingRDGResources.TransmittanceLut = TransmittanceLut;
 }
 
@@ -1821,11 +1820,9 @@ void FSkyAtmospherePendingRDGResources::CommitToSceneAndViewUniformBuffers(FRDGB
 		View.SkyAtmosphereCameraAerialPerspectiveVolume = ConvertToExternalAccessTexture(GraphBuilder, ExternalAccessQueue, ViewResources[ViewIndex].SkyAtmosphereCameraAerialPerspectiveVolume, ERHIAccess::SRVMask, ERHIPipeline::All);
 	}
 
-	if (bTransmittanceLutUpdated)
-	{
-		check(TransmittanceLut);
-		SkyInfo.GetTransmittanceLutTexture() = ConvertToExternalAccessTexture(GraphBuilder, ExternalAccessQueue, TransmittanceLut, ERHIAccess::SRVMask, ERHIPipeline::All);
-	}
+	// We always need to run the conversion even if TransmittanceLut has not been written to, as long as RegisterExternalTexture has been called on it. This in order to finish any asynchronous resource transition.
+	check(TransmittanceLut);
+	SkyInfo.GetTransmittanceLutTexture() = ConvertToExternalAccessTexture(GraphBuilder, ExternalAccessQueue, TransmittanceLut, ERHIAccess::SRVMask, ERHIPipeline::All);
 }
 
 FSkyAtmosphereRenderContext::FSkyAtmosphereRenderContext()
diff --git a/Engine/Source/Runtime/Renderer/Private/SkyAtmosphereRendering.h b/Engine/Source/Runtime/Renderer/Private/SkyAtmosphereRendering.h
--- a/Engine/Source/Runtime/Renderer/Private/SkyAtmosphereRendering.h
+++ b/Engine/Source/Runtime/Renderer/Private/SkyAtmosphereRendering.h
@@ -199,7 +199,6 @@ private:
 	FRDGTextureRef RealTimeReflectionCaptureSkyAtmosphereViewLutTexture = nullptr;
 	FRDGTextureRef RealTimeReflectionCaptureCamera360APLutTexture = nullptr;
 
-	bool bTransmittanceLutUpdated = false;
 	FRDGTextureRef TransmittanceLut = nullptr;
 
 	friend class FSceneRenderer;
