Fix use-after-free when material draw events are enabled. When MeshDrawCommand state bucketing is enabled, we can end up with a dangling MaterialRenderProxy pointer in the DebugData. This can happen because multiple MaterialRenderProxies can point to the same MDC when bucketing is enabled, so it's possible for the MDC to outlive the original MaterialRenderProxy it was created with.

diff --git a/Engine/Source/Runtime/Engine/Private/PSOPrecacheValidation.cpp b/Engine/Source/Runtime/Engine/Private/PSOPrecacheValidation.cpp
--- a/Engine/Source/Runtime/Engine/Private/PSOPrecacheValidation.cpp
+++ b/Engine/Source/Runtime/Engine/Private/PSOPrecacheValidation.cpp
@@ -364,6 +364,18 @@ void PSOCollectorStats::CheckFullPipelineStateInCache(
 	const FVertexFactoryType* VFType,
 	const FPrimitiveSceneProxy* PrimitiveSceneProxy,
 	int32 PSOCollectorIndex)
+{
+	const FMaterial* Material = MaterialRenderProxy ? MaterialRenderProxy->GetMaterialNoFallback(GMaxRHIFeatureLevel) : nullptr;
+	CheckFullPipelineStateInCache(Initializer, PSOPrecacheResult, Material, VFType, PrimitiveSceneProxy, PSOCollectorIndex);
+}
+
+void PSOCollectorStats::CheckFullPipelineStateInCache(
+	const FGraphicsPipelineStateInitializer& Initializer,
+	EPSOPrecacheResult PSOPrecacheResult,
+	const FMaterial* Material,
+	const FVertexFactoryType* VFType,
+	const FPrimitiveSceneProxy* PrimitiveSceneProxy,
+	int32 PSOCollectorIndex)
 {
 	bool bCheckMinimalPSOPrecached = true;
 	uint64 PrecacheStateHash = 0;
@@ -387,7 +399,6 @@ void PSOCollectorStats::CheckFullPipelineStateInCache(
 		bool bMinimalPSOPrecached = bCheckMinimalPSOPrecached ? PSOCollectorStats::GetMinimalPSOPrecacheStatsCollector().IsPrecached(Initializer.StatePrecachePSOHash) : true;
 		if (bMinimalPSOPrecached)
 		{
-			const FMaterial* Material = MaterialRenderProxy ? MaterialRenderProxy->GetMaterialNoFallback(GMaxRHIFeatureLevel) : nullptr;
 			LogPSOMissInfo(Initializer, EPSOPrecacheMissType::FullPSO, PSOPrecacheResult, Material, VFType, PrimitiveSceneProxy, EPSOPrecacheType::MeshPass, PSOCollectorIndex, 0);
 		}
 	}
diff --git a/Engine/Source/Runtime/Engine/Public/PSOPrecacheValidation.h b/Engine/Source/Runtime/Engine/Public/PSOPrecacheValidation.h
--- a/Engine/Source/Runtime/Engine/Public/PSOPrecacheValidation.h
+++ b/Engine/Source/Runtime/Engine/Public/PSOPrecacheValidation.h
@@ -97,11 +97,22 @@ namespace PSOCollectorStats
 	extern ENGINE_API void CheckFullPipelineStateInCache(
 		const FGraphicsPipelineStateInitializer& Initializer, 
 		EPSOPrecacheResult PSOPrecacheResult, 
-		const FMaterialRenderProxy* Material, 
+		const FMaterialRenderProxy* MaterialRenderProxy,
 		const FVertexFactoryType* VFType, 
 		const FPrimitiveSceneProxy* PrimitiveSceneProxy, 
 		int32 PSOCollectorIndex);
 
+	/**
+	 * Check if the full PSO is precached and if not output information on what state is missing
+	 */
+	extern ENGINE_API void CheckFullPipelineStateInCache(
+		const FGraphicsPipelineStateInitializer& Initializer,
+		EPSOPrecacheResult PSOPrecacheResult,
+		const FMaterial* Material,
+		const FVertexFactoryType* VFType,
+		const FPrimitiveSceneProxy* PrimitiveSceneProxy,
+		int32 PSOCollectorIndex);
+
 	/**
 	 * Check if the compute PSO is precached and if not output information on what state is missing
 	 */
diff --git a/Engine/Source/Runtime/Renderer/Private/InstanceCulling/InstanceCullingContext.cpp b/Engine/Source/Runtime/Renderer/Private/InstanceCulling/InstanceCullingContext.cpp
--- a/Engine/Source/Runtime/Renderer/Private/InstanceCulling/InstanceCullingContext.cpp
+++ b/Engine/Source/Runtime/Renderer/Private/InstanceCulling/InstanceCullingContext.cpp
@@ -1725,7 +1725,7 @@ void FInstanceCullingContext::SubmitDrawCommands(
 
 #if WANTS_DRAW_MESH_EVENTS
 			RHI_BREADCRUMB_EVENT_CONDITIONAL_F(RHICmdList, GShowMaterialDrawEvents != 0, "MaterialDrawEvent", "%s %s (%u instances)"
-				, VisibleMeshDrawCommand.MeshDrawCommand->GetDebugData().MaterialRenderProxy->GetMaterialName()
+				, VisibleMeshDrawCommand.MeshDrawCommand->GetDebugData().MaterialRenderProxyName
 				, VisibleMeshDrawCommand.MeshDrawCommand->GetDebugData().ResourceName
 				, VisibleMeshDrawCommand.MeshDrawCommand->NumInstances * InstanceFactor
 			);
diff --git a/Engine/Source/Runtime/Renderer/Private/MeshPassProcessor.cpp b/Engine/Source/Runtime/Renderer/Private/MeshPassProcessor.cpp
--- a/Engine/Source/Runtime/Renderer/Private/MeshPassProcessor.cpp
+++ b/Engine/Source/Runtime/Renderer/Private/MeshPassProcessor.cpp
@@ -1283,7 +1283,7 @@ bool FMeshDrawCommand::SubmitDrawBegin(
 		// Try and skip draw if the shaders are not completly loaded yet.
 		if (!MeshPipelineState.BoundShaderState.IsShaderAllLoaded())
 		{
-			UE_LOG(LogRenderer, Log, TEXT("Missing Preload shaders. Material Friendly: %s Material Proxy: %s Static: %s"), *MeshDrawCommand.DebugData.Material->GetFriendlyName(), *MeshDrawCommand.DebugData.MaterialRenderProxy->GetMaterialName(), *MeshDrawCommand.StatsData.CategoryName.ToString());
+			UE_LOG(LogRenderer, Log, TEXT("Missing Preload shaders. Material Friendly: %s Material Proxy: %s Static: %s"), *MeshDrawCommand.DebugData.Material->GetFriendlyName(), *MeshDrawCommand.DebugData.MaterialRenderProxyName, *MeshDrawCommand.StatsData.CategoryName.ToString());
 			MeshPipelineState.BoundShaderState.ForceShaderReload();
 			GraphicsPSOInit.BoundShaderState = MeshPipelineState.BoundShaderState.AsBoundShaderState();
 		}
@@ -1291,7 +1291,7 @@ bool FMeshDrawCommand::SubmitDrawBegin(
 
 #if PSO_PRECACHING_VALIDATE
 #if MESH_DRAW_COMMAND_DEBUG_DATA
-		PSOCollectorStats::CheckFullPipelineStateInCache(GraphicsPSOInit, PSOPrecacheResult, MeshDrawCommand.DebugData.MaterialRenderProxy,
+		PSOCollectorStats::CheckFullPipelineStateInCache(GraphicsPSOInit, PSOPrecacheResult, MeshDrawCommand.DebugData.Material,
 			MeshDrawCommand.DebugData.VertexFactoryType, MeshDrawCommand.DebugData.PrimitiveSceneProxyIfNotUsingStateBuckets, MeshDrawCommand.DebugData.PSOCollectorIndex);		
 #else
 		PSOCollectorStats::CheckFullPipelineStateInCache(GraphicsPSOInit, PSOPrecacheResult, nullptr, nullptr, nullptr, INDEX_NONE);
@@ -1495,7 +1495,7 @@ bool FMeshDrawCommand::SubmitDraw(
 {
 #if WANTS_DRAW_MESH_EVENTS
 	RHI_BREADCRUMB_EVENT_CONDITIONAL_F(RHICmdList, GShowMaterialDrawEvents != 0, "MaterialDrawEvent", "%s %s (%u instances)"
-		, MeshDrawCommand.DebugData.MaterialRenderProxy->GetMaterialName()
+		, MeshDrawCommand.DebugData.MaterialRenderProxyName
 		, MeshDrawCommand.DebugData.ResourceName
 		, MeshDrawCommand.NumInstances * InstanceFactor
 	);
@@ -1604,9 +1604,12 @@ uint64 FMeshDrawCommand::GetPipelineStateSortingKey(const FGraphicsPipelineRende
 #if MESH_DRAW_COMMAND_DEBUG_DATA
 void FMeshDrawCommand::SetDebugData(const FPrimitiveSceneProxy* PrimitiveSceneProxy, const FMaterial* Material, const FMaterialRenderProxy* MaterialRenderProxy, const FMeshProcessorShaders& UntypedShaders, const FVertexFactory* VertexFactory, const FMeshBatch& MeshBatch, int32 PSOCollectorIndex)
 {
+	// The following members are potentially invalid when using bucketing (since multiple proxies can share the same MDC, even in the case of materials)
 	DebugData.PrimitiveSceneProxyIfNotUsingStateBuckets = PrimitiveSceneProxy;
+	DebugData.MaterialRenderProxyIfNotUsingStateBuckets = MaterialRenderProxy;
+
 	DebugData.Material = Material;
-	DebugData.MaterialRenderProxy = MaterialRenderProxy;
+	DebugData.MaterialRenderProxyName = MaterialRenderProxy ? MaterialRenderProxy->GetMaterialName() : TEXT(""); 
 	DebugData.VertexShader = UntypedShaders.VertexShader;
 	DebugData.PixelShader = UntypedShaders.PixelShader;
 	DebugData.LODIndex = MeshBatch.LODIndex;
@@ -1634,7 +1637,7 @@ void FMeshDrawCommand::GetStatsData(FVisibleMeshDrawCommandStatsData& OutVisible
 	OutVisibleStatsData.LODIndex = DebugData.LODIndex;
 	OutVisibleStatsData.SegmentIndex = DebugData.SegmentIndex;
 	OutVisibleStatsData.ResourceName = DebugData.ResourceName;
-	OutVisibleStatsData.MaterialName = DebugData.MaterialRenderProxy->GetMaterialName();
+	OutVisibleStatsData.MaterialName = DebugData.MaterialRenderProxyName;
 #endif
 }
 #endif // MESH_DRAW_COMMAND_STATS
diff --git a/Engine/Source/Runtime/Renderer/Public/MeshPassProcessor.h b/Engine/Source/Runtime/Renderer/Public/MeshPassProcessor.h
--- a/Engine/Source/Runtime/Renderer/Public/MeshPassProcessor.h
+++ b/Engine/Source/Runtime/Renderer/Public/MeshPassProcessor.h
@@ -809,9 +809,15 @@ struct FMeshDrawCommandDebugData
 	// sorted from larger to smaller to minimize padding
 	TShaderRef<FShader> VertexShader;
 	TShaderRef<FShader> PixelShader;
+	// Warning: the following members may not be valid if state buckets are being used (ie if GPUScene is enabled); they can point to invalid or even deleted resources
+	// This is because multiple proxies can share the same MDC. In the case of MaterialRenderProxies, it's rare for them to differ, but it can happen
 	const FPrimitiveSceneProxy* PrimitiveSceneProxyIfNotUsingStateBuckets;
+	const FMaterialRenderProxy* MaterialRenderProxyIfNotUsingStateBuckets;
+
+	// This could in theory be inaccurate when bucketing is active, but it should be correct the vast majority of the time (even if the Proxy ptr changes), so it's good enough for draw events
+	FString MaterialRenderProxyName;
+
 	const FMaterial* Material;
-	const FMaterialRenderProxy* MaterialRenderProxy;
 #if PSO_PRECACHING_VALIDATE
 	// so far these are only used for PSO precaching validation
 	const FVertexFactory* VertexFactory;
diff --git a/Engine/Source/Runtime/SlateRHIRenderer/Private/SlateRHIRenderingPolicy.cpp b/Engine/Source/Runtime/SlateRHIRenderer/Private/SlateRHIRenderingPolicy.cpp
--- a/Engine/Source/Runtime/SlateRHIRenderer/Private/SlateRHIRenderingPolicy.cpp
+++ b/Engine/Source/Runtime/SlateRHIRenderer/Private/SlateRHIRenderingPolicy.cpp
@@ -1330,7 +1330,7 @@ void DrawSlateRenderBatch(
 		{
 			static const int32 MaterialPSOCollectorIndex = FPSOCollectorCreateManager::GetIndex(GetFeatureLevelShadingPath(GMaxRHIFeatureLevel), SlateMaterialPSOCollectorName);
 			// Material Render Proxy is not cached in the state, but could be done for better debugging if needed
-			PSOCollectorStats::CheckFullPipelineStateInCache(State.GraphicsPSOInit, EPSOPrecacheResult::Unknown, nullptr, nullptr, nullptr, MaterialPSOCollectorIndex);
+			PSOCollectorStats::CheckFullPipelineStateInCache(State.GraphicsPSOInit, EPSOPrecacheResult::Unknown, (const FMaterial*)nullptr, nullptr, nullptr, MaterialPSOCollectorIndex);
 		}
 #endif // PSO_PRECACHING_VALIDATE
 
