Fix CIS

diff --git a/Engine/Source/Runtime/Renderer/Private/MeshPassProcessor.cpp b/Engine/Source/Runtime/Renderer/Private/MeshPassProcessor.cpp
--- a/Engine/Source/Runtime/Renderer/Private/MeshPassProcessor.cpp
+++ b/Engine/Source/Runtime/Renderer/Private/MeshPassProcessor.cpp
@@ -1294,7 +1294,7 @@ bool FMeshDrawCommand::SubmitDrawBegin(
 		PSOCollectorStats::CheckFullPipelineStateInCache(GraphicsPSOInit, PSOPrecacheResult, MeshDrawCommand.DebugData.Material,
 			MeshDrawCommand.DebugData.VertexFactoryType, MeshDrawCommand.DebugData.PrimitiveSceneProxyIfNotUsingStateBuckets, MeshDrawCommand.DebugData.PSOCollectorIndex);		
 #else
-		PSOCollectorStats::CheckFullPipelineStateInCache(GraphicsPSOInit, PSOPrecacheResult, nullptr, nullptr, nullptr, INDEX_NONE);
+		PSOCollectorStats::CheckFullPipelineStateInCache(GraphicsPSOInit, PSOPrecacheResult, (const FMaterial*)nullptr, nullptr, nullptr, INDEX_NONE);
 #endif // MESH_DRAW_COMMAND_DEBUG_DATA
 #endif // PSO_PRECACHING_VALIDATE
 
