Mobile: Correctly apply r.Mobile.PixelFogQuality to materials that compute fog in a vertex shader

diff --git a/Engine/Source/Runtime/Renderer/Private/FogRendering.cpp b/Engine/Source/Runtime/Renderer/Private/FogRendering.cpp
--- a/Engine/Source/Runtime/Renderer/Private/FogRendering.cpp
+++ b/Engine/Source/Runtime/Renderer/Private/FogRendering.cpp
@@ -49,6 +49,20 @@ static TAutoConsoleVariable<bool> CVarOrthoFogHeightAdjustment(
 	TEXT("When enabled, uses the Ortho camera height to determine the fog cutoff height"),
 	ECVF_RenderThreadSafe);
 
+static TAutoConsoleVariable<int32> CVarMobilePixelFogQuality(
+	TEXT("r.Mobile.PixelFogQuality"),
+	1,
+	TEXT("Exponentional height fog rendering quality.\n")
+	TEXT("0 - basic per-pixel fog")
+	TEXT("1 - all per-pixel fog features (second fog, directional inscattering, aerial perspective)"),
+	ECVF_Scalability | ECVF_RenderThreadSafe);
+
+bool PlatformUsesBasicFogFeatures(EShaderPlatform ShaderPlatform)
+{
+	// Disables AP, "SecondFog" and "DirectionalLight Inscattering"
+	return IsMobilePlatform(ShaderPlatform) && CVarMobilePixelFogQuality.GetValueOnAnyThread() == 0;
+}
+
 IMPLEMENT_GLOBAL_SHADER_PARAMETER_STRUCT(FFogUniformParameters, "FogStruct");
 
 static void SetupDummyExponentialHeightFogUniformParameters(FFogUniformParameters& OutParameters)
@@ -222,6 +236,7 @@ void FSceneRenderer::InitFogConstants()
 				float CollapsedFogParameter[FExponentialHeightFogSceneInfo::NumFogs];
 				float MaxObserverHeightDifference = 65536.0f;
 				float MaxObserverHeight = FLT_MAX;
+				const bool bSimpleFog = PlatformUsesBasicFogFeatures(View.GetShaderPlatform());
 
 				const bool bUsingOrthoHeightAdjustment = !View.IsPerspectiveProjection() && CVarOrthoFogHeightAdjustment.GetValueOnRenderThread();				
 				if (bUsingOrthoHeightAdjustment)
@@ -254,7 +269,10 @@ void FSceneRenderer::InitFogConstants()
 				}
 
 				View.ExponentialFogParameters = FVector4f(CollapsedFogParameter[0], FogInfo.FogData[0].HeightFalloff, MaxObserverHeight, FogInfo.StartDistance);
-				View.ExponentialFogParameters2 = FVector4f(CollapsedFogParameter[1], FogInfo.FogData[1].HeightFalloff, FogInfo.FogData[1].Density, FogInfo.FogData[1].Height);
+				if (!bSimpleFog)
+				{
+					View.ExponentialFogParameters2 = FVector4f(CollapsedFogParameter[1], FogInfo.FogData[1].HeightFalloff, FogInfo.FogData[1].Density, FogInfo.FogData[1].Height);
+				}
 				View.ExponentialFogColor = FVector3f(FogInfo.FogColor.R, FogInfo.FogColor.G, FogInfo.FogColor.B);
 				View.FogMaxOpacity = FogInfo.FogMaxOpacity;
 				View.ExponentialFogParameters3 = FVector4f(FogInfo.FogData[0].Density, FogInfo.FogData[0].Height, FogInfo.InscatteringColorCubemap ? 1.0f : 0.0f, FogInfo.FogCutoffDistance);
@@ -282,7 +300,7 @@ void FSceneRenderer::InitFogConstants()
 					View.InscatteringLightDirection = -SunLight->Proxy->GetDirection();
 					View.DirectionalInscatteringColor = FogInfo.DirectionalInscatteringColor * SunLight->Proxy->GetColor().GetLuminance();
 				}
-				View.bUseDirectionalInscattering = SunLight != nullptr;
+				View.bUseDirectionalInscattering = !bSimpleFog && SunLight != nullptr;
 				View.bEnableVolumetricFog = FogInfo.bEnableVolumetricFog;
 				View.VolumetricFogStartDistance = FogInfo.VolumetricFogStartDistance;
 				View.VolumetricFogNearFadeInDistanceInv = FogInfo.VolumetricFogNearFadeInDistance > 0.0f ? (1.0f / FogInfo.VolumetricFogNearFadeInDistance) : 100000000.0f;
diff --git a/Engine/Source/Runtime/Renderer/Private/MobileFogRendering.cpp b/Engine/Source/Runtime/Renderer/Private/MobileFogRendering.cpp
--- a/Engine/Source/Runtime/Renderer/Private/MobileFogRendering.cpp
+++ b/Engine/Source/Runtime/Renderer/Private/MobileFogRendering.cpp
@@ -22,20 +22,15 @@
 
 DECLARE_GPU_STAT(MobileFog);
 
-static TAutoConsoleVariable<int32> CVarPixelFogQuality(
-	TEXT("r.Mobile.PixelFogQuality"),
-	1,
-	TEXT("Exponentional height fog rendering quality.\n")
-	TEXT("0 - basic per-pixel fog")
-	TEXT("1 - all per-pixel fog features (second fog, directional inscattering, aerial perspective)"),
-	ECVF_Scalability | ECVF_RenderThreadSafe);
-
 static TAutoConsoleVariable<int32> CVarPixelFogDepthTest(
 	TEXT("r.Mobile.PixelFogDepthTest"),
 	1,
 	TEXT("Whether to use depth and stencil tests for fog rendering"),
 	ECVF_RenderThreadSafe);
 
+
+extern bool PlatformUsesBasicFogFeatures(EShaderPlatform ShaderPlatform);
+
 class FMobileFogVS : public FGlobalShader
 {
 	DECLARE_GLOBAL_SHADER(FMobileFogVS);
@@ -148,9 +143,9 @@ void FMobileSceneRenderer::RenderFog(FRHICommandList& RHICmdList, const FViewInf
 		return;
 	}
 
-	const int32 PixelFogQuality = CVarPixelFogQuality.GetValueOnRenderThread();
+	const bool bSimpleFog = PlatformUsesBasicFogFeatures(View.GetShaderPlatform());
 	const bool bUseHeightFog = Scene->ExponentialFogs.Num() > 0 && ShouldRenderFog(*View.Family);
-	const bool bUseAerialPerspective = PixelFogQuality > 0 && ShouldRenderSkyAtmosphere(Scene, View.Family->EngineShowFlags);
+	const bool bUseAerialPerspective = !bSimpleFog && ShouldRenderSkyAtmosphere(Scene, View.Family->EngineShowFlags);
 
 	if (!bUseAerialPerspective && !bUseHeightFog)
 	{
@@ -195,8 +190,8 @@ void FMobileSceneRenderer::RenderFog(FRHICommandList& RHICmdList, const FViewInf
 
 	const bool bUseFogStartDistance = View.ExponentialFogParameters.W > 0;
 	const bool bUseFogInscatteringColorCubemap = View.FogInscatteringColorCubemap != nullptr;
-	const bool bUseFogSecondTerm = PixelFogQuality > 0 && (View.ExponentialFogParameters2.X > 0);
-	const bool bUseFogDirectionalInscatering = PixelFogQuality > 0 && !bUseFogInscatteringColorCubemap && (View.DirectionalInscatteringColor.GetLuminance() > 0 || View.bUseDirectionalInscattering);
+	const bool bUseFogSecondTerm = !bSimpleFog && (View.ExponentialFogParameters2.Z > 0); // Fog Density in Z
+	const bool bUseFogDirectionalInscatering = !bSimpleFog && !bUseFogInscatteringColorCubemap && (View.DirectionalInscatteringColor.GetLuminance() > 0 || View.bUseDirectionalInscattering);
 	
 	const bool bShouldRenderVolumetricFog = ShouldRenderVolumetricFog();
 	bFogHasComposedLocalFogVolumes = bShouldRenderVolumetricFog && ShouldRenderLocalFogVolume(Scene, *View.Family) && ShouldRenderLocalFogVolumeInVolumetricFog(Scene, *View.Family, bShouldRenderVolumetricFog);
diff --git a/Engine/Source/Runtime/Renderer/Private/SkyAtmosphereRendering.cpp b/Engine/Source/Runtime/Renderer/Private/SkyAtmosphereRendering.cpp
--- a/Engine/Source/Runtime/Renderer/Private/SkyAtmosphereRendering.cpp
+++ b/Engine/Source/Runtime/Renderer/Private/SkyAtmosphereRendering.cpp
@@ -411,10 +411,10 @@ static auto GetSizeAndInvSize = [](int32 Width, int32 Height)
 	return FVector4f(FWidth, FHeight, 1.0f / FWidth, 1.0f / FHeight);
 };
 
-bool ShouldApplyAerialPerspective(EShaderPlatform ShaderPlatform)
+static bool ShouldApplyAerialPerspective(EShaderPlatform ShaderPlatform)
 {
-	static auto* MobilePixelFogQuality = IConsoleManager::Get().FindTConsoleVariableDataInt(TEXT("r.Mobile.PixelFogQuality"));
-	return !IsMobilePlatform(ShaderPlatform) || (MobilePixelFogQuality && (MobilePixelFogQuality->GetValueOnAnyThread() == 1));
+	extern bool PlatformUsesBasicFogFeatures(EShaderPlatform ShaderPlatform);
+	return PlatformUsesBasicFogFeatures(ShaderPlatform) == false;
 }
 
 void SetupSkyAtmosphereViewSharedUniformShaderParameters(const FViewInfo& View, const FSkyAtmosphereSceneProxy& SkyAtmosphereProxy, FSkyAtmosphereViewSharedUniformShaderParameters& OutParameters)
