Fixing multiple cases of binding a 4 component texture format to a shader input expecting a single component.

diff --git a/Engine/Source/Runtime/Engine/Private/SceneManagement.cpp b/Engine/Source/Runtime/Engine/Private/SceneManagement.cpp
--- a/Engine/Source/Runtime/Engine/Private/SceneManagement.cpp
+++ b/Engine/Source/Runtime/Engine/Private/SceneManagement.cpp
@@ -1245,7 +1245,7 @@ FViewUniformShaderParameters::FViewUniformShaderParameters()
 
 	GlobalDistanceFieldPageAtlasTexture = BlackVolume;
 	GlobalDistanceFieldCoverageAtlasTexture = BlackVolume;
-	GlobalDistanceFieldPageTableTexture = BlackVolume;
+	GlobalDistanceFieldPageTableTexture = BlackUintVolume;
 	GlobalDistanceFieldMipTexture = BlackVolume;
 
 	GlobalDistanceFieldPageAtlasTextureSampler = TStaticSamplerState<SF_Trilinear, AM_Wrap, AM_Wrap, AM_Wrap>::GetRHI();
diff --git a/Engine/Source/Runtime/Renderer/Private/Substrate/Substrate.cpp b/Engine/Source/Runtime/Renderer/Private/Substrate/Substrate.cpp
--- a/Engine/Source/Runtime/Renderer/Private/Substrate/Substrate.cpp
+++ b/Engine/Source/Runtime/Renderer/Private/Substrate/Substrate.cpp
@@ -710,7 +710,7 @@ static void BindSubstrateGlobalUniformParameters(FRDGBuilder& GraphBuilder, cons
 		{
 			const FRDGSystemTextures& SystemTextures = FRDGSystemTextures::Get(GraphBuilder);
 			FRDGBufferSRVRef DefaultBuffer = GraphBuilder.CreateSRV(GSystemTextures.GetDefaultBuffer(GraphBuilder, 4, 0u), PF_R32_UINT);
-			OutSubstrateUniformParameters.ClosureOffsetTexture = SystemTextures.Black;
+			OutSubstrateUniformParameters.ClosureOffsetTexture = GSystemTextures.GetZeroUIntDummy(GraphBuilder);
 			OutSubstrateUniformParameters.ClosureTileCountBuffer = DefaultBuffer;
 			OutSubstrateUniformParameters.ClosureTileBuffer = DefaultBuffer;
 		}
@@ -747,7 +747,7 @@ static void BindSubstrateGlobalUniformParameters(FRDGBuilder& GraphBuilder, cons
 		OutSubstrateUniformParameters.MaterialTextureArray = DefaultTextureArray;
 		OutSubstrateUniformParameters.TopLayerTexture = SystemTextures.DefaultNormal8Bit;
 		OutSubstrateUniformParameters.OpaqueRoughRefractionTexture = SystemTextures.Black;
-		OutSubstrateUniformParameters.ClosureOffsetTexture = SystemTextures.Black;
+		OutSubstrateUniformParameters.ClosureOffsetTexture = GSystemTextures.GetZeroUIntDummy(GraphBuilder);
 		OutSubstrateUniformParameters.ClosureTileCountBuffer = DefaultBuffer;
 		OutSubstrateUniformParameters.ClosureTileBuffer = DefaultBuffer;
 		OutSubstrateUniformParameters.SampledMaterialTexture = SystemTextures.Black;
diff --git a/Engine/Source/Runtime/Renderer/Public/GlobalDistanceFieldParameters.h b/Engine/Source/Runtime/Renderer/Public/GlobalDistanceFieldParameters.h
--- a/Engine/Source/Runtime/Renderer/Public/GlobalDistanceFieldParameters.h
+++ b/Engine/Source/Runtime/Renderer/Public/GlobalDistanceFieldParameters.h
@@ -158,7 +158,7 @@ public:
 			FRHIBatchedShaderParameters& BatchedParameters = RHICmdList.GetScratchShaderParameters();
 
 			SetTextureParameter(BatchedParameters, GlobalDistanceFieldPageAtlasTexture, ParameterData.PageAtlasTexture ? ParameterData.PageAtlasTexture : GBlackVolumeTexture->TextureRHI.GetReference());
-			SetTextureParameter(BatchedParameters, GlobalDistanceFieldPageTableTexture, ParameterData.PageTableTexture ? ParameterData.PageTableTexture : GBlackVolumeTexture->TextureRHI.GetReference());
+			SetTextureParameter(BatchedParameters, GlobalDistanceFieldPageTableTexture, ParameterData.PageTableTexture ? ParameterData.PageTableTexture : GBlackUintVolumeTexture->TextureRHI.GetReference());
 			SetTextureParameter(BatchedParameters, GlobalDistanceFieldMipTexture, ParameterData.MipTexture ? ParameterData.MipTexture : GBlackVolumeTexture->TextureRHI.GetReference());
 
 			SetShaderValueArray(BatchedParameters, GlobalVolumeTranslatedCenterAndExtent, ParameterData.TranslatedCenterAndExtent, GlobalDistanceField::MaxClipmaps);
