Fix crash when selecting a groom which was not rendered before.

diff --git a/Engine/Source/Runtime/Renderer/Private/HairStrands/HairStrandsVisibility.cpp b/Engine/Source/Runtime/Renderer/Private/HairStrands/HairStrandsVisibility.cpp
--- a/Engine/Source/Runtime/Renderer/Private/HairStrands/HairStrandsVisibility.cpp
+++ b/Engine/Source/Runtime/Renderer/Private/HairStrands/HairStrandsVisibility.cpp
@@ -4404,6 +4404,12 @@ void DrawHitProxies(
 		return;
 	}
 
+	// TransientResources haven't been allocated by the main path. Skipping hit proxy rendering until then.
+	if (Scene.HairStrandsSceneData.TransientResources == nullptr)
+	{
+		return;
+	}
+
 	// The hit proxy view reuse data generated by regular view. This means it assumes LOD selection, simulation, and interpolation has run. 
 	// Geometry won't be updated for proxy view
 	const FIntPoint Resolution = HitProxyTexture->Desc.Extent;
