Fixed unity build on FMaterialCacheStackProviderRenderProxy

diff --git a/Engine/Source/Runtime/Renderer/Internal/MaterialCache/MaterialCacheVirtualTextureRenderProxy.h b/Engine/Source/Runtime/Renderer/Internal/MaterialCache/MaterialCacheVirtualTextureRenderProxy.h
--- a/Engine/Source/Runtime/Renderer/Internal/MaterialCache/MaterialCacheVirtualTextureRenderProxy.h
+++ b/Engine/Source/Runtime/Renderer/Internal/MaterialCache/MaterialCacheVirtualTextureRenderProxy.h
@@ -3,12 +3,11 @@
 #pragma once
 
 #include "MaterialCache/MaterialCacheVirtualTextureInterface.h"
+#include "MaterialCache/MaterialCacheStackProvider.h"
 #include "Templates/UniquePtr.h"
 #include "TextureResource.h"
 #include "Math/MathFwd.h"
 
-class FMaterialCacheStackProviderRenderProxy;
-
 class UMaterialCacheVirtualTextureRenderProxy : public FMaterialCacheVirtualTextureRenderProxyInterface
 {
 public: /** UMaterialCacheVirtualTextureRenderProxyInterface */
