DoOcclusionQueries checks EngineShowFlags.DisableOcclusionQueries consistently.

diff --git a/Engine/Source/Runtime/Renderer/Private/SceneVisibility.cpp b/Engine/Source/Runtime/Renderer/Private/SceneVisibility.cpp
--- a/Engine/Source/Runtime/Renderer/Private/SceneVisibility.cpp
+++ b/Engine/Source/Runtime/Renderer/Private/SceneVisibility.cpp
@@ -351,11 +351,16 @@ static FAutoConsoleVariableRef CVarOcclusionCullEnable(
 	ECVF_RenderThreadSafe
 );
 
-bool FSceneRenderer::DoOcclusionQueries() const
+bool DoOcclusionQueries(const FSceneViewFamily& ViewFamily)
 {
 	return GOcclusionCullEnabled && !ViewFamily.EngineShowFlags.DisableOcclusionQueries;
 }
 
+bool FSceneRenderer::DoOcclusionQueries() const
+{
+	return ::DoOcclusionQueries(ViewFamily);
+}
+
 static int32 GRelevanceNumPrimitivesPerPacket = 0;
 static FAutoConsoleVariableRef CVarRelevanceNumPrimitivesPerPacket(
 	TEXT("r.Visibility.Relevance.NumPrimitivesPerPacket"),
@@ -3517,7 +3522,7 @@ FVisibilityViewPacket::FVisibilityViewPacket(FVisibilityTaskData& InTaskData, FS
 	, ViewIndex(InViewIndex)
 	, ViewElementPDI(&View, nullptr, nullptr)
 {
-	if (ViewState && !View.Family->EngineShowFlags.Wireframe && GOcclusionCullEnabled)
+	if (ViewState && !View.Family->EngineShowFlags.Wireframe && DoOcclusionQueries(*View.Family))
 	{
 		if (TaskConfig.Schedule == EVisibilityTaskSchedule::Parallel)
 		{
@@ -3631,7 +3636,7 @@ void FVisibilityViewPacket::BeginInitVisibility()
 	// Allocate the view's visibility maps.
 	View.PrimitiveVisibilityMap.Init(!bShouldVisibilityCull, Scene.Primitives.Num());
 	View.PrimitiveRayTracingVisibilityMap.Init(false, Scene.Primitives.Num());
-	View.PrimitiveDefinitelyUnoccludedMap.Init(!GOcclusionCullEnabled, Scene.Primitives.Num());
+	View.PrimitiveDefinitelyUnoccludedMap.Init(!DoOcclusionQueries(*View.Family), Scene.Primitives.Num());
 	View.PotentiallyFadingPrimitiveMap.Init(false, Scene.Primitives.Num());
 	View.PrimitiveFadeUniformBuffers.AddZeroed(Scene.Primitives.Num());
 	View.PrimitiveFadeUniformBufferMap.Init(false, Scene.Primitives.Num());
@@ -4933,7 +4938,7 @@ void FSceneRenderer::PrepareViewStateForVisibility(const FSceneTexturesConfig& S
 		const bool bIsHitTesting = ViewFamily.EngineShowFlags.HitProxies;
 		// Don't test occlusion queries in collision viewmode as they can be bigger then the rendering bounds.
 		const bool bCollisionView = ViewFamily.EngineShowFlags.CollisionVisibility || ViewFamily.EngineShowFlags.CollisionPawn;
-		if (GIsHighResScreenshot || !DoOcclusionQueries() || bIsHitTesting || bCollisionView || ViewFamily.EngineShowFlags.DisableOcclusionQueries)
+		if (GIsHighResScreenshot || !DoOcclusionQueries() || bIsHitTesting || bCollisionView)
 		{
 			View.bDisableQuerySubmissions = true;
 			View.bIgnoreExistingQueries = true;
