Raytracing: Fix (light card) mesh creation crash due to null lumen raytracing data shader binding table.

Minimal hotfix until larger refactor of SetupRayTracingPipelineStatesAndSBT() can be done to prevent the current invalid state, see UE-294890.

#rnx

diff --git a/Engine/Source/Runtime/Renderer/Private/DeferredShadingRenderer.cpp b/Engine/Source/Runtime/Renderer/Private/DeferredShadingRenderer.cpp
--- a/Engine/Source/Runtime/Renderer/Private/DeferredShadingRenderer.cpp
+++ b/Engine/Source/Runtime/Renderer/Private/DeferredShadingRenderer.cpp
@@ -1110,7 +1110,7 @@ void FDeferredShadingSceneRenderer::SetupRayTracingRenderingData(FRDGBuilder& Gr
 		{
 			if (GRHISupportsRayTracingShaders || GRHISupportsInlineRayTracing)
 			{
-				if (ReferenceView.LumenRayTracingData.PipelineState)
+				if (ReferenceView.LumenRayTracingData.PipelineState && ReferenceView.LumenRayTracingData.ShaderBindingTable)
 				{
 					RHICmdList.SetRayTracingMissShader(ReferenceView.LumenRayTracingData.ShaderBindingTable, RAY_TRACING_MISS_SHADER_SLOT_DEFAULT, ReferenceView.LumenRayTracingData.PipelineState, 0 /* MissShaderPipelineIndex */, 0, nullptr, 0);
 				}
