Disable render scene async cleanup as a temporary fix for:

PLAY-87990
PLAY-88429

diff --git a/Engine/Source/Runtime/Renderer/Private/SceneRenderBuilder.cpp b/Engine/Source/Runtime/Renderer/Private/SceneRenderBuilder.cpp
--- a/Engine/Source/Runtime/Renderer/Private/SceneRenderBuilder.cpp
+++ b/Engine/Source/Runtime/Renderer/Private/SceneRenderBuilder.cpp
@@ -295,7 +295,7 @@ static void DestroySplitScreenDebugViewFamilies(TConstArrayView<const FSceneView
 
 //////////////////////////////////////////////////////////////////////////
 
-static int32 GSceneRenderCleanUpMode = 1;
+static int32 GSceneRenderCleanUpMode = 0;
 static FAutoConsoleVariableRef CVarSceneRenderCleanUpMode(
 	TEXT("r.SceneRender.CleanUpMode"),
 	GSceneRenderCleanUpMode,
