Add initial mesh shader support to render curve.

diff --git a/Engine/Plugins/Runtime/HairStrands/Source/HairStrandsCore/Private/GroomBuilder.cpp b/Engine/Plugins/Runtime/HairStrands/Source/HairStrandsCore/Private/GroomBuilder.cpp
--- a/Engine/Plugins/Runtime/HairStrands/Source/HairStrandsCore/Private/GroomBuilder.cpp
+++ b/Engine/Plugins/Runtime/HairStrands/Source/HairStrandsCore/Private/GroomBuilder.cpp
@@ -14,6 +14,7 @@
 #include "ShaderParameterStruct.h"
 #include "SceneView.h"
 #include "Math/MathFwd.h"
+#include "RenderCurveDefinitions.h"
 
 #include "Async/ParallelFor.h"
 #include "HAL/IConsoleManager.h"
@@ -57,7 +58,7 @@ FString FGroomBuilder::GetVersion()
 	FString Out = TEXT("v20_debug4");
 	if (IsRenderCurveEnabled())
 	{
-		Out += TEXT("_Curve3e");
+		Out += TEXT("_Curve4a");
 	}
 	return Out;
 }
@@ -4298,10 +4299,10 @@ struct FPackedCluster
 	{
 		// 16b
 		FVector3f Center = FVector3f::ZeroVector;
-		uint8 CurveCount = 0;
+		uint8 CurvePerCluster = 0;
 		uint8 PointPerCurve = 0;
+		uint8 CurveCount = 0;
 		uint8 Pad0 = 0;
-		uint8 Pad1 = 0;
 
 		// 16b 
 		FFloat16 MaxLength = 0;
@@ -4346,6 +4347,8 @@ void InternalBuildRenderCurveResourceBulkData(const FHairStrandsDatas& In, FRend
 	const uint32 PointPerCurve = 16u; // TODO: Use histogram for picking correct value? Pixel metric needs to take that into account
 	const uint32 ClusterCount = FMath::DivideAndRoundUp(In.GetNumCurves(), CurvePerCluster);
 
+	static_assert(CurvePerCluster == RENDERCURVE_CURVE_PER_CLUSTER);
+
 	const FHairStrandsCurves& Curves = In.StrandsCurves;
 	const FHairStrandsPoints& Points = In.StrandsPoints;
 
@@ -4365,6 +4368,7 @@ void InternalBuildRenderCurveResourceBulkData(const FHairStrandsDatas& In, FRend
 	{
 		uint32 PointPerCurve = 0;
 		uint32 CurvePerCluster = 0;
+		uint32 CurveCount = 0;
 		float MaxRadius = 0.f;
 		float MaxLength = 0.f;
 		FVector3f BoundCenter = FVector3f::ZeroVector;
@@ -4430,6 +4434,7 @@ void InternalBuildRenderCurveResourceBulkData(const FHairStrandsDatas& In, FRend
 	{
 		FCluster& Cluster = Clusters[ClusterIndex];
 		Cluster.Curves.SetNum(CurvePerCluster);
+		Cluster.CurveCount   	= 0;
 		Cluster.PointPerCurve   = PointPerCurve;
 		Cluster.CurvePerCluster = CurvePerCluster;
 		Cluster.MaxRadius       = 0.f;
@@ -4452,6 +4457,9 @@ void InternalBuildRenderCurveResourceBulkData(const FHairStrandsDatas& In, FRend
 				continue;
 			}
 
+			// Valid curve count
+			++Cluster.CurveCount;
+
 			const uint32 ClusterCurveIndex = FMath::Min(CurveIt, ClusterCurveCount-1);
 			const uint32 CurveIndex = ClusterCurveIndices[ClusterIndex][ClusterCurveIndex];
 
@@ -4536,7 +4544,8 @@ void InternalBuildRenderCurveResourceBulkData(const FHairStrandsDatas& In, FRend
 		const FCluster& Cluster = Clusters[ClusterIndex];
 		FPackedCluster& PackedCluster = PackedClusters[ClusterIndex];
 		PackedCluster.Header.Center           = Cluster.Curves[0].Points[0]; // First point for now
-		PackedCluster.Header.CurveCount       = Cluster.CurvePerCluster;
+		PackedCluster.Header.CurveCount       = Cluster.CurveCount;
+		PackedCluster.Header.CurvePerCluster  = Cluster.CurvePerCluster;
 		PackedCluster.Header.PointPerCurve    = Cluster.PointPerCurve;
 		PackedCluster.Header.MaxLength        = Cluster.MaxLength;
 		PackedCluster.Header.MaxRadius        = Cluster.MaxRadius;
@@ -4554,7 +4563,7 @@ void InternalBuildRenderCurveResourceBulkData(const FHairStrandsDatas& In, FRend
 		PackedCluster.Header.LocalBoundExtentY = FFloat16(Cluster.BoundExtent.Y);
 		PackedCluster.Header.LocalBoundExtentZ = FFloat16(Cluster.BoundExtent.Z);
 
-		PackedCluster.Data.PackedPoints.SetNum(PackedCluster.Header.CurveCount * PackedCluster.Header.PointPerCurve);
+		PackedCluster.Data.PackedPoints.SetNum(PackedCluster.Header.CurvePerCluster * PackedCluster.Header.PointPerCurve);
 		for (uint32 PointIt = 0; PointIt < Cluster.PointPerCurve; ++PointIt)
 		{
 			for (uint32 CurveIt = 0; CurveIt < Cluster.CurvePerCluster; ++CurveIt)
diff --git a/Engine/Shaders/Private/HairStrands/RenderCurveRaster.usf b/Engine/Shaders/Private/HairStrands/RenderCurveRaster.usf
--- a/Engine/Shaders/Private/HairStrands/RenderCurveRaster.usf
+++ b/Engine/Shaders/Private/HairStrands/RenderCurveRaster.usf
@@ -5,6 +5,7 @@
 #include "/Engine/Public/Platform.ush"
 #include "/Engine/Private/Common.ush"
 #include "/Engine/Private/SceneData.ush"
+#include "/Engine/Shared/RenderCurveDefinitions.h"
 #include "../Nanite/NaniteHZBCull.ush"
 #include "../ColorMap.ush"
 
@@ -22,15 +23,13 @@
 #define BIN_RASTER_INV_SIZE (1.f / float(RASTER_TILE_SIZE))
 #define RASTER_TILE_SIZE_DIV_AS_SHIFT 3
 
-// 
-#define NUM_CURVE_PER_CLUSTER 64
-#define RENDER_CURVE_PRIMITIVE_DATA_STRIDE_IN_BYTES 16
-
 #define FPackedSegmentType uint4
 
+#if COMPUTESHADER
 #ifndef THREADGROUP_SIZE
 #error THREADGROUP_SIZE needs to be defined
 #endif
+#endif
 
 /////////////////////////////////////////////////////////////////////////////////////////////////////////////////////////////
 // Debug
@@ -102,7 +101,7 @@ FRenderCurveInstanceData GetRenderCurveInstanceData(uint InPrimitiveIndex)
 	FRenderCurveInstanceData Out = (FRenderCurveInstanceData)0;
 	if (InPrimitiveIndex < Scene.RenderCurve.InstanceCount)
 	{
-		const uint4 Packed = Scene.RenderCurve.RenderCurveInstanceData.Load4(InPrimitiveIndex * RENDER_CURVE_PRIMITIVE_DATA_STRIDE_IN_BYTES);
+		const uint4 Packed = Scene.RenderCurve.RenderCurveInstanceData.Load4(InPrimitiveIndex * RENDERCURVE_PRIMITIVE_DATA_STRIDE_IN_BYTES);
 		Out.bIsValid = true;
 		Out.PersistentPrimitiveIndex = Packed.x;
 		Out.InstanceSceneDataOffset  = Packed.y;
@@ -126,6 +125,7 @@ struct FClusterHeader
 {
 	float3 Center;
 	uint CurveCount;
+	uint CurvePerCluster;
 	uint PointPerCurve;
 	float MaxLength;
 	float MaxRadius;
@@ -142,8 +142,9 @@ FClusterHeader GetClusterHeader(uint InClusterIndex)
 		const uint4 Packed1 = Scene.RenderCurve.ClusterData.Load4(InClusterIndex * Scene.RenderCurve.MaxClusterStrideInBytes + 16u);
 
 		Out.Center           = asfloat(Packed0.xyz);
-		Out.CurveCount       = BitFieldExtractU32(Packed0.w, 8, 0);
+		Out.CurvePerCluster  = BitFieldExtractU32(Packed0.w, 8, 0);
 		Out.PointPerCurve    = BitFieldExtractU32(Packed0.w, 8, 8);
+		Out.CurveCount		 = BitFieldExtractU32(Packed0.w, 8, 16);
 
 		Out.MaxLength        = f16tof32(BitFieldExtractU32(Packed1.x, 16,  0));
 		Out.MaxRadius        = f16tof32(BitFieldExtractU32(Packed1.x, 16, 16));
@@ -213,7 +214,7 @@ FCurvePoint GetClusterPoint(FClusterHeader Header, uint ClusterIndex, uint Curve
 {
 	const uint HeaderStride = 32u;
 	const uint PointStride  = 8u;
-	const uint PointOffsetInBytes = (Header.CurveCount * PointIndex + CurveIndex) * PointStride;
+	const uint PointOffsetInBytes = (Header.CurvePerCluster * PointIndex + CurveIndex) * PointStride;
 	const uint2 Packed = Scene.RenderCurve.ClusterData.Load2(ClusterIndex * Scene.RenderCurve.MaxClusterStrideInBytes + HeaderStride + PointOffsetInBytes);
 	return UnpackCurvePoint(Packed, Header.Center, Header.MaxRadius);
 }
@@ -1531,7 +1532,7 @@ void CompactionCS(uint DispatchThreadID : SV_DispatchThreadID, uint GroupThreadI
 
 /////////////////////////////////////////////////////////////////////////////////////////////////////////////////////////////
 
-#ifdef RasterizerCS
+#ifdef SWRasterizerCS
 
 #define MAX_THREAD_ITERATION_COUNT 4096
 #define RASTER_TILE_COUNT_2D (BIN_TILE_SIZE / RASTER_TILE_SIZE)           // = 4x4
@@ -1698,7 +1699,7 @@ FShaderPrintContext GetShaderPrintContextPerRasterThread(FCoord Coord)
 #endif
 
 [numthreads(THREADGROUP_SIZE, 1, 1)]
-void RasterizerCS(uint DispatchThreadID : SV_DispatchThreadID, uint GroupThread1D : SV_GroupIndex, uint GroupID : SV_GroupID)
+void SWRasterizerCS(uint DispatchThreadID : SV_DispatchThreadID, uint GroupThread1D : SV_GroupIndex, uint GroupID : SV_GroupID)
 {
 	ResolvedView = ResolveView();
 
@@ -2044,13 +2045,16 @@ void RasterizerCS(uint DispatchThreadID : SV_DispatchThreadID, uint GroupThread1
 			#endif
 
 			// Write out
-			OutputTexture[Coord.Out.PixelCoord] = float4(OutColor, 1.f);
+			//if (OutTransmittance < 1.f)
+			{
+				OutputTexture[Coord.Out.PixelCoord] = float4(OutColor, 1.f);
+			}
 
 		} // if (bValidWork)
 
 	} // while()
 }
-#endif // RasterizerCS
+#endif // SWRasterizerCS
 
 /////////////////////////////////////////////////////////////////////////////////////////////////////////////////////////////
 
@@ -2070,10 +2074,11 @@ Buffer<uint> VisibleClusterArgs;
 StructuredBuffer<uint2> VisibleClusters;
 
 Texture2D<uint2> SceneTileDepthTexture;
+StructuredBuffer<uint> ViewMinMaxZ;
 
+#if PERMUTATION_SWRASTER
 Texture2DArray<uint> TileSegmentCount;
 StructuredBuffer<uint> TileDataAllocatedCount;
-StructuredBuffer<uint> ViewMinMaxZ;
 
 StructuredBuffer<uint2> ZBinData;
 StructuredBuffer<uint2> RasterWork;
@@ -2081,6 +2086,7 @@ StructuredBuffer<uint> RasterWorkAllocatedCount;
 StructuredBuffer<FPackedSegmentType> ZBinSegmentData;
 StructuredBuffer<uint> ZBinSegmentAllocatedCount;
 StructuredBuffer<uint> ZBinDataAllocatedCount;
+#endif
 
 void PrintRatio(inout FShaderPrintContext Ctx, uint In, uint InMax, uint InDigit)
 {
@@ -2094,10 +2100,12 @@ void DebugDrawingCS(uint2 DispatchThreadId : SV_DispatchThreadID)
 
 	const uint VisibleInstanceCount = VisibleInstanceArgs[0];
 	const uint VisibleClusterCount = VisibleClusterArgs[3];
+	#if PERMUTATION_SWRASTER
 	const uint TileDataAllocCount = TileDataAllocatedCount[0];
 	const uint RasterWorkAllocCount = RasterWorkAllocatedCount[0];
 	const uint ZBinSegAllocatedCount = ZBinSegmentAllocatedCount[0];
 	const uint ZBinDatAllocatedCount = ZBinDataAllocatedCount[0];
+	#endif
 
 	const FMinMaxZ MinMaxZ = UnpackMinMaxZ(ViewMinMaxZ[0], ViewMinMaxZ[1]);
 
@@ -2120,6 +2128,7 @@ void DebugDrawingCS(uint2 DispatchThreadId : SV_DispatchThreadID)
 		Print(Ctx, TEXT("MaxZ                : "), FontWhite); Print(Ctx, MinMaxZ.MaxZ, FontYellow); Newline(Ctx);
 		Newline(Ctx);
 
+		#if PERMUTATION_SWRASTER
 		Print(Ctx, TEXT("Segment & ZBin"), FontOrange); Newline(Ctx);
 		Print(Ctx, TEXT("MaxSegmentDataCount : "), FontWhite); Print(Ctx, MaxSegmentDataCount, FontYellow); Newline(Ctx);
 		Print(Ctx, TEXT("ZBin Data alloc.    : "), FontWhite); PrintRatio(Ctx, ZBinDatAllocatedCount, MaxZBinDataCount, 9); Newline(Ctx);
@@ -2140,6 +2149,7 @@ void DebugDrawingCS(uint2 DispatchThreadId : SV_DispatchThreadID)
 		Print(Ctx, TEXT("Raster Work         : "), FontWhite); PrintRatio(Ctx, RasterWorkAllocCount, MaxRasterWorkCount, 6); Newline(Ctx);
 		Print(Ctx, TEXT("Raster Load         : "), FontWhite); Print(Ctx, RasterWorkAllocCount / float(NumRasterizers), FontYellow); Newline(Ctx);
 		Newline(Ctx);
+		#endif
 
 		Print(Ctx, TEXT("Memory"), FontOrange); Newline(Ctx);
 		Print(Ctx, TEXT("Buffer Memory(MB)   : "), FontWhite); Print(Ctx, TotalBufferMemoryInMBytes, FontYellow); Newline(Ctx);
@@ -2147,6 +2157,7 @@ void DebugDrawingCS(uint2 DispatchThreadId : SV_DispatchThreadID)
 		Newline(Ctx);
 
 		// Cursor bin
+		#if PERMUTATION_SWRASTER
 		if (0)
 		{
 			const uint2 BinCoord = uint2(GetCursorPos()) >> BIN_TILE_SIZE_DIV_AS_SHIFT;
@@ -2158,10 +2169,11 @@ void DebugDrawingCS(uint2 DispatchThreadId : SV_DispatchThreadID)
 			AddQuadSS(Ctx, BinCoord * BIN_TILE_SIZE, (BinCoord + 1) * BIN_TILE_SIZE, ColorRed);
 			Print(Ctx, TEXT("SegmentCount        : "), FontWhite); Print(Ctx, SegmentCount, FontYellow); Newline(Ctx);
 		}
+		#endif
 	}
 
 	// Draw bining tiles
-	#if 0
+	#if PERMUTATION_SWRASTER
 	if (all(DispatchThreadId.xy < BinTileRes))
 	{
 		const uint2 BinCoord = DispatchThreadId.xy;
@@ -2189,7 +2201,7 @@ void DebugDrawingCS(uint2 DispatchThreadId : SV_DispatchThreadID)
 	#endif
 
 	// Draw ZBin tiles
-	#if 1
+	#if PERMUTATION_SWRASTER
 	if (0)
 	//if (all(DispatchThreadId == 0))
 	{
@@ -2257,7 +2269,7 @@ void DebugDrawingCS(uint2 DispatchThreadId : SV_DispatchThreadID)
 
 		// Cluster curves
 		#if 0
-		for (uint CurveIt = 0; CurveIt < ClusterHeader.CurveCount; ++CurveIt)
+		for (uint CurveIt = 0; CurveIt < ClusterHeader.CurvePerCluster; ++CurveIt)
 		{
 			float3 PrevPoint = 0;
 			for (uint PointIt = 0; PointIt < ClusterHeader.PointPerCurve; ++PointIt)
@@ -2422,3 +2434,205 @@ void SegmentLUTCS(uint2 DispatchThreadId : SV_DispatchThreadID, uint2 GroupThrea
 	#endif // PERMUTATION_DEBUG
 }
 #endif // SegmentLUTCS
+
+/////////////////////////////////////////////////////////////////////////////////////////////////////////////////////////////
+
+#ifdef IndirectArgsCS
+
+Buffer<uint>				InVisibleClusterArgs;
+RWBuffer<uint> 				OutHWRasterIndirectArgs;
+RWStructuredBuffer<uint>	OutVisibleClusterCount;
+
+[numthreads(THREADGROUP_SIZE, 1, 1)]
+void IndirectArgsCS(uint GroupThread1D : SV_GroupIndex, uint GroupID : SV_GroupID)
+{
+	if (GroupThread1D == 0)
+	{
+		const uint VisibleClusterCount = InVisibleClusterArgs[3];
+		OutHWRasterIndirectArgs[0] = RENDERCURVE_CURVE_PER_CLUSTER;
+		OutHWRasterIndirectArgs[1] = VisibleClusterCount;
+		OutHWRasterIndirectArgs[2] = 1;
+
+		OutVisibleClusterCount[0] = VisibleClusterCount;
+	}
+}
+#endif // IndirectArgsCS
+
+/////////////////////////////////////////////////////////////////////////////////////////////////////////////////////////////
+
+#if defined(HWRasterizerMS) || defined(HWRasterizerPS)
+struct VSOut
+{
+	float4 Position	: SV_Position;
+	float CoordU : TEXCOORD1;
+};
+
+struct VSPrimiveOut
+{
+	// Use uint4 to prevent compiler from erroneously packing per-vertex and per-prim attributes together
+	nointerpolation uint PackedVisibleCurve : TEXCOORD2;
+};
+
+struct FVisibleCurve
+{
+	uint VisibleClusterIndex;	/* [0-65k] */
+	uint CurveIndex; 			/* [0-64]  */
+	uint PointIndex; 			/* [0-64]  */
+};
+
+uint PackCurveSegment(FVisibleCurve In)
+{
+	return (In.VisibleClusterIndex << 16u) | (In.CurveIndex << 8) | (In.PointIndex);
+}
+
+FVisibleCurve UnpackCurveSegment(uint In)
+{
+	FVisibleCurve Out;
+	Out.VisibleClusterIndex = BitFieldExtractU32(In, 16, 16);
+	Out.CurveIndex = BitFieldExtractU32(In, 8, 8);
+	Out.PointIndex = BitFieldExtractU32(In, 8, 0);
+	return Out;
+}
+#endif
+
+#ifdef HWRasterizerMS
+
+StructuredBuffer<uint>    ViewMinMaxZ;
+Texture2D<uint2>          SceneTileDepthTexture;
+StructuredBuffer<uint2>   VisibleClusters;
+StructuredBuffer<uint>    VisibleClustersCount;
+
+#ifndef MESHSHADER
+#error This shader requires mesh shader support
+#endif
+
+// Dispatch layout
+// ---------------
+// 1 group  -> 1 curve
+// 1 thread -> 1 point -> 2 vertices | 2 triangles
+// 1 cluster == 64 group
+
+MESH_SHADER_TRIANGLE_ATTRIBUTES(128)
+void HWRasterizerMS(
+	uint GroupThreadID : SV_GroupThreadID,
+	uint3 GroupID : SV_GroupID,
+	MESH_SHADER_VERTEX_EXPORT(VSOut, 128),
+	MESH_SHADER_TRIANGLE_EXPORT(126),
+	MESH_SHADER_PRIMITIVE_EXPORT(VSPrimiveOut, 126)
+)
+{
+	ResolvedView = ResolveView();
+
+	uint NumUniqueVerts = 0;
+	uint3 VertIndexes = 0;
+	
+	bool bValid = false;
+	float3 SP0 = 0;
+	float3 SP1 = 0;
+	
+	// 1. Load data
+	float4 ClipP0 = float4(INFINITE_FLOAT, INFINITE_FLOAT, INFINITE_FLOAT, INFINITE_FLOAT);
+	float4 ClipP1 = ClipP0;
+	float UCoord0 = 0;
+	uint PointCount = 0;
+	uint CurveCount = 0;
+	float SignT = 1;
+	const uint PointIt = GroupThreadID;
+	const uint CurveIt = GroupID.x;
+	const uint ClusterFetchIndex = GroupID.y; // TODO linearize .yz
+	if (ClusterFetchIndex < VisibleClustersCount[0])
+	{
+		const uint2 VisibleData = VisibleClusters[ClusterFetchIndex];
+		const uint PrimitiveIndex = VisibleData.x;
+		const uint ClusterIndex = VisibleData.y;
+		const FRenderCurveInstanceData RenderCurveInstanceData = GetRenderCurveInstanceData(PrimitiveIndex);
+		const FClusterHeader ClusterHeader = GetClusterHeader(ClusterIndex);
+		const float4x4 LocalToClip = mul(RenderCurveInstanceData.LocalToTranslatedWorld, ResolvedView.TranslatedWorldToClip);
+
+		const uint PointIt0 = GroupThreadID;
+		const uint PointIt1 = PointIt0 + 1 < ClusterHeader.PointPerCurve ? PointIt0 + 1 : PointIt0 - 1;
+		SignT = PointIt0 + 1 < ClusterHeader.PointPerCurve ? 1 : -1;
+
+		CurveCount = ClusterHeader.CurveCount;
+		PointCount = ClusterHeader.CurveCount > 0 && CurveIt < ClusterHeader.CurveCount ? ClusterHeader.PointPerCurve : 0u;
+		if (PointIt0 < PointCount && CurveIt < CurveCount)
+		{
+			const FCurvePoint Point0 = GetClusterPoint(ClusterHeader, ClusterIndex, CurveIt, PointIt0);
+			const FCurvePoint Point1 = GetClusterPoint(ClusterHeader, ClusterIndex, CurveIt, PointIt1);
+			if (Point0.bValid && Point1.bValid)
+			{
+				float4 ClipPosition0 = mul(float4(Point0.Position, 1), LocalToClip);
+				float4 ClipPosition1 = mul(float4(Point1.Position, 1), LocalToClip);
+				UCoord0 = Point0.UCoord;
+				ClipP0 = ClipPosition0;
+				ClipP1 = ClipPosition1;
+				bValid = true;
+			}
+		}
+	}
+
+	// 2. Write out VertexCount/TriangleCount for the whole batch
+	const uint VertexCount = 2 * PointCount;
+	const uint TriangleCount = PointCount > 0 ? 2 * (PointCount - 1) : 0u;
+	SetMeshOutputCounts(VertexCount, TriangleCount);
+
+	// 3. Write Vertex / Triangle
+	if (PointIt < PointCount && CurveIt < CurveCount)
+	{
+		const float2 CP0 = ClipP0.xy;
+		const float2 CP1 = ClipP1.xy;
+		const float2 CT  = normalize(CP1 - CP0) * SignT;
+		const float2 CR  = abs(CT.x) > 0.001f ? float2(CT.y, -CT.x) : float2(-CT.y, CT.x);
+
+		const float2 InvClipResolution = 2.f / Resolution;
+		const float  ClipPixelSize = max(InvClipResolution.x, InvClipResolution.y);
+
+		VSOut OutVL, OutVR;
+		OutVL.Position = bValid ? float4(ClipP0.xy - CR * ClipPixelSize * ClipP0.w, ClipP0.zw) : ClipP0;
+		OutVR.Position = bValid ? float4(ClipP0.xy + CR * ClipPixelSize * ClipP0.w, ClipP0.zw) : ClipP0;
+		OutVL.CoordU = UCoord0;
+		OutVR.CoordU = UCoord0;
+
+		// Vertices
+		{
+			const uint VertexIndex = PointIt*2;
+			MESH_SHADER_WRITE_VERTEX(VertexIndex + 0, OutVL);
+			MESH_SHADER_WRITE_VERTEX(VertexIndex + 1, OutVR);
+		}
+		
+		// Triangles
+		const uint TriangleIndex = PointIt*2;
+		if (TriangleIndex < TriangleCount)
+		{
+			// Triangle indices
+			const uint3 OutT0 = uint3(0, 1, 2) + TriangleIndex;
+			const uint3 OutT1 = uint3(1, 3, 2) + TriangleIndex;
+			MESH_SHADER_WRITE_TRIANGLE(TriangleIndex + 0, OutT0);
+			MESH_SHADER_WRITE_TRIANGLE(TriangleIndex + 1, OutT1);
+
+			// Attributes
+			FVisibleCurve VisibleCurve;
+			VisibleCurve.VisibleClusterIndex = ClusterFetchIndex;
+			VisibleCurve.CurveIndex = CurveIt;
+			VisibleCurve.PointIndex = PointIt;
+			VSPrimiveOut OutPrimAttributes;
+			OutPrimAttributes.PackedVisibleCurve = PackCurveSegment(VisibleCurve);
+			MESH_SHADER_WRITE_PRIMITIVE(TriangleIndex + 0, OutPrimAttributes);
+			MESH_SHADER_WRITE_PRIMITIVE(TriangleIndex + 1, OutPrimAttributes);
+		}
+	}
+}
+#endif // HWRasterizerMS
+
+#ifdef HWRasterizerPS
+
+void HWRasterizerPS(
+	in float4 SvPosition: SV_Position,
+	in float InCoordU : TEXCOORD1,
+	in uint InPackedAttribute : TEXCOORD2,
+	out float4 OutColor0 : SV_Target0)
+{
+	const FVisibleCurve Curve = UnpackCurveSegment(InPackedAttribute);
+	OutColor0 = float4(InCoordU, Curve.VisibleClusterIndex, Curve.CurveIndex, Curve.PointIndex);
+}
+#endif // HWRasterizerPS
\ No newline at end of file
diff --git a/Engine/Shaders/Shared/RenderCurveDefinitions.h b/Engine/Shaders/Shared/RenderCurveDefinitions.h
new file mode 100644
--- /dev/null
+++ b/Engine/Shaders/Shared/RenderCurveDefinitions.h
@@ -0,0 +1,9 @@
+// Copyright Epic Games, Inc. All Rights Reserved.
+
+#pragma once
+
+// Number of curves per cluster
+#define RENDERCURVE_CURVE_PER_CLUSTER	64u
+
+// Primitive data size
+#define RENDERCURVE_PRIMITIVE_DATA_STRIDE_IN_BYTES 16
\ No newline at end of file
diff --git a/Engine/Source/Runtime/Renderer/Private/HairStrands/RenderCurveRaster.cpp b/Engine/Source/Runtime/Renderer/Private/HairStrands/RenderCurveRaster.cpp
--- a/Engine/Source/Runtime/Renderer/Private/HairStrands/RenderCurveRaster.cpp
+++ b/Engine/Source/Runtime/Renderer/Private/HairStrands/RenderCurveRaster.cpp
@@ -42,6 +42,20 @@ static TAutoConsoleVariable<int32> CVarRenderCurveNumRasterizers(
 	ECVF_RenderThreadSafe
 );
 
+static TAutoConsoleVariable<int32> CVarRenderCurveSWRaster(
+	TEXT("r.RenderCurve.SWRaster"),
+	1,
+	TEXT("Enable curve software rasterizer."),
+	ECVF_RenderThreadSafe
+);
+
+static TAutoConsoleVariable<int32> CVarRenderCurveHWRaster(
+	TEXT("r.RenderCurve.HWRaster"),
+	0,
+	TEXT("Enable curve hardware rasterizer."),
+	ECVF_RenderThreadSafe
+);
+
 /////////////////////////////////////////////////////////////////////////////////////////
 // FPackedSegment
 
@@ -60,6 +74,7 @@ inline bool IsRenderCurveSupported(EShaderPlatform In)
 	return 
 		IsRenderCurveEnabled() &&
 		IsFeatureLevelSupported(In, ERHIFeatureLevel::SM6) && 
+		FDataDrivenShaderPlatformInfo::GetSupportsMeshShadersTier1(In) && 
 		IsHairStrandsSupported(EHairStrandsShaderType::Strands, In) && 
 		!IsVulkanPlatform(In) && 
 		!IsMetalPlatform(In);
@@ -215,6 +230,34 @@ public:
 
 IMPLEMENT_GLOBAL_SHADER(FRenderCurveSceneTileDepthCS, "/Engine/Private/HairStrands/RenderCurveRaster.usf", "SceneTileDepthCS", SF_Compute);
 
+/////////////////////////////////////////////////////////////////////////////////////////
+// Indirect args for raster
+
+class FRenderCurveIndirectArgsCS : public FGlobalShader
+{
+	DECLARE_GLOBAL_SHADER(FRenderCurveIndirectArgsCS);
+	SHADER_USE_PARAMETER_STRUCT(FRenderCurveIndirectArgsCS, FGlobalShader);
+
+	using FPermutationDomain = TShaderPermutationDomain<>;
+
+	BEGIN_SHADER_PARAMETER_STRUCT(FParameters, )
+		SHADER_PARAMETER_RDG_BUFFER_SRV(Buffer<uint>, 				InVisibleClusterArgs)
+		SHADER_PARAMETER_RDG_BUFFER_UAV(RWBuffer<uint>, 			OutHWRasterIndirectArgs)
+		SHADER_PARAMETER_RDG_BUFFER_UAV(RWStructuredBuffer<uint>,	OutVisibleClusterCount)
+	END_SHADER_PARAMETER_STRUCT()
+
+public:
+	static uint32 GetGroupSize() { return 32u; }
+	static bool ShouldCompilePermutation(const FGlobalShaderPermutationParameters& Parameters) { return IsRenderCurveSupported(Parameters.Platform); }
+	static void ModifyCompilationEnvironment(const FGlobalShaderPermutationParameters &Parameters, FShaderCompilerEnvironment &OutEnvironment)
+	{
+		FGlobalShader::ModifyCompilationEnvironment(Parameters, OutEnvironment);
+		OutEnvironment.SetDefine(TEXT("THREADGROUP_SIZE"), GetGroupSize());
+	}
+};
+
+IMPLEMENT_GLOBAL_SHADER(FRenderCurveIndirectArgsCS, "/Engine/Private/HairStrands/RenderCurveRaster.usf", "IndirectArgsCS", SF_Compute);
+
 /////////////////////////////////////////////////////////////////////////////////////////
 // Bin segments
 
@@ -297,10 +340,10 @@ IMPLEMENT_GLOBAL_SHADER(FRenderCurveCompactionCS, "/Engine/Private/HairStrands/R
 /////////////////////////////////////////////////////////////////////////////////////////
 // Compact binned segments into contiguous list 
 
-class FRenderCurveRasterizerCS : public FGlobalShader
+class FRenderCurveSWRasterizerCS : public FGlobalShader
 {
-	DECLARE_GLOBAL_SHADER(FRenderCurveRasterizerCS);
-	SHADER_USE_PARAMETER_STRUCT(FRenderCurveRasterizerCS, FGlobalShader);
+	DECLARE_GLOBAL_SHADER(FRenderCurveSWRasterizerCS);
+	SHADER_USE_PARAMETER_STRUCT(FRenderCurveSWRasterizerCS, FGlobalShader);
 
 	class FDebug : SHADER_PERMUTATION_BOOL("PERMUTATION_DEBUG");
 	using FPermutationDomain = TShaderPermutationDomain<FDebug>;
@@ -331,7 +374,66 @@ public:
 	}
 };
 
-IMPLEMENT_GLOBAL_SHADER(FRenderCurveRasterizerCS, "/Engine/Private/HairStrands/RenderCurveRaster.usf", "RasterizerCS", SF_Compute);
+IMPLEMENT_GLOBAL_SHADER(FRenderCurveSWRasterizerCS, "/Engine/Private/HairStrands/RenderCurveRaster.usf", "SWRasterizerCS", SF_Compute);
+
+/////////////////////////////////////////////////////////////////////////////////////////
+class FRenderCurveHWRasterizerMS : public FGlobalShader
+{
+	DECLARE_GLOBAL_SHADER(FRenderCurveHWRasterizerMS);
+	SHADER_USE_PARAMETER_STRUCT(FRenderCurveHWRasterizerMS, FGlobalShader);
+
+	using FPermutationDomain = TShaderPermutationDomain<>;
+
+	BEGIN_SHADER_PARAMETER_STRUCT(FParameters, )
+		SHADER_PARAMETER_STRUCT_INCLUDE(FRenderCurveCommonParameters, CommonParameters)
+		SHADER_PARAMETER_RDG_BUFFER_SRV(StructuredBuffer<uint>, ViewMinMaxZ)
+		SHADER_PARAMETER_RDG_TEXTURE(Texture2D<uint2>, SceneTileDepthTexture)
+		SHADER_PARAMETER_RDG_BUFFER_SRV(StructuredBuffer<uint2>, VisibleClusters)
+		SHADER_PARAMETER_RDG_BUFFER_SRV(StructuredBuffer<uint>,  VisibleClustersCount)
+	END_SHADER_PARAMETER_STRUCT()
+
+public:
+	static uint32 GetGroupSize() { return 128u; }
+	static bool ShouldCompilePermutation(const FGlobalShaderPermutationParameters& Parameters) { return IsRenderCurveSupported(Parameters.Platform); }
+	static void ModifyCompilationEnvironment(const FGlobalShaderPermutationParameters &Parameters, FShaderCompilerEnvironment &OutEnvironment)
+	{
+		FGlobalShader::ModifyCompilationEnvironment(Parameters, OutEnvironment);
+		OutEnvironment.SetDefine(TEXT("THREADGROUP_SIZE"), GetGroupSize());
+	}
+};
+
+IMPLEMENT_GLOBAL_SHADER(FRenderCurveHWRasterizerMS, "/Engine/Private/HairStrands/RenderCurveRaster.usf", "HWRasterizerMS", SF_Mesh);
+
+/////////////////////////////////////////////////////////////////////////////////////////
+class FRenderCurveHWRasterizerPS : public FGlobalShader
+{
+	DECLARE_GLOBAL_SHADER(FRenderCurveHWRasterizerPS);
+	SHADER_USE_PARAMETER_STRUCT(FRenderCurveHWRasterizerPS, FGlobalShader);
+
+	using FPermutationDomain = TShaderPermutationDomain<>;
+
+	BEGIN_SHADER_PARAMETER_STRUCT(FParameters, )
+		SHADER_PARAMETER_STRUCT_INCLUDE(FRenderCurveCommonParameters, CommonParameters)
+		RENDER_TARGET_BINDING_SLOTS()
+	END_SHADER_PARAMETER_STRUCT()
+
+public:
+	static bool ShouldCompilePermutation(const FGlobalShaderPermutationParameters& Parameters) { return IsRenderCurveSupported(Parameters.Platform); }
+	static void ModifyCompilationEnvironment(const FGlobalShaderPermutationParameters &Parameters, FShaderCompilerEnvironment &OutEnvironment)
+	{
+		FGlobalShader::ModifyCompilationEnvironment(Parameters, OutEnvironment);
+	}
+};
+
+IMPLEMENT_GLOBAL_SHADER(FRenderCurveHWRasterizerPS, "/Engine/Private/HairStrands/RenderCurveRaster.usf", "HWRasterizerPS", SF_Pixel);
+
+/////////////////////////////////////////////////////////////////////////////////////////
+
+BEGIN_SHADER_PARAMETER_STRUCT(FRenderCurveHWRasterizerParameters, )
+	SHADER_PARAMETER_STRUCT_INCLUDE(FRenderCurveHWRasterizerMS::FParameters, MS)
+	SHADER_PARAMETER_STRUCT_INCLUDE(FRenderCurveHWRasterizerPS::FParameters, PS)
+	RDG_BUFFER_ACCESS(IndirectArgs, ERHIAccess::IndirectArgs)
+END_SHADER_PARAMETER_STRUCT()
 
 /////////////////////////////////////////////////////////////////////////////////////////
 
@@ -340,7 +442,8 @@ class FRenderCurveDebugDrawingCS : public FGlobalShader
 	DECLARE_GLOBAL_SHADER(FRenderCurveDebugDrawingCS);
 	SHADER_USE_PARAMETER_STRUCT(FRenderCurveDebugDrawingCS, FGlobalShader);
 
-	using FPermutationDomain = TShaderPermutationDomain<>;
+	class FSWRaster : SHADER_PERMUTATION_BOOL("PERMUTATION_SWRASTER");
+	using FPermutationDomain = TShaderPermutationDomain<FSWRaster>;
 
 	BEGIN_SHADER_PARAMETER_STRUCT(FParameters, )
 		SHADER_PARAMETER(uint32, TotalBufferMemoryInMBytes)
@@ -391,6 +494,10 @@ struct FRenderCurveTransientData
 	FRDGBufferSRVRef VisibleClusterArgsSRV = nullptr;
 	FRDGBufferRef    VisibleClusterArgs = nullptr;
 
+	FRDGBufferRef 	 VisibleClusterCount = nullptr;
+	FRDGBufferSRVRef VisibleClusterCountSRV = nullptr;
+	FRDGBufferRef 	 HWRasterizerArgs = nullptr;
+
 	FRDGTextureRef   SceneTileDepth = nullptr;
 
 	FRDGTextureRef   TileSegmentCount = nullptr;	
@@ -479,6 +586,13 @@ void AddRenderCurveRasterPipeline(
 	
 	if (InstanceCount == 0) { return; }
 
+	const bool bSWRaster = CVarRenderCurveSWRaster.GetValueOnRenderThread() > 0;
+	const bool bHWRaster = CVarRenderCurveHWRaster.GetValueOnRenderThread() > 0;
+	if (!bSWRaster && !bHWRaster)
+	{
+		return;
+	}
+
 	QUICK_SCOPE_CYCLE_COUNTER(STAT_CLM_CurveRasterPipeline);
 	RDG_EVENT_SCOPE_STAT(GraphBuilder, CurveRasterPipeline, "CurveRasterPipeline");
 	RDG_GPU_STAT_SCOPE(GraphBuilder, CurveRasterPipeline);
@@ -651,13 +765,23 @@ void AddRenderCurveRasterPipeline(
 		FComputeShaderUtils::AddPass(GraphBuilder, RDG_EVENT_NAME("RenderCurve::SceneTileDepth"), ComputeShaderRaster, Parameters, FIntVector(CommonParameters.BinTileRes.X, CommonParameters.BinTileRes.Y, 1));
 	}
 
-	// Binning
+	// Indirect args for HW rasterizer and visible cluster count
 	{
-		// TODO change cluster culling to output directly into a structure buffer as we don't use the indirect args?
-		FRDGBufferRef VisibleClusterCount = GraphBuilder.CreateBuffer(FRDGBufferDesc::CreateStructuredDesc(sizeof(uint32), 1), TEXT("RenderCurve.Binning.VisibleClusterCount"));
-		AddCopyBufferPass(GraphBuilder, VisibleClusterCount, 0/*DstOffset*/, TransientData.VisibleClusterArgs, sizeof(uint32) * 3u, sizeof(uint32)); // Copy into a Structured buffer for scalar loading
-		MemoryTracker.Add(VisibleClusterCount);
+		TransientData.VisibleClusterCount = GraphBuilder.CreateBuffer(FRDGBufferDesc::CreateStructuredDesc(sizeof(uint32), 1), TEXT("RenderCurve.Binning.VisibleClusterCount"));
+		TransientData.VisibleClusterCountSRV = GraphBuilder.CreateSRV(TransientData.VisibleClusterCount);
+		TransientData.HWRasterizerArgs = GraphBuilder.CreateBuffer(FRDGBufferDesc::CreateIndirectDesc(IndirectArgsSizeInBytes, 1u), TEXT("RenderCurve.HWRasterizerArgs"));
+		
+		TShaderMapRef<FRenderCurveIndirectArgsCS> ComputeShader(ShaderMap);
+		FRenderCurveIndirectArgsCS::FParameters* Parameters = GraphBuilder.AllocParameters<FRenderCurveIndirectArgsCS::FParameters>();
+		Parameters->InVisibleClusterArgs = TransientData.VisibleClusterArgsSRV;
+		Parameters->OutHWRasterIndirectArgs = GraphBuilder.CreateUAV(TransientData.HWRasterizerArgs, PF_R32_UINT);
+		Parameters->OutVisibleClusterCount = GraphBuilder.CreateUAV(TransientData.VisibleClusterCount, PF_R32_UINT);
+		FComputeShaderUtils::AddPass(GraphBuilder, RDG_EVENT_NAME("RenderCurve::IndirectArgs"), ComputeShader, Parameters, FIntVector(1, 1, 1));
+	}
 
+	// Binning
+	if (bSWRaster)
+	{
 		FRDGBufferRef VisibleClusterQueue = GraphBuilder.CreateBuffer(FRDGBufferDesc::CreateStructuredDesc(sizeof(uint32), 1), TEXT("RenderCurve.Binning.VisibleClusterQueue"));
 		FRDGBufferUAVRef RWVisibleClusterQueue = GraphBuilder.CreateUAV(VisibleClusterQueue, PF_R32_UINT);
 		AddClearUAVPass(GraphBuilder, RWVisibleClusterQueue, 0u);
@@ -690,7 +814,7 @@ void AddRenderCurveRasterPipeline(
 		Parameters->ViewMinMaxZ = TransientData.MinMaxZ;
 		Parameters->SceneTileDepthTexture = TransientData.SceneTileDepth;
 		Parameters->VisibleClusters = TransientData.VisibleClustersSRV;
-		Parameters->VisibleClustersCount = GraphBuilder.CreateSRV(VisibleClusterCount, PF_R32_UINT);
+		Parameters->VisibleClustersCount = TransientData.VisibleClusterCountSRV;
 		Parameters->VisibleClustersQueue = RWVisibleClusterQueue;
 		Parameters->RWTileSegmentCount = RWTileSegmentCount;
 		Parameters->RWTileData = GraphBuilder.CreateUAV(TileData, PF_R32_UINT);
@@ -732,6 +856,7 @@ void AddRenderCurveRasterPipeline(
 	//   Offset           Offset   Offset
 	//
 	// * RasterWork - List of work passed to the rasterize. Contains ZBinOffset|ZBinCount|TileCoord.
+	if (bSWRaster)
 	{
 		FRDGBufferRef ZBinData = GraphBuilder.CreateBuffer(FRDGBufferDesc::CreateStructuredDesc(sizeof(uint32)*2, CommonParameters.MaxZBinDataCount), TEXT("RenderCurve.Compaction.ZBinData"));
 		TransientData.ZBinData = GraphBuilder.CreateSRV(ZBinData);
@@ -784,13 +909,14 @@ void AddRenderCurveRasterPipeline(
 	}
 
 	// Rasterizer
+	if (bSWRaster)
 	{		
 		FRDGBufferRef RasterWorkQueue = GraphBuilder.CreateBuffer(FRDGBufferDesc::CreateStructuredDesc(sizeof(uint32), 1), TEXT("RenderCurve.Rasterizer.WorkQueue"));
 		FRDGBufferUAVRef RWRasterWorkQueue = GraphBuilder.CreateUAV(RasterWorkQueue, PF_R32_UINT);
 		AddClearUAVPass(GraphBuilder, RWRasterWorkQueue, 0u);
 		MemoryTracker.Add(RasterWorkQueue);
 
-		FRenderCurveRasterizerCS::FParameters* Parameters = GraphBuilder.AllocParameters<FRenderCurveRasterizerCS::FParameters>();
+		FRenderCurveSWRasterizerCS::FParameters* Parameters = GraphBuilder.AllocParameters<FRenderCurveSWRasterizerCS::FParameters>();
 		Parameters->CommonParameters 			= CommonParameters;
 		Parameters->SceneTileDepthTexture		= TransientData.SceneTileDepth;
 		Parameters->SceneDepthTexture 			= SceneDepthTexture;
@@ -805,16 +931,63 @@ void AddRenderCurveRasterPipeline(
 		Parameters->RasterWorkQueue				= RWRasterWorkQueue;
 		Parameters->OutputTexture				= GraphBuilder.CreateUAV(SceneColorTexture);
 
-		FRenderCurveRasterizerCS::FPermutationDomain PermutationVector;
-		PermutationVector.Set<FRenderCurveRasterizerCS::FDebug>(bDebug);
-		TShaderMapRef<FRenderCurveRasterizerCS> ComputeShader(ShaderMap, PermutationVector);
-		FComputeShaderUtils::AddPass(GraphBuilder, RDG_EVENT_NAME("RenderCurve::Rasterizer"), ComputeShader, Parameters, FIntVector(CommonParameters.NumRasterizers, 1, 1));
+		FRenderCurveSWRasterizerCS::FPermutationDomain PermutationVector;
+		PermutationVector.Set<FRenderCurveSWRasterizerCS::FDebug>(bDebug);
+		TShaderMapRef<FRenderCurveSWRasterizerCS> ComputeShader(ShaderMap, PermutationVector);
+		FComputeShaderUtils::AddPass(GraphBuilder, RDG_EVENT_NAME("RenderCurve::SWRasterizer"), ComputeShader, Parameters, FIntVector(CommonParameters.NumRasterizers, 1, 1));
+	}
+
+	// Rasterizer
+	if (bHWRaster)
+	{	
+		FRDGTextureDesc Desc = FRDGTextureDesc::Create2D(CommonParameters.Resolution, PF_FloatRGBA, FClearValueBinding::Black, TexCreate_UAV | TexCreate_ShaderResource | TexCreate_RenderTargetable);
+		FRDGTextureRef Output = GraphBuilder.CreateTexture(Desc, TEXT("RenderCurve.HWRasterizerOutput"));
+		AddClearRenderTargetPass(GraphBuilder, Output, FLinearColor::Black);
+
+		FRenderCurveHWRasterizerParameters* Parameters = GraphBuilder.AllocParameters<FRenderCurveHWRasterizerParameters>();
+		Parameters->MS.CommonParameters = CommonParameters;
+		Parameters->MS.ViewMinMaxZ = TransientData.MinMaxZ;
+		Parameters->MS.SceneTileDepthTexture = TransientData.SceneTileDepth;
+		Parameters->MS.VisibleClusters = TransientData.VisibleClustersSRV;
+		Parameters->MS.VisibleClustersCount = TransientData.VisibleClusterCountSRV;
+		Parameters->PS.RenderTargets[0] = FRenderTargetBinding(Output, ERenderTargetLoadAction::ENoAction);
+		Parameters->IndirectArgs = TransientData.HWRasterizerArgs;
+
+		TShaderMapRef<FRenderCurveHWRasterizerMS> MeshShader(ShaderMap);
+		TShaderMapRef<FRenderCurveHWRasterizerPS> PixelShader(ShaderMap);
+
+		GraphBuilder.AddPass(
+			RDG_EVENT_NAME("RenderCurve::HWRasterizer"),
+			Parameters,
+			ERDGPassFlags::Raster,
+			[Parameters, MeshShader, PixelShader](FRDGAsyncTask, FRHICommandList& RHICmdList)
+		{
+			FGraphicsPipelineStateInitializer GraphicsPSOInit;
+			RHICmdList.ApplyCachedRenderTargets(GraphicsPSOInit);
+			GraphicsPSOInit.BlendState = TStaticBlendState<CW_RGBA, BO_Add, BF_One, BF_Zero, BO_Add, BF_One, BF_Zero>::GetRHI();
+			GraphicsPSOInit.RasterizerState = TStaticRasterizerState<>::GetRHI();
+			GraphicsPSOInit.DepthStencilState = TStaticDepthStencilState<false, CF_Always>::GetRHI();
+			GraphicsPSOInit.BoundShaderState.SetMeshShader(MeshShader.GetMeshShader());
+			GraphicsPSOInit.BoundShaderState.PixelShaderRHI = PixelShader.GetPixelShader();
+			GraphicsPSOInit.PrimitiveType = PT_TriangleList;
+			SetGraphicsPipelineState(RHICmdList, GraphicsPSOInit, 0);
+
+			SetShaderParameters(RHICmdList, MeshShader, MeshShader.GetMeshShader(), Parameters->MS);
+			SetShaderParameters(RHICmdList, PixelShader, PixelShader.GetPixelShader(), Parameters->PS);
+
+			RHICmdList.SetViewport(0, 0, 0.0f, Parameters->MS.CommonParameters.Resolution.X, Parameters->MS.CommonParameters.Resolution.Y, 1.0f); // MinRect?
+			//RHICmdList.DispatchMeshShader(64, 1, 1);			
+			RHICmdList.DispatchIndirectMeshShader(Parameters->IndirectArgs->GetIndirectRHICallBuffer(), 0);
+			
+			//FORCEINLINE_DEBUGGABLE void DispatchIndirectMeshShader(FRHIBuffer* ArgumentBuffer, uint32 ArgumentOffset)
+		});
 	}
 
 	// Debug drawing
 	if (bDebug)
 	{
 		FRenderCurveDebugDrawingCS::FPermutationDomain PermutationVector;
+		PermutationVector.Set<FRenderCurveDebugDrawingCS::FSWRaster>(bSWRaster);
 		TShaderMapRef<FRenderCurveDebugDrawingCS> ComputeShader(ShaderMap, PermutationVector);
 
 		FRenderCurveDebugDrawingCS::FParameters* Parameters = GraphBuilder.AllocParameters<FRenderCurveDebugDrawingCS::FParameters>();
diff --git a/Engine/Source/Runtime/Renderer/Private/HairStrands/RenderCurveSceneExtension.cpp b/Engine/Source/Runtime/Renderer/Private/HairStrands/RenderCurveSceneExtension.cpp
--- a/Engine/Source/Runtime/Renderer/Private/HairStrands/RenderCurveSceneExtension.cpp
+++ b/Engine/Source/Runtime/Renderer/Private/HairStrands/RenderCurveSceneExtension.cpp
@@ -7,6 +7,7 @@
 #include "SceneUniformBuffer.h"
 #include "ShaderParameterMacros.h"
 #include "HairStrandsInterface.h"
+#include "RenderCurveDefinitions.h"
 
 namespace RenderCurve
 {
@@ -14,7 +15,6 @@ namespace RenderCurve
 IMPLEMENT_SCENE_EXTENSION(FRenderCurveSceneExtension);
 
 BEGIN_SHADER_PARAMETER_STRUCT(FRenderCurveSceneParameters, RENDERER_API)
-	SHADER_PARAMETER(uint32, RenderCurveInstanceDataStride)
 	SHADER_PARAMETER(uint32, InstanceCount)
 	SHADER_PARAMETER(uint32, ClusterCount)
 	SHADER_PARAMETER(uint32, MaxClusterStrideInBytes)
@@ -218,10 +218,10 @@ void FRenderCurveSceneExtension::FinishBufferUpload(FRDGBuilder& GraphBuilder,	F
 
 	if (OutParameters != nullptr)
 	{
+		static_assert(RENDERCURVE_PRIMITIVE_DATA_STRIDE_IN_BYTES == sizeof(FPackedRenderCurveInstanceData));
+
 		OutParameters->InstanceCount = Datas.Num();
 		OutParameters->RenderCurveInstanceData = GraphBuilder.CreateSRV(RenderCurveInstanceDataBuffer);
-		OutParameters->RenderCurveInstanceDataStride = sizeof(FPackedRenderCurveInstanceData);
-
 		OutParameters->ClusterData = GraphBuilder.CreateSRV(ClusterDataBuffer);
 		OutParameters->ClusterCount = Header.TotalClusterCount;
 		OutParameters->MaxClusterStrideInBytes = Header.ClusterStideInBytes;
@@ -311,7 +311,6 @@ static void GetRenderCurveSceneParameters(FRDGBuilder& GraphBuilder, FRenderCurv
 	OutParameters.RenderCurveInstanceData = DefaultBuffer;
 	OutParameters.ClusterData = DefaultBuffer;
 	OutParameters.InstanceCount = InSceneData ? InSceneData->GetInstanceCount() : 0u;
-	OutParameters.RenderCurveInstanceDataStride = sizeof(FRenderCurveSceneExtension::FPackedRenderCurveInstanceData);
 }
 
 static void GetDefaultRenderCurveSceneParameters(FRenderCurveSceneParameters& OutParameters, FRDGBuilder& GraphBuilder)
