Material Cache, fixed invalid shading commands on incomplete shader maps

diff --git a/Engine/Source/Runtime/Renderer/Private/MaterialCache/MaterialCacheMeshProcessor.cpp b/Engine/Source/Runtime/Renderer/Private/MaterialCache/MaterialCacheMeshProcessor.cpp
--- a/Engine/Source/Runtime/Renderer/Private/MaterialCache/MaterialCacheMeshProcessor.cpp
+++ b/Engine/Source/Runtime/Renderer/Private/MaterialCache/MaterialCacheMeshProcessor.cpp
@@ -233,6 +233,29 @@ bool FMaterialCacheMeshProcessor::TryAddMeshBatch(const FMeshBatch& RESTRICT Mes
 }
 
 #if WITH_EDITOR
+static bool QueryOrSignalMaterialCacheShaderMap(ERHIFeatureLevel::Type FeatureLevel, const FMaterialRenderProxy* Proxy)
+{
+	UMaterialInterface* Interface = Proxy->GetMaterialInterface();
+	if (!Interface)
+	{
+		return false;
+	}
+
+	FMaterialResource* Resource = Interface->GetMaterialResource(FeatureLevel);
+	if (!Resource)
+	{
+		return false;
+	}
+
+	if (!Resource->IsRenderingThreadShaderMapComplete())
+	{
+		Resource->SubmitCompileJobs_RenderThread(EShaderCompileJobPriority::High);
+		return false;
+	}
+
+	return true;
+}
+
 bool IsMaterialCacheMaterialReady(ERHIFeatureLevel::Type FeatureLevel, const FPrimitiveSceneProxy* Proxy)
 {
 	// Must have a valid material cache proxy
@@ -254,12 +277,12 @@ bool IsMaterialCacheMaterialReady(ERHIFeatureLevel::Type FeatureLevel, const FPr
 		
 		for (const Nanite::FSceneProxyBase::FMaterialSection& MaterialSection : NaniteProxy->GetMaterialSections())
 		{
-			if (!MaterialSection.RasterMaterialProxy->GetMaterialNoFallback(FeatureLevel))
+			if (!QueryOrSignalMaterialCacheShaderMap(FeatureLevel, MaterialSection.RasterMaterialProxy))
 			{
 				return false;
 			}
 
-			if (!MaterialSection.ShadingMaterialProxy->GetMaterialNoFallback(FeatureLevel))
+			if (!QueryOrSignalMaterialCacheShaderMap(FeatureLevel, MaterialSection.ShadingMaterialProxy))
 			{
 				return false;
 			}
@@ -275,7 +298,7 @@ bool IsMaterialCacheMaterialReady(ERHIFeatureLevel::Type FeatureLevel, const FPr
 
 		for (const FStaticMeshBatch& StaticMesh : PrimitiveSceneInfo->StaticMeshes)
 		{
-			if (!StaticMesh.MaterialRenderProxy->GetMaterialNoFallback(FeatureLevel))
+			if (!QueryOrSignalMaterialCacheShaderMap(FeatureLevel, StaticMesh.MaterialRenderProxy))
 			{
 				return false;
 			}
