Update PrimitiveRayTracingData when DrawnInGame flag is modified.

diff --git a/Engine/Source/Runtime/Renderer/Private/RendererScene.cpp b/Engine/Source/Runtime/Renderer/Private/RendererScene.cpp
--- a/Engine/Source/Runtime/Renderer/Private/RendererScene.cpp
+++ b/Engine/Source/Runtime/Renderer/Private/RendererScene.cpp
@@ -1821,6 +1821,9 @@ void FScene::UpdatePrimitivesDrawnInGame_RenderThread(TArrayView<FPrimitiveScene
 #if RHI_RAYTRACING
 			if (bPrimitiveIndexValid && PrimitiveSceneProxy->HasRayTracingRepresentation())
 			{
+				FScene::FPrimitiveRayTracingData& PrimitiveRayTracingData = PrimitiveRayTracingDatas[PrimitiveIndex];
+				PrimitiveRayTracingData.bDrawInGame = bDrawnInGame;
+
 				FRayTracingInstance CachedRayTracingInstance;
 				ERayTracingPrimitiveFlags& Flags = PrimitiveRayTracingFlags[PrimitiveIndex];
 
