[HWRT] Add events to ray tracing dynamic geometry dispatches.

#tests replayrun
#rnx

diff --git a/Engine/Source/Runtime/Renderer/Private/RayTracing/RayTracingDynamicGeometryUpdateManager.cpp b/Engine/Source/Runtime/Renderer/Private/RayTracing/RayTracingDynamicGeometryUpdateManager.cpp
--- a/Engine/Source/Runtime/Renderer/Private/RayTracing/RayTracingDynamicGeometryUpdateManager.cpp
+++ b/Engine/Source/Runtime/Renderer/Private/RayTracing/RayTracingDynamicGeometryUpdateManager.cpp
@@ -10,6 +10,7 @@
 #include "PSOPrecacheMaterial.h"
 #include "PSOPrecacheValidation.h"
 #include "Rendering/RayTracingGeometryManager.h"
+#include "UnrealEngine.h"
 
 #if RHI_RAYTRACING
 
@@ -601,6 +602,12 @@ void FRayTracingDynamicGeometryUpdateManager::AddDispatchCommands(
 					}
 				#endif
 
+#if WANTS_DRAW_MESH_EVENTS
+					DispatchCmd.Geometry = UpdateParams.Geometry;
+					DispatchCmd.MinVertexIndex = MinVertexIndex;
+					DispatchCmd.bApplyWorldPositionOffset = UpdateParams.bApplyWorldPositionOffset;
+#endif
+
 					GeometryBuildParams.DispatchCommands.Add(DispatchCmd);
 
 					break;
@@ -949,6 +956,16 @@ void FRayTracingDynamicGeometryUpdateManager::DispatchUpdates(FRHICommandList& R
 					ShaderBindingState = FShaderBindingState();
 				}
 
+#if WANTS_DRAW_MESH_EVENTS				
+				SCOPED_CONDITIONAL_DRAW_EVENTF(RHICmdList, RayTracingDynamicGeometryDispatch, GShowMaterialDrawEvents != 0, 
+					TEXT("%s - NumVertices:%d MinVertexIndex:%d WPO:%d")
+					, Cmd.Geometry->Initializer.DebugName
+					, Cmd.NumCPUVertices
+					, Cmd.MinVertexIndex
+					, (Cmd.bApplyWorldPositionOffset ? 1 : 0)
+				);
+#endif
+
 				FRHIBatchedShaderParameters& BatchedParameters = RHICmdList.GetScratchShaderParameters();
 
 				FRWBuffer* TargetBuffer = Cmd.TargetBuffer;
diff --git a/Engine/Source/Runtime/Renderer/Public/RayTracingDynamicGeometryUpdateManager.h b/Engine/Source/Runtime/Renderer/Public/RayTracingDynamicGeometryUpdateManager.h
--- a/Engine/Source/Runtime/Renderer/Public/RayTracingDynamicGeometryUpdateManager.h
+++ b/Engine/Source/Runtime/Renderer/Public/RayTracingDynamicGeometryUpdateManager.h
@@ -30,6 +30,12 @@ struct FMeshComputeDispatchCommand
 	uint32 NumThreads;
 	uint32 NumCPUVertices;
 	FRWBuffer* TargetBuffer;
+
+#if WANTS_DRAW_MESH_EVENTS
+	FRayTracingGeometry* Geometry;	
+	uint32 MinVertexIndex;
+	bool bApplyWorldPositionOffset;
+#endif 
 };
 
 class FRayTracingDynamicGeometryUpdateManager
