[HWRT] Update the entire ray tracing cached state when updating a primitive.

#tests CitySample perf run
#rnx

diff --git a/Engine/Source/Runtime/Renderer/Private/RendererScene.cpp b/Engine/Source/Runtime/Renderer/Private/RendererScene.cpp
--- a/Engine/Source/Runtime/Renderer/Private/RendererScene.cpp
+++ b/Engine/Source/Runtime/Renderer/Private/RendererScene.cpp
@@ -1824,12 +1824,7 @@ void FScene::UpdatePrimitivesDrawnInGame_RenderThread(TArrayView<FPrimitiveScene
 				FScene::FPrimitiveRayTracingData& PrimitiveRayTracingData = PrimitiveRayTracingDatas[PrimitiveIndex];
 				PrimitiveRayTracingData.bDrawInGame = bDrawnInGame;
 
-				FRayTracingInstance CachedRayTracingInstance;
-				ERayTracingPrimitiveFlags& Flags = PrimitiveRayTracingFlags[PrimitiveIndex];
-
-				// Write flags
-				Flags = PrimitiveSceneInfo->Proxy->GetCachedRayTracingInstance(CachedRayTracingInstance);
-				FPrimitiveSceneInfo::UpdateCachedRayTracingInstance(PrimitiveSceneInfo, CachedRayTracingInstance, Flags);
+				UpdateCachedRayTracingState(PrimitiveSceneProxy);
 			}
 #endif
 		}
