Fix type for r.MotionBlur.Visualize

diff --git a/Engine/Source/Runtime/Renderer/Private/PostProcess/PostProcessMotionBlur.cpp b/Engine/Source/Runtime/Renderer/Private/PostProcess/PostProcessMotionBlur.cpp
--- a/Engine/Source/Runtime/Renderer/Private/PostProcess/PostProcessMotionBlur.cpp
+++ b/Engine/Source/Runtime/Renderer/Private/PostProcess/PostProcessMotionBlur.cpp
@@ -76,7 +76,7 @@ namespace
 	);
 
 	TAutoConsoleVariable<int32> CVarMotionBlurVisualize(
-		TEXT("vr.MotionBlur.Visualize"),
+		TEXT("r.MotionBlur.Visualize"),
 		0,
 		TEXT("0: Default view.\n")
 		TEXT("1: HSV view. Encode direction in H, and magnitude in V.\n")
