[VSM] Fixed unbatched non-nanite rendering (it was using the incorrect load balancer).

diff --git a/Engine/Source/Runtime/Renderer/Private/VirtualShadowMaps/VirtualShadowMapArray.cpp b/Engine/Source/Runtime/Renderer/Private/VirtualShadowMaps/VirtualShadowMapArray.cpp
--- a/Engine/Source/Runtime/Renderer/Private/VirtualShadowMaps/VirtualShadowMapArray.cpp
+++ b/Engine/Source/Runtime/Renderer/Private/VirtualShadowMaps/VirtualShadowMapArray.cpp
@@ -3899,7 +3899,7 @@ void FVirtualShadowMapArray::RenderVirtualShadowMapsNonNanite(FRDGBuilder& Graph
 				InstanceCullingContext->IndirectArgs, 
 				InstanceCullingContext->DrawCommandDescs,
 				InstanceCullingContext->InstanceIdOffsets,
-				InstanceCullingContext->LoadBalancers[FInstanceCullingMergedContext::FirstGenericBinIndex],
+				InstanceCullingContext->LoadBalancers[0],
 				MakeArrayView(&CullingBatchInfo, 1),
 				MakeArrayView(&VSMCullingBatchInfo, 1),
 				MakeArrayView<const uint32>(nullptr, 0),
