SkyAtmosphere - Fix lut caching bug when resources are re-created.

#tests pc editor
#rnx

diff --git a/Engine/Source/Runtime/Renderer/Private/SkyAtmosphereRendering.cpp b/Engine/Source/Runtime/Renderer/Private/SkyAtmosphereRendering.cpp
--- a/Engine/Source/Runtime/Renderer/Private/SkyAtmosphereRendering.cpp
+++ b/Engine/Source/Runtime/Renderer/Private/SkyAtmosphereRendering.cpp
@@ -1202,6 +1202,8 @@ void InitSkyAtmosphereForScene(FRHICommandListImmediate& RHICmdList, FRDGBuilder
 		EPixelFormat TextureLUTFormat = GetSkyLutTextureFormat(Scene->GetFeatureLevel());
 		EPixelFormat TextureLUTSmallFormat = GetSkyLutSmallTextureFormat();
 
+		bool bCachedLUTsResourcesAreTheSame = true;
+
 		//
 		// Initialise per scene/atmosphere resources
 		//
@@ -1213,7 +1215,7 @@ void InitSkyAtmosphereForScene(FRHICommandListImmediate& RHICmdList, FRDGBuilder
 			Desc = FPooledRenderTargetDesc::Create2DDesc(
 				FIntPoint(TransmittanceLutWidth, TransmittanceLutHeight),
 				TranstmittanceLUTUseSmallFormat ? TextureLUTSmallFormat : TextureLUTFormat, FClearValueBinding::None, TexCreate_None, TexCreate_ShaderResource | TexCreate_UAV, false);
-			GRenderTargetPool.FindFreeElement(RHICmdList, Desc, TransmittanceLutTexture, TEXT("SkyAtmosphere.TransmittanceLut"));
+			bCachedLUTsResourcesAreTheSame = GRenderTargetPool.FindFreeElement(RHICmdList, Desc, TransmittanceLutTexture, TEXT("SkyAtmosphere.TransmittanceLut")) && bCachedLUTsResourcesAreTheSame;
 		}
 		else
 		{
@@ -1224,7 +1226,13 @@ void InitSkyAtmosphereForScene(FRHICommandListImmediate& RHICmdList, FRDGBuilder
 		Desc = FPooledRenderTargetDesc::Create2DDesc(
 			FIntPoint(MultiScatteredLuminanceLutWidth, MultiScatteredLuminanceLutHeight),
 			TextureLUTFormat, FClearValueBinding::None, TexCreate_None, TexCreate_ShaderResource | TexCreate_UAV, false);
-		GRenderTargetPool.FindFreeElement(RHICmdList, Desc, MultiScatteredLuminanceLutTexture, TEXT("SkyAtmosphere.MultiScatteredLuminanceLut"));
+		bCachedLUTsResourcesAreTheSame = GRenderTargetPool.FindFreeElement(RHICmdList, Desc, MultiScatteredLuminanceLutTexture, TEXT("SkyAtmosphere.MultiScatteredLuminanceLut")) && bCachedLUTsResourcesAreTheSame;
+
+		// If the resources have changed, we must recache them.
+		if (!bCachedLUTsResourcesAreTheSame)
+		{
+			Scene->SetCachedTransmittanceAndMultiScatteringLUTsVersion(0);
+		}
 
 		if (CVarSkyAtmosphereDistantSkyLightLUT.GetValueOnRenderThread() > 0)
 		{
