[VSM] Changed initialization to always retain cache data if UseVirtualShadowMaps returns true, otherwise we loose all cache data if a render with shadows disabled is performed.

diff --git a/Engine/Source/Runtime/Renderer/Private/VirtualShadowMaps/VirtualShadowMapArray.cpp b/Engine/Source/Runtime/Renderer/Private/VirtualShadowMaps/VirtualShadowMapArray.cpp
--- a/Engine/Source/Runtime/Renderer/Private/VirtualShadowMaps/VirtualShadowMapArray.cpp
+++ b/Engine/Source/Runtime/Renderer/Private/VirtualShadowMaps/VirtualShadowMapArray.cpp
@@ -660,6 +660,9 @@ void FVirtualShadowMapArray::Initialize(
 	bEnabled = bInEnabled;
 	CacheManager = InCacheManager;
 
+	// Always retain the VSM cache data if the feature is enabled.
+	const bool bShouldRetainCacheInfo = UseVirtualShadowMaps(Scene.GetShaderPlatform(), Scene.GetFeatureLevel());
+
 	bCullBackfacingPixels = CVarCullBackfacingPixels.GetValueOnRenderThread() != 0;
 	bUseHzbOcclusion = CVarShadowsVirtualUseHZB.GetValueOnRenderThread() != 0;
 	UniformParameters.NumFullShadowMaps = 0;
@@ -713,7 +716,7 @@ void FVirtualShadowMapArray::Initialize(
 	UniformParameters.PerViewData = MakeEmptyVirtualShadowMapPerViewParameters(GraphBuilder);
 	UniformParameters.CachePrimitiveAsDynamic = GraphBuilder.CreateSRV(GSystemTextures.GetDefaultStructuredBuffer(GraphBuilder, sizeof(uint32)));
 
-	if (bEnabled)
+	if (bShouldRetainCacheInfo)
 	{
 		// Always reserve IDs for the single-page SMs.
 		NumShadowMapSlots = VSM_MAX_SINGLE_PAGE_SHADOW_MAPS;
@@ -799,7 +802,7 @@ void FVirtualShadowMapArray::Initialize(
 		UniformParameters.PhysicalPagePool = GSystemTextures.GetZeroUIntArrayAtomicCompatDummy(GraphBuilder);
 	}
 
-	if (bEnabled && bUseHzbOcclusion)
+	if (bShouldRetainCacheInfo && bUseHzbOcclusion)
 	{
 		const int HZBPoolArraySize = HasSeparateDynamicHZB() ? 2 : 1;
 		HZBPhysicalArray = CacheManager->SetHZBPhysicalPoolSize(GraphBuilder, GetHZBPhysicalPoolSize(), HZBPoolArraySize, PF_R32_FLOAT);
