Move ensure when missing SBT allocation and fix ensure to match the SBT allocation function

diff --git a/Engine/Source/Runtime/Renderer/Private/RayTracing/RayTracing.cpp b/Engine/Source/Runtime/Renderer/Private/RayTracing/RayTracing.cpp
--- a/Engine/Source/Runtime/Renderer/Private/RayTracing/RayTracing.cpp
+++ b/Engine/Source/Runtime/Renderer/Private/RayTracing/RayTracing.cpp
@@ -1379,10 +1379,10 @@ namespace RayTracing
 					if (SceneInfo->GetRayTracingLODDataNum() > uint32(LODIndex))
 					{
 						const FPrimitiveSceneInfo::FRayTracingLODData& RTLODData = SceneInfo->GetRayTracingLODData(LODIndex);
-
-						ensure(RTLODData.SBTAllocation);
 						if (RTLODData.SBTAllocation == nullptr)
 						{
+							// No SBT allocation if no valid segments (see logic FRayTracingShaderBindingTable::AllocateStaticRange)
+							ensure(RTLODData.CachedMeshCommandFlags.bAllSegmentsDecal && !(RTLODData.CachedMeshCommandFlags.bAnySegmentsDecal && !RayTracing::ShouldExcludeDecals()));
 							return;
 						}
 						const bool bUsesLightingChannels = bLightingChannelsUsingAHS && PrimitiveRayTracingData.bUsesLightingChannels;
