Material Cache, added cache flushing on invalidation

diff --git a/Engine/Source/Runtime/Engine/Internal/MaterialCache/MaterialCacheVirtualTextureInterface.h b/Engine/Source/Runtime/Engine/Internal/MaterialCache/MaterialCacheVirtualTextureInterface.h
--- a/Engine/Source/Runtime/Engine/Internal/MaterialCache/MaterialCacheVirtualTextureInterface.h
+++ b/Engine/Source/Runtime/Engine/Internal/MaterialCache/MaterialCacheVirtualTextureInterface.h
@@ -16,6 +16,12 @@ public:
 		checkNoEntry();
 		return FUintVector2::ZeroValue;
 	}
+
+	/** Flush all pages */
+	virtual void Flush()
+	{
+		checkNoEntry();
+	}
 };
 
 UCLASS(Abstract)
diff --git a/Engine/Source/Runtime/Renderer/Internal/MaterialCache/MaterialCacheVirtualTextureRenderProxy.h b/Engine/Source/Runtime/Renderer/Internal/MaterialCache/MaterialCacheVirtualTextureRenderProxy.h
--- a/Engine/Source/Runtime/Renderer/Internal/MaterialCache/MaterialCacheVirtualTextureRenderProxy.h
+++ b/Engine/Source/Runtime/Renderer/Internal/MaterialCache/MaterialCacheVirtualTextureRenderProxy.h
@@ -16,9 +16,10 @@ public: /** UMaterialCacheVirtualTextureRenderProxyInterface */
 		return TextureDescriptor;
 	}
 
-public:
+	virtual void Flush() override;
+
 	/** Virtual texture resource */
-	FTextureResource* TextureResource = nullptr;
+	FVirtualTexture2DResource* TextureResource = nullptr;
 
 	/** Optional, owned stack provider render proxy */
 	TUniquePtr<FMaterialCacheStackProviderRenderProxy> StackProviderRenderProxy;
diff --git a/Engine/Source/Runtime/Renderer/Private/MaterialCache/MaterialCacheVirtualTexture.cpp b/Engine/Source/Runtime/Renderer/Private/MaterialCache/MaterialCacheVirtualTexture.cpp
--- a/Engine/Source/Runtime/Renderer/Private/MaterialCache/MaterialCacheVirtualTexture.cpp
+++ b/Engine/Source/Runtime/Renderer/Private/MaterialCache/MaterialCacheVirtualTexture.cpp
@@ -292,9 +292,15 @@ void UMaterialCacheVirtualTexture::CreateSceneProxy()
 
 FMaterialCacheVirtualTextureRenderProxyInterface* UMaterialCacheVirtualTexture::CreateRenderProxy(uint32 UVCoordinateIndex)
 {
+	FVirtualTexture2DResource* VTResource = GetVirtualTexture2DResource();
+	if (!VTResource)
+	{
+		return nullptr;
+	}
+	
 	UMaterialCacheVirtualTextureRenderProxy* Proxy = new UMaterialCacheVirtualTextureRenderProxy();
-	Proxy->TextureResource = GetResource();
-	Proxy->UVCoordinateIndex  = UVCoordinateIndex;
+	Proxy->TextureResource   = VTResource;
+	Proxy->UVCoordinateIndex = UVCoordinateIndex;
 
 	// Stack providers are optional
 	if (UMaterialCacheStackProvider* StackProvider = MaterialStackProvider.Get())
@@ -508,3 +514,10 @@ FScene* UMaterialCacheVirtualTexture::GetRenderScene()
 	return Scene->GetRenderScene();
 }
 
+void UMaterialCacheVirtualTextureRenderProxy::Flush()
+{
+	if (IAllocatedVirtualTexture* AllocatedVT = TextureResource->GetAllocatedVT())
+	{
+		GetRendererModule().FlushVirtualTextureCache(AllocatedVT, FVector2f(0, 0), FVector2f(1, 1));
+	}
+}
diff --git a/Engine/Source/Runtime/Renderer/Private/PrimitiveSceneInfo.cpp b/Engine/Source/Runtime/Renderer/Private/PrimitiveSceneInfo.cpp
--- a/Engine/Source/Runtime/Renderer/Private/PrimitiveSceneInfo.cpp
+++ b/Engine/Source/Runtime/Renderer/Private/PrimitiveSceneInfo.cpp
@@ -2157,6 +2157,14 @@ void FPrimitiveSceneInfo::FlushRuntimeVirtualTexture()
 			}
 		}
 	}
+
+	if (Proxy)
+	{
+		if (FMaterialCacheVirtualTextureRenderProxyInterface* MaterialCacheProxy = Proxy->GetMaterialCacheRenderProxy())
+		{
+			MaterialCacheProxy->Flush();
+		}
+	}
 }
 
 void FPrimitiveSceneInfo::LinkLODParentComponent()
