Fixed SA warning on nullptr access

diff --git a/Engine/Source/Runtime/Renderer/Private/PrimitiveSceneInfo.cpp b/Engine/Source/Runtime/Renderer/Private/PrimitiveSceneInfo.cpp
--- a/Engine/Source/Runtime/Renderer/Private/PrimitiveSceneInfo.cpp
+++ b/Engine/Source/Runtime/Renderer/Private/PrimitiveSceneInfo.cpp
@@ -2146,6 +2146,11 @@ FInstanceDataBufferHeader FPrimitiveSceneInfo::GetInstanceDataHeader() const
 
 void FPrimitiveSceneInfo::FlushRuntimeVirtualTexture()
 {
+	if (!Proxy)
+	{
+		return;
+	}
+	
 	if (bWritesRuntimeVirtualTexture)
 	{
 		for (TSparseArray<FRuntimeVirtualTextureSceneProxy*>::TIterator It(Scene->RuntimeVirtualTextures); It; ++It)
@@ -2157,12 +2162,9 @@ void FPrimitiveSceneInfo::FlushRuntimeVirtualTexture()
 		}
 	}
 
-	if (Proxy)
+	if (FMaterialCacheVirtualTextureRenderProxyInterface* MaterialCacheProxy = Proxy->GetMaterialCacheRenderProxy())
 	{
-		if (FMaterialCacheVirtualTextureRenderProxyInterface* MaterialCacheProxy = Proxy->GetMaterialCacheRenderProxy())
-		{
-			MaterialCacheProxy->Flush();
-		}
+		MaterialCacheProxy->Flush();
 	}
 }
 
