Workaround for overflowing instance payload buffer causing check/crash/editor shutdown.

diff --git a/Engine/Source/Runtime/Renderer/Private/GPUScene.cpp b/Engine/Source/Runtime/Renderer/Private/GPUScene.cpp
--- a/Engine/Source/Runtime/Renderer/Private/GPUScene.cpp
+++ b/Engine/Source/Runtime/Renderer/Private/GPUScene.cpp
@@ -78,6 +78,9 @@ END_GLOBAL_SHADER_PARAMETER_STRUCT()
 
 IMPLEMENT_GLOBAL_SHADER_PARAMETER_STRUCT(FGPUSceneWriterUniformParameters, "GPUSceneWriter");
 
+// Max buffer size is 2GB and also largest POT in 32 bits. 2GB / sizeof(float4) <=> 2 << (30 - 4)
+constexpr uint64 MaxPayloadBufferSizeInFloat4 = uint64(2) << (30 - 4);
+
 static TAutoConsoleVariable<int32> CVarGPUSceneUploadEveryFrame(
 	TEXT("r.GPUScene.UploadEveryFrame"),
 	0,
@@ -300,9 +303,11 @@ struct FPrimitiveUploadInfoHeader
 
 	/** Optional */
 	int32 NumInstanceUploads = 0;
-	int32 NumInstancePayloadDataUploads = 0;
+	int32 InstancePayloadDataOffset = 0;
+	int32 InstancePayloadDataStride = 0;
 	int32 LightmapUploadCount = 0;
-
+	// Derived: Number of instances to upload payload for, may be clamped to fit in the buffer.
+	int32 NumPayloadUploadInstances = 0;
 	/** NaniteSceneProxy must be set if the proxy is a Nanite proxy */
 	const Nanite::FSceneProxyBase* NaniteSceneProxy = nullptr;
 	const FPrimitiveSceneInfo* PrimitiveSceneInfo = nullptr;
@@ -391,18 +396,6 @@ void ValidateInstanceUploadInfo(const FInstanceUploadInfo& UploadInfo, FRDGBuffe
 		check(UploadInfo.InstancePayloadExtensionCount == 0 && UploadInfo.InstancePayloadExtension.Num() == 0);
 	}
 
-	// RandomID is not stored in the payload but in the instance scene data.
-	const bool bHasAnyPayloadData = bHasHierarchyOffset || bHasLocalBounds || bHasDynamicData || bHasSkinningData || bHasLightShadowUVBias || bHasCustomData|| bHasPayloadExtension;
-
-	if (bHasAnyPayloadData)
-	{
-		check(InstancePayloadDataBuffer);
-		check(UploadInfo.InstancePayloadDataOffset != INDEX_NONE);
-
-		const int32 PayloadBufferSize = InstancePayloadDataBuffer->GetSize() / InstancePayloadDataBuffer->GetStride();
-		check(UploadInfo.InstancePayloadDataOffset < PayloadBufferSize);
-	}
-
 	if (UploadInfo.InstanceSceneDataBuffers != nullptr)
 	{
 		if (UploadInfo.InstanceSceneDataBuffers->GetNumInstances() > 0 && !ensure(UploadInfo.InstanceSceneDataBuffers->GetPrimitiveToRelativeWorld().Equals(UploadInfo.PrimitiveToWorld)))
@@ -483,16 +476,16 @@ struct FUploadDataSourceAdapterScenePrimitives
 		PrimitiveUploadInfo.NaniteSceneProxy = PrimitiveSceneProxy->IsNaniteMesh() ? static_cast<const Nanite::FSceneProxyBase*>(PrimitiveSceneProxy) : nullptr;
 		PrimitiveUploadInfo.PrimitiveSceneInfo = PrimitiveSceneInfo;
 
+		PrimitiveUploadInfo.InstancePayloadDataOffset = PrimitiveSceneInfo->GetInstancePayloadDataOffset();
+		PrimitiveUploadInfo.InstancePayloadDataStride = PrimitiveSceneInfo->GetInstancePayloadDataStride();
 		// TODO: Somehow associate the procedural update with the update entry such that we can cross validate the existence of a writer here?
 		if (PrimitiveSceneProxy->IsInstanceDataGPUOnly())
 		{
 			PrimitiveUploadInfo.NumInstanceUploads = 0;
-			PrimitiveUploadInfo.NumInstancePayloadDataUploads = 0;
 		}
 		else
 		{
 			PrimitiveUploadInfo.NumInstanceUploads = PrimitiveSceneInfo->GetNumInstanceSceneDataEntries();
-			PrimitiveUploadInfo.NumInstancePayloadDataUploads = PrimitiveSceneInfo->GetInstancePayloadDataStride() * PrimitiveUploadInfo.NumInstanceUploads;
 		}
 	}
 
@@ -687,6 +680,21 @@ FGPUScene::FGPUScene(FScene &InScene)
 			bLoggedInstanceOverflow = false;
 		}
 
+		if (InstancePayloadDataSizeReserveUnclamped  >= uint64(MaxPayloadBufferSizeInFloat4))
+		{
+			OutMessages.Add(FCoreDelegates::EOnScreenMessageSeverity::Warning, FText::FromString(FString::Printf(TEXT("GPU-Scene Payload data overflow detected, reduce instance count and/or e.g, use of custom data to avoid rendering artifacts and potentially crashes (depending on the use of payload data)"))));
+			OutMessages.Add(FCoreDelegates::EOnScreenMessageSeverity::Warning, FText::FromString(FString::Printf(TEXT("  Buffer size requested %dM, max instance buffer size (in float4s): %dM"), uint32(InstancePayloadDataSizeReserveUnclamped >> 20), MaxPayloadBufferSizeInFloat4 >> 20)));
+			if (!bLoggedPayloadBufferOverflow )
+			{
+				UE_LOG(LogRenderer, Warning, TEXT("GPU-Scene Payload data overflow detected, reduce instance count and/or e.g, use of custom data to avoid rendering artifacts and potentially crashes (depending on the use of payload data).\n")
+					TEXT(" Attempted allocation (POT) %dM, instance buffer size: %dM"), uint32(InstancePayloadDataSizeReserveUnclamped >> 20),  MaxPayloadBufferSizeInFloat4 >> 20);
+				bLoggedPayloadBufferOverflow = true;
+			}
+		}
+		else
+		{
+			bLoggedPayloadBufferOverflow = false;
+		}
 		if (NumGPULights < MaxLightsDuringPrevUpdate)
 		{
 			OutMessages.Add(FCoreDelegates::EOnScreenMessageSeverity::Warning, FText::FromString(FString::Printf(TEXT("GPU-Scene Light data overflow detected, reduce light count to avoid rendering artifacts"))));
@@ -1078,7 +1086,13 @@ FGPUScene::FRegisteredBuffers FGPUScene::UpdateBufferAllocations(FRDGBuilder& Gr
 	}
 
 	const uint32 PayloadFloat4Count = FMath::Max(PrevDynamicPrimitivePayloadMaxSize, FMath::Max(InstancePayloadDataAllocator.GetMaxSize(), InitialBufferSize));
-	const uint32 InstancePayloadDataSizeReserve = FMath::RoundUpToPowerOfTwo(PayloadFloat4Count * sizeof(FVector4f));
+	InstancePayloadDataSizeReserveUnclamped = FMath::RoundUpToPowerOfTwo(uint64(PayloadFloat4Count));
+	uint32 InstancePayloadDataSizeReserve = uint32(InstancePayloadDataSizeReserveUnclamped * sizeof(FVector4f));
+
+	if (InstancePayloadDataSizeReserveUnclamped > MaxPayloadBufferSizeInFloat4)
+	{
+		InstancePayloadDataSizeReserve = MaxPayloadBufferSizeInFloat4 * sizeof(FVector4f); 
+	}
 	BufferState.InstancePayloadDataBuffer = AllocOrResizeBufferIfNeeded(GraphBuilder, InstancePayloadDataBuffer, InstancePayloadDataSizeReserve, TEXT("GPUScene.InstancePayloadData"));
 
 	const uint32 LightMapDataBufferSize = FMath::RoundUpToPowerOfTwo(FMath::Max(LightmapDataAllocator.GetMaxSize(), InitialBufferSize));
@@ -1276,14 +1290,26 @@ void FGPUScene::UploadGeneral(FRDGBuilder& GraphBuilder, const FRegisteredBuffer
 	TaskContext.InstanceDataTileStride   = GetInstanceDataTileStride();
 	TaskContext.PrimitiveUploadInfos.SetNumUninitialized(NumPrimitiveDataUploads);
 
+	// Max offset to scatter to.
+	const uint32 MaxInstancePayloadOffset = InstancePayloadDataBuffer->Desc.NumElements;
+
 	for (int32 ItemIndex = 0; ItemIndex < NumPrimitiveDataUploads; ++ItemIndex)
 	{
-		FPrimitiveUploadInfoHeader& UploadInfo = TaskContext.PrimitiveUploadInfos[ItemIndex];
+		FPrimitiveUploadInfoHeader UploadInfo;
 		UploadDataSourceAdapter.GetPrimitiveInfoHeader(ItemIndex, UploadInfo);
-
+		UploadInfo.NumPayloadUploadInstances = UploadInfo.InstancePayloadDataStride ? UploadInfo.NumInstanceUploads : 0;
+		// Figure out how much space we actually have left.
+		int32 PayloadSpaceLeft = FMath::Max(int32(0), int32(MaxInstancePayloadOffset) - int32(UploadInfo.InstancePayloadDataOffset));
+		// If we overflow, clamp the number of instances to upload payload data for to fit inside the buffer.
+		if (UploadInfo.NumPayloadUploadInstances * UploadInfo.InstancePayloadDataStride > PayloadSpaceLeft)
+		{
+			UploadInfo.NumPayloadUploadInstances = PayloadSpaceLeft / UploadInfo.InstancePayloadDataStride;
+		}
 		TaskContext.NumLightmapDataUploads += UploadInfo.LightmapUploadCount; // Not thread safe
 		TaskContext.NumInstanceSceneDataUploads += UploadInfo.NumInstanceUploads; // Not thread safe
-		TaskContext.NumInstancePayloadDataUploads += UploadInfo.NumInstancePayloadDataUploads; // Not thread safe
+		TaskContext.NumInstancePayloadDataUploads += UploadInfo.NumPayloadUploadInstances * UploadInfo.InstancePayloadDataStride; // Not thread safe
+
+		TaskContext.PrimitiveUploadInfos[ItemIndex] = UploadInfo;
 	}
 
 	TaskContext.PrimitiveUploader = PrimitiveUploadBuffer.Begin(GraphBuilder, BufferState.PrimitiveBuffer, UploadDataSourceAdapter.GetItemPrimitiveIds().Num(), sizeof(FPrimitiveSceneShaderData::Data), TEXT("PrimitiveUploadBuffer"));
@@ -1335,7 +1361,7 @@ void FGPUScene::UploadGeneral(FRDGBuilder& GraphBuilder, const FRegisteredBuffer
 				InstanceUpdates.QueueInstances(UploadInfo, ItemIndex, PrimitiveItemInfo);
 
 				NumInstanceSceneDataUploads += UploadInfo.NumInstanceUploads;
-				NumInstancePayloadDataUploads += UploadInfo.NumInstancePayloadDataUploads;
+				NumInstancePayloadDataUploads += UploadInfo.NumPayloadUploadInstances * UploadInfo.InstancePayloadDataStride;
 			}
 
 			TaskContext.PrimitiveUploader->Add(UploadDataSourceAdapter.GetItemPrimitiveIds());
@@ -1369,6 +1395,7 @@ void FGPUScene::UploadGeneral(FRDGBuilder& GraphBuilder, const FRegisteredBuffer
 					UploadDataSourceAdapter.GetInstanceInfo(ItemIndex, UploadInfo);
 					ValidateInstanceUploadInfo(UploadInfo, TaskContext.InstancePayloadUploader ? GetAsBuffer(TaskContext.InstancePayloadUploader->GetDstResource()) : nullptr);
 					FInstanceBatcher::FPrimitiveItemInfo PrimitiveItemInfo = InstanceUpdates.PerPrimitiveItemInfo[ItemIndex];
+					FPrimitiveUploadInfoHeader PrimitiveHeader = TaskContext.PrimitiveUploadInfos[ItemIndex];
 
 					check(TaskContext.NumInstancePayloadDataUploads > 0 || UploadInfo.InstancePayloadDataStride == 0); // Sanity check
 
@@ -1434,7 +1461,8 @@ void FGPUScene::UploadGeneral(FRDGBuilder& GraphBuilder, const FRegisteredBuffer
 							*DstVector = InstanceSceneData.Data[RefIndex];
 						}
 
-						if (UploadInfo.InstancePayloadDataStride > 0)
+						// Skip uploading if overflowing the buffer
+						if (InstanceIndex < PrimitiveHeader.NumPayloadUploadInstances)
 						{
 							const uint32 UploadPayloadItemOffset = PrimitiveItemInfo.InstancePayloadDataUploadOffset + InstanceIndex * UploadInfo.InstancePayloadDataStride;
 
@@ -1649,15 +1677,23 @@ struct FUploadDataSourceAdapterDynamicPrimitives
 
 		const FGPUScenePrimitiveCollector::FPrimitiveData& PrimData = PrimitiveData[ItemIndex];
 		PrimitiveUploadInfo.NumInstanceUploads = PrimData.NumInstances;
-		PrimitiveUploadInfo.NumInstancePayloadDataUploads = PrimData.SourceData.GetPayloadFloat4Stride() * PrimData.NumInstances;
-		
+		if (PrimData.LocalPayloadDataOffset != INDEX_NONE)
+		{
+			PrimitiveUploadInfo.InstancePayloadDataOffset = PayloadStartOffset + PrimData.LocalPayloadDataOffset;
+			PrimitiveUploadInfo.InstancePayloadDataStride = PrimData.SourceData.GetPayloadFloat4Stride();
+		}
+		else
+		{
+			PrimitiveUploadInfo.InstancePayloadDataOffset = INDEX_NONE;
+			PrimitiveUploadInfo.InstancePayloadDataStride = 0;
+		}
 		if (PrimData.SourceData.DataWriterGPU.IsBound())
 		{
 			// Only upload if we have data, otherwise expect the delegate to handle missing data
 			PrimitiveUploadInfo.NumInstanceUploads = PrimData.SourceData.InstanceSceneData.Num();
 			if (PrimData.SourceData.InstanceCustomData.Num() == 0)
 			{
-				PrimitiveUploadInfo.NumInstancePayloadDataUploads = 0;
+				PrimitiveUploadInfo.InstancePayloadDataStride = 0;
 			}
 		}
 	}
@@ -1669,7 +1705,7 @@ struct FUploadDataSourceAdapterDynamicPrimitives
 		FPrimitiveUniformShaderParameters Tmp = *PrimData.ShaderParams;
 		Tmp.InstanceSceneDataOffset = InstanceIDStartOffset + PrimData.LocalInstanceSceneDataOffset;
 		Tmp.NumInstanceSceneDataEntries = PrimData.NumInstances;
-		if (PrimitiveData[ItemIndex].LocalPayloadDataOffset != INDEX_NONE)
+		if (PrimData.LocalPayloadDataOffset != INDEX_NONE)
 		{
 			Tmp.InstancePayloadDataOffset = PayloadStartOffset + PrimData.LocalPayloadDataOffset;
 			Tmp.InstancePayloadDataStride = PrimData.SourceData.GetPayloadFloat4Stride();
diff --git a/Engine/Source/Runtime/Renderer/Private/GPUScene.h b/Engine/Source/Runtime/Renderer/Private/GPUScene.h
--- a/Engine/Source/Runtime/Renderer/Private/GPUScene.h
+++ b/Engine/Source/Runtime/Renderer/Private/GPUScene.h
@@ -505,8 +505,9 @@ private:
 	uint32 MaxInstancesDuringPrevUpdate = 0;
 	bool bLoggedLightOverflow = false;
 	int32 MaxLightsDuringPrevUpdate = 0;
+	bool bLoggedPayloadBufferOverflow = false;
 #endif // UE_BUILD_SHIPPING
-
+	uint64 InstancePayloadDataSizeReserveUnclamped = 0;
 };
 
 class FGPUSceneScopeBeginEndHelper
