[Auto Exposure] Fix exposure compensation being incorrectly evaluated when Method = Histogram and AutoExposureMinBrightness >= AutoExposureMaxBrightness.
- always need to generate histogram to calculate average scene luminance which is used to evaluate the curve.
- this was working as expected when using Basic method.

diff --git a/Engine/Source/Runtime/Renderer/Private/PostProcess/PostProcessing.cpp b/Engine/Source/Runtime/Renderer/Private/PostProcess/PostProcessing.cpp
--- a/Engine/Source/Runtime/Renderer/Private/PostProcess/PostProcessing.cpp
+++ b/Engine/Source/Runtime/Renderer/Private/PostProcess/PostProcessing.cpp
@@ -713,9 +713,7 @@ void AddPostProcessingPasses(
 			// Force the histogram on when we are visualizing HDR.
 			bVisualizeHDR ||
 			// Skip if not using histogram eye adaptation.
-			(bEyeAdaptationEnabled && AutoExposureMethod == EAutoExposureMethod::AEM_Histogram &&
-			// Skip if we don't have any exposure range to generate (eye adaptation will clamp).
-			View.FinalPostProcessSettings.AutoExposureMinBrightness < View.FinalPostProcessSettings.AutoExposureMaxBrightness);
+			(bEyeAdaptationEnabled && AutoExposureMethod == EAutoExposureMethod::AEM_Histogram);
 
 		const bool bLocalExposureEnabled =
 			EngineShowFlags.VisualizeLocalExposure ||
