SMAA debug visualization modes

diff --git a/Engine/Shaders/Private/SMAA/SMAADebugVisualization.usf b/Engine/Shaders/Private/SMAA/SMAADebugVisualization.usf
new file mode 100644
--- /dev/null
+++ b/Engine/Shaders/Private/SMAA/SMAADebugVisualization.usf
@@ -0,0 +1,79 @@
+// Copyright Epic Games, Inc. All Rights Reserved.
+
+#include "SMAACommon.ush"
+
+struct FSMAADebugVisualizationVertexOutput
+{
+	noperspective float2 UV : TEXCOORD0;
+	float4 Offset : TEXCOORD1;
+	float4 Position : SV_POSITION;
+};
+
+Texture2D EdgeTex;
+Texture2D BlendTex;
+
+int DebugMode;
+
+/*-----------------------------------------------------------------------------
+	vertex shader
+-----------------------------------------------------------------------------*/
+#ifdef MainVS
+
+	void MainVS(
+		float2 InPosition : ATTRIBUTE0,
+		float2 InUV       : ATTRIBUTE1,
+		out FSMAADebugVisualizationVertexOutput Output
+	)
+	{
+		DrawRectangle( float4( InPosition, 0, 1 ), InUV, Output.Position, Output.UV);
+
+		SMAANeighborhoodBlendingVS(Output.UV, Output.Offset);
+	}
+
+#endif // MainVS
+
+
+/*-----------------------------------------------------------------------------
+	pixel shader
+-----------------------------------------------------------------------------*/
+#ifdef MainPS
+
+	void MainPS(
+		FSMAADebugVisualizationVertexOutput Input,
+		out float4 OutColor : SV_Target0)
+	{
+		float4 color = float4(0,0,0,0);
+
+		if (DebugMode == 1)
+		{
+			color = EdgeTex.Sample(LinearSampler, Input.UV);
+		}
+		else if (DebugMode == 2)
+		{
+			color = BlendTex.Sample(LinearSampler, Input.UV);
+		}
+		else if (DebugMode == 3)
+		{
+			// sampling logic taken from SMAANeighborhoodBlendingPS, in SMAA.ush
+	
+			// Fetch the blending weights for current pixel:
+			float4 a;
+			a.x = BlendTex.Sample(LinearSampler, Input.Offset.xy).a; // Right
+			a.y = BlendTex.Sample(LinearSampler, Input.Offset.zw).g; // Top
+			a.wz = BlendTex.Sample(LinearSampler, Input.UV).xz; // Bottom / Left
+
+			// Is there any blending weight with a value greater than 0.0?
+			if (dot(a, float4(1.0, 1.0, 1.0, 1.0)) < 1e-5)
+			{
+				color = Input_Texture.SampleLevel(LinearSampler, Input.UV, 0);
+			}
+			else
+			{
+				color = float4(1.0, 0.0, 0.0, 0.0);
+			}
+		}
+
+		OutColor = color;
+	}
+
+#endif // MainPS
diff --git a/Engine/Source/Runtime/Renderer/Private/PostProcess/SubpixelMorphologicalAA.cpp b/Engine/Source/Runtime/Renderer/Private/PostProcess/SubpixelMorphologicalAA.cpp
--- a/Engine/Source/Runtime/Renderer/Private/PostProcess/SubpixelMorphologicalAA.cpp
+++ b/Engine/Source/Runtime/Renderer/Private/PostProcess/SubpixelMorphologicalAA.cpp
@@ -18,6 +18,15 @@ TAutoConsoleVariable<int32> CVarSMAAQuality(
 	TEXT(" 3: Ultra (%99 of the quality) \n"),
 	ECVF_Scalability | ECVF_RenderThreadSafe);
 
+TAutoConsoleVariable<int32> CVarSMAADebugVisualization(
+	TEXT("r.SMAA.DebugVisualization"), 0,
+	TEXT("Selects the SMAA debug visualization mode.\n")
+	TEXT(" 0: Disabled \n")
+	TEXT(" 1: Edge Texture \n")
+	TEXT(" 2: Blend Texture \n")
+	TEXT(" 3: Scene Color + Smoothened edges highlights \n"),
+	ECVF_RenderThreadSafe);
+
 #define DECLARE_SMAA_QUALITY_PERMUTATIONS               \
 class FQuality : SHADER_PERMUTATION_INT("SMAA_QUALITY", 4);
 
@@ -124,6 +133,36 @@ IMPLEMENT_GLOBAL_SHADER(FSMAANeighborhoodBlendingVS, "/Engine/Private/SMAA/SMAAN
 
 IMPLEMENT_GLOBAL_SHADER(FSMAANeighborhoodBlendingPS, "/Engine/Private/SMAA/SMAANeighborhoodBlendingShader.usf", "MainPS", SF_Pixel);
 
+class FSMAADebugVisualizationPS : public FGlobalShader
+{
+public:
+	DECLARE_GLOBAL_SHADER(FSMAADebugVisualizationPS);
+	SHADER_USE_PARAMETER_STRUCT(FSMAADebugVisualizationPS, FGlobalShader);
+
+	BEGIN_SHADER_PARAMETER_STRUCT(FParameters, )
+		SHADER_PARAMETER_STRUCT(FScreenPassTextureInput, Input)
+		SHADER_PARAMETER(FVector4f, RTMetrics)
+		SHADER_PARAMETER_RDG_TEXTURE(Texture2D, EdgeTex)
+		SHADER_PARAMETER_RDG_TEXTURE(Texture2D, BlendTex)
+		SHADER_PARAMETER_SAMPLER(SamplerState, LinearSampler)
+		SHADER_PARAMETER(int32, DebugMode)
+		RENDER_TARGET_BINDING_SLOTS()
+		END_SHADER_PARAMETER_STRUCT()
+};
+
+class FSMAADebugVisualizationVS : public FGlobalShader
+{
+public:
+	DECLARE_GLOBAL_SHADER(FSMAADebugVisualizationVS);
+	SHADER_USE_PARAMETER_STRUCT_WITH_LEGACY_BASE(FSMAADebugVisualizationVS, FGlobalShader);
+
+	using FParameters = FSMAADebugVisualizationPS::FParameters;
+};
+
+IMPLEMENT_GLOBAL_SHADER(FSMAADebugVisualizationVS, "/Engine/Private/SMAA/SMAADebugVisualization.usf", "MainVS", SF_Vertex);
+
+IMPLEMENT_GLOBAL_SHADER(FSMAADebugVisualizationPS, "/Engine/Private/SMAA/SMAADebugVisualization.usf", "MainPS", SF_Pixel);
+
 ESMAAQuality GetSMAAQuality()
 {
 	return ESMAAQuality(FMath::Clamp(CVarSMAAQuality.GetValueOnRenderThread(), 0, 3));
@@ -193,7 +232,7 @@ void AddSMAABlendingWeightCalculationPass(FRDGBuilder& GraphBuilder, const FScen
 	PassParameters->Input = GetScreenPassTextureInput(FScreenPassTexture(EdgeTexture), BilinearClampSampler);
 	PassParameters->RTMetrics = RTMetrics;
 	PassParameters->PointSampler = TStaticSamplerState<SF_Point, AM_Clamp, AM_Clamp, AM_Clamp>::GetRHI();
-	PassParameters->LinearSampler = TStaticSamplerState<SF_Bilinear, AM_Clamp, AM_Clamp, AM_Clamp>::GetRHI();
+	PassParameters->LinearSampler = BilinearClampSampler;
 
 	check(GEngine);
 	check(GEngine->SMAAAreaTexture && GEngine->SMAASearchTexture);
@@ -244,7 +283,7 @@ FScreenPassTexture AddSMAANeighborhoodBlendingPass(FRDGBuilder& GraphBuilder, co
 	PassParameters->Input = GetScreenPassTextureInput(InputSceneTexture, BilinearClampSampler);
 	PassParameters->RTMetrics = RTMetrics;
 	PassParameters->BlendTex = InputBlendTex;
-	PassParameters->LinearSampler = TStaticSamplerState<SF_Bilinear, AM_Clamp, AM_Clamp, AM_Clamp>::GetRHI();
+	PassParameters->LinearSampler = BilinearClampSampler;
 
 	FSMAANeighborhoodBlendingPS::FPermutationDomain PixelPermutationVector;
 	PixelPermutationVector.Set<FSMAANeighborhoodBlendingPS::FAlphaChannelDim>(IsPostProcessingWithAlphaChannelSupported());
@@ -269,6 +308,42 @@ FScreenPassTexture AddSMAANeighborhoodBlendingPass(FRDGBuilder& GraphBuilder, co
 	return MoveTemp(Output);
 }
 
+FScreenPassTexture AddSMAADebugVisualizationPass(FRDGBuilder& GraphBuilder, const FSceneView& InSceneView, const FScreenPassTexture& InputSceneTexture, const FRDGTextureRef& InputEdgeTex, const FRDGTextureRef& InputBlendTex, FScreenPassRenderTarget& Output, const FVector4f& RTMetrics)
+{
+	checkSlow(InSceneView.bIsViewInfo);
+	const FViewInfo& View = static_cast<const FViewInfo&>(InSceneView);
+
+	FRHISamplerState* BilinearClampSampler = TStaticSamplerState<SF_Bilinear, AM_Clamp, AM_Clamp, AM_Clamp>::GetRHI();
+
+	FSMAADebugVisualizationPS::FParameters* PassParameters = GraphBuilder.AllocParameters<FSMAADebugVisualizationPS::FParameters>();
+	PassParameters->RenderTargets[0] = Output.GetRenderTargetBinding();
+	PassParameters->Input = GetScreenPassTextureInput(InputSceneTexture, BilinearClampSampler);
+	PassParameters->RTMetrics = RTMetrics;
+	PassParameters->EdgeTex = InputEdgeTex;
+	PassParameters->BlendTex = InputBlendTex;
+	PassParameters->LinearSampler = BilinearClampSampler;
+	PassParameters->DebugMode = (int) CVarSMAADebugVisualization.GetValueOnAnyThread();
+
+	TShaderMapRef<FSMAADebugVisualizationVS> VertexShader(View.ShaderMap);
+	TShaderMapRef<FSMAADebugVisualizationPS> PixelShader(View.ShaderMap);
+	
+	AddDrawScreenPass(
+		GraphBuilder,
+		RDG_EVENT_NAME("SMAADebugVisualization"),
+		View,
+		FScreenPassTextureViewport(Output),
+		FScreenPassTextureViewport(InputSceneTexture),
+		FScreenPassPipelineState(VertexShader, PixelShader),
+		PassParameters,
+		[VertexShader, PixelShader, PassParameters] (FRHICommandList& RHICmdList)
+		{
+			SetShaderParameters(RHICmdList, VertexShader, VertexShader.GetVertexShader(), *PassParameters);
+			SetShaderParameters(RHICmdList, PixelShader, PixelShader.GetPixelShader(), *PassParameters);
+		});
+
+	return MoveTemp(Output);
+}
+
 FScreenPassTexture AddSMAAPasses(FRDGBuilder& GraphBuilder, const FSceneView& View, const FSMAAInputs& Inputs)
 {
 	check(Inputs.SceneColor.IsValid());
@@ -323,5 +398,15 @@ FScreenPassTexture AddSMAAPasses(FRDGBuilder& GraphBuilder, const FSceneView& Vi
 
 	AddSMAABlendingWeightCalculationPass(GraphBuilder, View, EdgeTexture, BlendTexture, DepthStencilTexture, Inputs.Quality, RTMetrics);
 
-	return AddSMAANeighborhoodBlendingPass(GraphBuilder, View, Inputs.SceneColor, BlendTexture, Output, RTMetrics);
+#if !UE_BUILD_SHIPPING
+	int32 DebugVisMode = CVarSMAADebugVisualization.GetValueOnAnyThread();
+	if (DebugVisMode > 0 && DebugVisMode <= 3)
+	{
+		return AddSMAADebugVisualizationPass(GraphBuilder, View, Inputs.SceneColor, EdgeTexture, BlendTexture, Output, RTMetrics);
+	}
+	else
+#endif // !UE_BUILD_SHIPPING
+	{
+		return AddSMAANeighborhoodBlendingPass(GraphBuilder, View, Inputs.SceneColor, BlendTexture, Output, RTMetrics);
+	}	
 }
