[HWRT] Call RemoveCachedRayTracingPrimitives() in UpdateCachedRayTracingState(...) instead of doing it lazily.
- Need to immediately free SBT allocations to avoid conflicts due to usage of FRHIRayTracingGeometry addresses in FAllocationKey.

diff --git a/Engine/Source/Runtime/Renderer/Private/PrimitiveSceneInfo.cpp b/Engine/Source/Runtime/Renderer/Private/PrimitiveSceneInfo.cpp
--- a/Engine/Source/Runtime/Renderer/Private/PrimitiveSceneInfo.cpp
+++ b/Engine/Source/Runtime/Renderer/Private/PrimitiveSceneInfo.cpp
@@ -2087,7 +2087,6 @@ void FPrimitiveSceneInfo::UpdateCachedRaytracingData(FScene* Scene, const TArray
 			// should have been marked dirty by calling UpdateCachedRayTracingState on the scene before
 			// scene info is being updated here
 			check(Scene->PrimitiveRayTracingDatas[SceneInfo->GetIndex()].bCachedRaytracingDataDirty);
-			SceneInfo->RemoveCachedRayTracingPrimitives();
 		}
 
 		CacheRayTracingPrimitives(Scene, SceneInfos);
diff --git a/Engine/Source/Runtime/Renderer/Private/RayTracing/RayTracingShaderBindingTable.cpp b/Engine/Source/Runtime/Renderer/Private/RayTracing/RayTracingShaderBindingTable.cpp
--- a/Engine/Source/Runtime/Renderer/Private/RayTracing/RayTracingShaderBindingTable.cpp
+++ b/Engine/Source/Runtime/Renderer/Private/RayTracing/RayTracingShaderBindingTable.cpp
@@ -444,6 +444,7 @@ FRayTracingSBTAllocation* FRayTracingShaderBindingTable::AllocateStaticRange(uin
 		TotalStaticAllocationCount++;
 	}
 	check(Allocation.Allocation->AllocatedLayers == AllocatedLayers);
+	check(Allocation.Allocation->GetSegmentCount() == FMath::CountBits((uint32)AllocatedLayers) * SegmentCount);
 	
 	Allocation.RefCount++;
 	return Allocation.Allocation;
diff --git a/Engine/Source/Runtime/Renderer/Private/RendererScene.cpp b/Engine/Source/Runtime/Renderer/Private/RendererScene.cpp
--- a/Engine/Source/Runtime/Renderer/Private/RendererScene.cpp
+++ b/Engine/Source/Runtime/Renderer/Private/RendererScene.cpp
@@ -4122,8 +4122,12 @@ void FScene::UpdateCachedRenderStates(FPrimitiveSceneProxy* SceneProxy)
 
 void FScene::UpdateCachedRayTracingState(FPrimitiveSceneProxy* SceneProxy)
 {
+	check(IsInRenderingThread());
+
 	if (SceneProxy->GetPrimitiveSceneInfo())
 	{
+		SceneProxy->GetPrimitiveSceneInfo()->RemoveCachedRayTracingPrimitives();
+
 		const int32 PrimitiveIndex = SceneProxy->GetPrimitiveSceneInfo()->GetIndex();
 		if (PrimitiveIndex != INDEX_NONE)
 		{
