[HWRT] Renamed RefreshRayTracingMeshCommandCache() to RefreshCachedRayTracingData(), and removed unused RefreshRayTracingInstances().
- RefreshRayTracingMeshCommandCache() does a lot more than it did originally, essentialy recaching all ray tracing data, so rename makes it more clear.
- RefreshRayTracingInstances() has been a common source of bugs since it updates cached flags and instances, but not mesh commands and SBT allocations, which can lead to invalid state. Since this function is not currently used, taking the opportunity to remove it, and force rethinking the implementation when similar functionality is needed in the future.

diff --git a/Engine/Source/Runtime/Renderer/Private/PrimitiveSceneInfo.cpp b/Engine/Source/Runtime/Renderer/Private/PrimitiveSceneInfo.cpp
--- a/Engine/Source/Runtime/Renderer/Private/PrimitiveSceneInfo.cpp
+++ b/Engine/Source/Runtime/Renderer/Private/PrimitiveSceneInfo.cpp
@@ -73,20 +73,6 @@ static FAutoConsoleVariableRef CVarRayTracingPrimitiveCacheMultithreaded(
 	TEXT("Enable multithreading of raytracing primitive mesh command caching. 0=disabled, 1=enabled (default)"),
 	ECVF_RenderThreadSafe);
 
-static int32 GRayTracingRefreshInstancesMultithreaded = 1;
-static FAutoConsoleVariableRef CVarRayTracingRefreshInstancesMultithreaded(
-	TEXT("r.RayTracing.RefreshInstancesMultithreaded"),
-	GRayTracingRefreshInstancesMultithreaded,
-	TEXT("Enable multithreading of ray tracing refresh primitive instances. 0=disables, 1=enabled (default)"),
-	ECVF_RenderThreadSafe);
-
-static int32 GRayTracingRefreshInstancesMultithreadedMinBatchSize = 2048;
-static FAutoConsoleVariableRef CVarRayTracingRefreshInstancesMultithreadedMinBatchSize(
-	TEXT("r.RayTracing.RefreshInstancesMultithreadedMinBatchSize"),
-	GRayTracingRefreshInstancesMultithreadedMinBatchSize,
-	TEXT("Determine minimum size of a single batch for the ray tracing refresh primitive instances."),
-	ECVF_RenderThreadSafe);
-
 /** An implementation of FStaticPrimitiveDrawInterface that stores the drawn elements for the rendering thread to use. */
 class FBatchingSPDI : public FStaticPrimitiveDrawInterface
 {
@@ -1023,7 +1009,7 @@ void FPrimitiveSceneInfo::RemoveCachedNaniteMaterialBins()
 }
 
 #if RHI_RAYTRACING
-void FScene::RefreshRayTracingMeshCommandCache()
+void FScene::RefreshCachedRayTracingData()
 {
 	// Get rid of all existing cached commands
 	for (FPrimitiveSceneInfo* SceneInfo : Primitives)
@@ -1051,51 +1037,6 @@ static void CheckPrimitiveRayTracingData(const FScene::FPrimitiveRayTracingData&
 	check(PrimitiveRayTracingData.ProxyGeometryType == SceneProxy->GetRayTracingProxyType());
 }
 
-void FScene::RefreshRayTracingInstances()
-{
-	if (IsRayTracingEnabled())
-	{
-		SCOPED_NAMED_EVENT(FPrimitiveSceneInfo_UpdateCachedRayTracingInstances, FColor::Turquoise);
-		checkf(GRHISupportsMultithreadedShaderCreation, TEXT("Raytracing code needs the ability to create shaders from task threads."));
-
-		EParallelForFlags ParallelForFlags = GRayTracingRefreshInstancesMultithreaded && FApp::ShouldUseThreadingForPerformance() ? EParallelForFlags::None : EParallelForFlags::ForceSingleThread;
-
-		ParallelFor(TEXT("FPrimitiveSceneInfo_UpdateCachedRayTracingInstances"), Primitives.Num(), GRayTracingRefreshInstancesMultithreadedMinBatchSize, [this](int32 PrimitiveSceneInfoIndex)
-			{
-				FTaskTagScope Scope(ETaskTag::EParallelRenderingThread);
-
-				FPrimitiveSceneInfo* SceneInfo = Primitives[PrimitiveSceneInfoIndex];
-
-				const int32 PrimitiveIndex = SceneInfo->GetIndex();
-				FPrimitiveSceneProxy* SceneProxy = SceneInfo->Proxy;
-
-#if DO_CHECK
-				{
-					Experimental::FHashElementId SceneRayTracingGroupId;
-					const int32 RayTracingGroupId = SceneProxy->GetRayTracingGroupId();
-					if (RayTracingGroupId != -1)
-					{
-						SceneRayTracingGroupId = PrimitiveRayTracingGroups.FindId(RayTracingGroupId);
-					}
-
-					check(PrimitiveRayTracingGroupIds[PrimitiveIndex] == SceneRayTracingGroupId);
-
-					const FScene::FPrimitiveRayTracingData& PrimitiveRayTracingData = PrimitiveRayTracingDatas[PrimitiveIndex];
-					CheckPrimitiveRayTracingData(PrimitiveRayTracingData, SceneProxy);
-				}
-#endif
-
-				FRayTracingInstance CachedRayTracingInstance;
-				ERayTracingPrimitiveFlags& Flags = PrimitiveRayTracingFlags[PrimitiveIndex];
-
-				// Write flags
-				Flags = SceneProxy->GetCachedRayTracingInstance(CachedRayTracingInstance);
-				FPrimitiveSceneInfo::UpdateCachedRayTracingInstance(SceneInfo, CachedRayTracingInstance, Flags);
-
-			}, ParallelForFlags);
-	}
-}
-
 struct FDeferredRayTracingMeshCommandData
 {
 	FPrimitiveSceneInfo* SceneInfo;
diff --git a/Engine/Source/Runtime/Renderer/Private/RayTracing/RayTracing.cpp b/Engine/Source/Runtime/Renderer/Private/RayTracing/RayTracing.cpp
--- a/Engine/Source/Runtime/Renderer/Private/RayTracing/RayTracing.cpp
+++ b/Engine/Source/Runtime/Renderer/Private/RayTracing/RayTracing.cpp
@@ -95,8 +95,6 @@ static FAutoConsoleCommand UpdateCachedRayTracingStateCmd(
 	TEXT("Update cached ray tracing state (mesh commands and instances)."),
 	FConsoleCommandDelegate::CreateStatic([] { bUpdateCachedRayTracingState = true; }));
 
-static bool bRefreshRayTracingInstances = false;
-
 static ERayTracingProxyType ActiveRayTracingProxyTypes = ERayTracingProxyType::All;
 
 static_assert(sizeof(FScene::FPrimitiveRayTracingData) == 8, "FScene::FPrimitiveRayTracingData is packed to 8 bytes to be cache efficient during GatherRelevantPrimitives");
@@ -411,22 +409,10 @@ namespace RayTracing
 			// In some situations, we need to refresh the cached ray tracing mesh commands because they contain data about the currently bound shader. 
 			// This operation is a bit expensive but only happens once as we transition between RT types which should be rare.
 			Scene.CachedRayTracingMeshCommandsType = CurrentType;
-			Scene.RefreshRayTracingMeshCommandCache();
+			Scene.RefreshCachedRayTracingData();
 			bUpdateCachedRayTracingState = false;
 		}
 
-		if (bRefreshRayTracingInstances)
-		{
-			Scene.WaitForCacheRayTracingPrimitivesTask();
-
-			// In some situations, we need to refresh the cached ray tracing instance.
-			// This assumes that cached instances will keep using the same LOD since CachedRayTracingMeshCommands is not recalculated
-			// eg: Need to update PrimitiveRayTracingFlags
-			// This operation is a bit expensive but only happens once as we transition between RT types which should be rare.
-			Scene.RefreshRayTracingInstances();
-			bRefreshRayTracingInstances = false;
-		}
-
 		if (bNaniteRayTracingModeChanged)
 		{
 			for (FViewInfo* View : SceneUpdateInputs.Views)
diff --git a/Engine/Source/Runtime/Renderer/Private/ScenePrivate.h b/Engine/Source/Runtime/Renderer/Private/ScenePrivate.h
--- a/Engine/Source/Runtime/Renderer/Private/ScenePrivate.h
+++ b/Engine/Source/Runtime/Renderer/Private/ScenePrivate.h
@@ -2900,9 +2900,7 @@ public:
 
 #if RHI_RAYTRACING
 	/** Force a refresh of all cached ray tracing data in the scene (when path tracing mode changes or coarse mesh streaming for example). */
-	void RefreshRayTracingMeshCommandCache();
-	
-	void RefreshRayTracingInstances();
+	void RefreshCachedRayTracingData();
 #endif
 
 	/** Nanite shading material commands. These are stored on the scene as they are computed at FPrimitiveSceneInfo::AddToScene time. */
