[HWRT] Expose r.RayTracing.ParallelPrimitiveGather to control whether ParallelFor loops are used to gather primitives relevant to ray tracing.
- useful for profiling.

diff --git a/Engine/Source/Runtime/Renderer/Private/RayTracing/RayTracing.cpp b/Engine/Source/Runtime/Renderer/Private/RayTracing/RayTracing.cpp
--- a/Engine/Source/Runtime/Renderer/Private/RayTracing/RayTracing.cpp
+++ b/Engine/Source/Runtime/Renderer/Private/RayTracing/RayTracing.cpp
@@ -88,6 +88,13 @@ static FAutoConsoleVariableRef CVarRayTracingDebugForceOpaque(
 	TEXT("Forces all ray tracing geometry instances to be opaque, effectively disabling any-hit shaders. This is useful for debugging and profiling. (default = 0)")
 );
 
+static int32 GRayTracingParallelPrimitiveGather = 1;
+static FAutoConsoleVariableRef CVarRayTracingParallelPrimitiveGather(
+	TEXT("r.RayTracing.ParallelPrimitiveGather"),
+	GRayTracingParallelPrimitiveGather,
+	TEXT("Whether to gather primitives relevant to ray tracing using parallel loops. 0=disabled, 1=enabled (default)"),
+	ECVF_RenderThreadSafe);
+
 static bool bUpdateCachedRayTracingState = false;
 
 static FAutoConsoleCommand UpdateCachedRayTracingStateCmd(
@@ -1118,7 +1125,7 @@ namespace RayTracing
 				{
 					Context.StaticPrimitives.AddElement(PrimitiveIndex);
 				}
-			});
+			}, GRayTracingParallelPrimitiveGather ? EParallelForFlags::None : EParallelForFlags::ForceSingleThread);
 
 			if (Contexts.Num() > 0)
 			{
@@ -1386,11 +1393,11 @@ namespace RayTracing
 						RelevantPrimitive->bUsesLightingChannels = bUsesLightingChannels;
 					}
 				}
-			});
+			}, GRayTracingParallelPrimitiveGather ? EParallelForFlags::None : EParallelForFlags::ForceSingleThread);
 
 		if (Contexts.Num() > 0)
 		{
-			SCOPED_NAMED_EVENT(GatherRayTracingRelevantStaticPrimitives__Merge, FColor::Emerald);
+			SCOPED_NAMED_EVENT(GatherRayTracingRelevantStaticPrimitives_Merge, FColor::Emerald);
 
 			uint32 NumStaticPrimitives = 0;
 			uint32 NumCachedStaticPrimitives = 0;
@@ -1732,7 +1739,7 @@ namespace RayTracing
 
 					RayTracingScene.SetInstance(CachedStaticFarFieldInstanceRange, FarFieldInstanceIndexInRange, MoveTemp(FarFieldRayTracingInstance), SceneProxy, false);
 				}
-			});
+			}, GRayTracingParallelPrimitiveGather ? EParallelForFlags::None : EParallelForFlags::ForceSingleThread);
 	}
 
 	void AddDynamicInstancesToRayTracingScene(
