Material Cache, added block compression support
- Default to BC3 on all physical textures, second layer could be BC1 due to opacity alpha

diff --git a/Engine/Shaders/Private/MaterialCache/MaterialCacheABufferPages.usf b/Engine/Shaders/Private/MaterialCache/MaterialCacheABufferPages.usf
--- a/Engine/Shaders/Private/MaterialCache/MaterialCacheABufferPages.usf
+++ b/Engine/Shaders/Private/MaterialCache/MaterialCacheABufferPages.usf
@@ -4,38 +4,62 @@
 #include "/Engine/Shared/MaterialCacheDefinitions.h"
 #include "/Engine/Private/GammaCorrectionCommon.ush"
 #include "/Engine/Private/MaterialCache/MaterialCacheCommon.ush"
+#include "/Engine/Private/BlockCompressionCommon.ush"
+#include "/Engine/Private/BCCompressionCommon.ush"
 #include "/Engine/Public/RootConstants.ush"
 
 StructuredBuffer<FMaterialCachePageWriteData> PageWriteData;
 
-RWTexture2D<float4> RWVTLayer0;
-RWTexture2D<float4> RWVTLayer1;
-RWTexture2D<float4> RWVTLayer2;
+RWTexture2D<uint4> RWVTLayer0;
+RWTexture2D<uint4> RWVTLayer1;
+RWTexture2D<uint4> RWVTLayer2;
 
 Texture2DArray<float4> ABuffer0;
 Texture2DArray<float4> ABuffer1;
 Texture2DArray<float4> ABuffer2;
 
-uint4 TileParams;
-uint4 TileOrderingParams;
+SamplerState TextureSampler;
 
-[numthreads(64, 1, 1)]
-void WritePagesMain(uint3 DTid : SV_DispatchThreadID)
+uint   BlockCount;
+float2 TexelSize;
+
+uint4 CompressBC3(Texture2DArray<float4> Source, float3 UV, bool bSRGB)
 {
-	const FMaterialCachePageWriteData Data = PageWriteData[DTid.y];
+	float3 BlockRGB[16];
+	ReadBlockRGB(Source, TextureSampler, UV, TexelSize.xy, BlockRGB);
+	float BlockA[16];
+	ReadBlockAlpha(Source, TextureSampler, UV, TexelSize.xy, BlockA);
+
+	if (bSRGB)
+	{
+		return CompressBC3Block_SRGB(BlockRGB, BlockA);
+	}
+	else
+	{
+		return CompressBC3Block(BlockRGB, BlockA);
+	}
+}
 
-	const uint2 PagePixelPos = GetMaterialCacheCSShadePixel(DTid, TileParams, TileOrderingParams);
+[numthreads(8, 8, 1)]
+void WritePagesMain(uint3 DTid : SV_DispatchThreadID)
+{
+	uint2 BlockIndex = DTid.xy;
 	
-	const uint3 ABufferPhysicalLocation = Data.ABufferPhysicalPosition + uint3(PagePixelPos, 0);
-	const uint2 VTPhysicalLocation      = Data.VTPhysicalPosition + PagePixelPos;
+	if (any(BlockIndex.xy >= BlockCount.xx))
+	{
+		return;
+	}
 
-	float4 ABufferValue0 = ABuffer0[ABufferPhysicalLocation];
-	float4 ABufferValue1 = ABuffer1[ABufferPhysicalLocation];
-	float4 ABufferValue2 = ABuffer2[ABufferPhysicalLocation];
+	const FMaterialCachePageWriteData Data = PageWriteData[DTid.z];
 
-	// -- Platform Block Compression --
+	static const uint BlockSH = 2;
 	
-	RWVTLayer0[VTPhysicalLocation] = float4(LinearToSrgb(ABufferValue0.xyz), ABufferValue0.w);
-	RWVTLayer1[VTPhysicalLocation] = ABufferValue1;
-	RWVTLayer2[VTPhysicalLocation] = ABufferValue2;
+	const uint3  ABufferPhysicalLocation = Data.ABufferPhysicalPosition + uint3(BlockIndex.xy << BlockSH, 0);
+	const float3 ABufferPhysicalUV       = float3(ABufferPhysicalLocation.xy * TexelSize.xy + TexelSize.xy * 0.5f, ABufferPhysicalLocation.z);
+	const uint2  VTPhysicalLocation      = (Data.VTPhysicalPosition >> BlockSH) + BlockIndex.xy;
+
+	// Compress and write blocks
+	RWVTLayer0[VTPhysicalLocation] = CompressBC3(ABuffer0, ABufferPhysicalUV, true);
+	RWVTLayer1[VTPhysicalLocation] = CompressBC3(ABuffer1, ABufferPhysicalUV, false);
+	RWVTLayer2[VTPhysicalLocation] = CompressBC3(ABuffer2, ABufferPhysicalUV, false);
 }
diff --git a/Engine/Source/Runtime/Renderer/Internal/MaterialCache/MaterialCacheABuffer.h b/Engine/Source/Runtime/Renderer/Internal/MaterialCache/MaterialCacheABuffer.h
--- a/Engine/Source/Runtime/Renderer/Internal/MaterialCache/MaterialCacheABuffer.h
+++ b/Engine/Source/Runtime/Renderer/Internal/MaterialCache/MaterialCacheABuffer.h
@@ -25,10 +25,22 @@ enum class EMaterialCacheAttributeMask
 	// TODO: Conditional targets
 };
 
+struct EMaterialCacheFormat
+{
+	EPixelFormat Format;
+	uint32       bIsSRGB : 1;
+};
+
 ENUM_CLASS_FLAGS(EMaterialCacheAttributeMask);
 
-static void GetMaterialCacheABufferFormats(EMaterialCacheAttributeMask Mask, TArray<EPixelFormat, TInlineAllocator<MaterialCacheMaxABuffers>>& OutFormats) {
+static void GetMaterialCacheABufferRenderFormats(EMaterialCacheAttributeMask Mask, TArray<EPixelFormat, TInlineAllocator<MaterialCacheMaxABuffers>>& OutFormats) {
 	OutFormats.Add(PF_R8G8B8A8);
 	OutFormats.Add(PF_A2B10G10R10);
 	OutFormats.Add(PF_R8G8B8A8);
 }
+
+static void GetMaterialCacheABufferPhysicalFormats(EMaterialCacheAttributeMask Mask, TArray<EMaterialCacheFormat, TInlineAllocator<MaterialCacheMaxABuffers>>& OutFormats) {
+	OutFormats.Add({ PF_DXT5, true  });
+	OutFormats.Add({ PF_DXT5, false });
+	OutFormats.Add({ PF_DXT5, false });
+}
diff --git a/Engine/Source/Runtime/Renderer/Internal/MaterialCache/MaterialCacheShaders.h b/Engine/Source/Runtime/Renderer/Internal/MaterialCache/MaterialCacheShaders.h
--- a/Engine/Source/Runtime/Renderer/Internal/MaterialCache/MaterialCacheShaders.h
+++ b/Engine/Source/Runtime/Renderer/Internal/MaterialCache/MaterialCacheShaders.h
@@ -96,14 +96,15 @@ class FMaterialCacheABufferWritePagesCS : public FGlobalShader
 public:
 	BEGIN_SHADER_PARAMETER_STRUCT(FParameters, )
 		SHADER_PARAMETER_RDG_BUFFER_SRV(StructuredBuffer<FMaterialCacheBinData>, PageWriteData)
-		SHADER_PARAMETER_RDG_TEXTURE_UAV(RWTexture2D<float4>, RWVTLayer0)
-		SHADER_PARAMETER_RDG_TEXTURE_UAV(RWTexture2D<float4>, RWVTLayer1)
-		SHADER_PARAMETER_RDG_TEXTURE_UAV(RWTexture2D<float4>, RWVTLayer2)
+		SHADER_PARAMETER_RDG_TEXTURE_UAV(RWTexture2D<uint4>, RWVTLayer0)
+		SHADER_PARAMETER_RDG_TEXTURE_UAV(RWTexture2D<uint4>, RWVTLayer1)
+		SHADER_PARAMETER_RDG_TEXTURE_UAV(RWTexture2D<uint4>, RWVTLayer2)
 		SHADER_PARAMETER_RDG_TEXTURE_SRV(Texture2DArray<float4>, ABuffer0)
 		SHADER_PARAMETER_RDG_TEXTURE_SRV(Texture2DArray<float4>, ABuffer1)
 		SHADER_PARAMETER_RDG_TEXTURE_SRV(Texture2DArray<float4>, ABuffer2)
-		SHADER_PARAMETER(FUintVector4, TileParams)
-		SHADER_PARAMETER(FUintVector4, TileOrderingParams)
+		SHADER_PARAMETER_SAMPLER(SamplerState, TextureSampler)
+		SHADER_PARAMETER(uint32, BlockCount)
+		SHADER_PARAMETER(FVector2f, TexelSize)
 	END_SHADER_PARAMETER_STRUCT()
 	
 	using FPermutationDomain = TShaderPermutationDomain<>;
diff --git a/Engine/Source/Runtime/Renderer/Private/MaterialCache/MaterialCacheRenderer.cpp b/Engine/Source/Runtime/Renderer/Private/MaterialCache/MaterialCacheRenderer.cpp
--- a/Engine/Source/Runtime/Renderer/Private/MaterialCache/MaterialCacheRenderer.cpp
+++ b/Engine/Source/Runtime/Renderer/Private/MaterialCache/MaterialCacheRenderer.cpp
@@ -738,7 +738,7 @@ static void MaterialCacheCreateABuffer(FRDGBuilder& GraphBuilder, FMaterialCache
 	FIntPoint RenderTileSize = GetMaterialCacheRenderTileSize();
 
 	TArray<EPixelFormat, TInlineAllocator<MaterialCacheMaxABuffers>> Formats;
-	GetMaterialCacheABufferFormats({}, Formats);
+	GetMaterialCacheABufferRenderFormats({}, Formats);
 
 	ETextureCreateFlags Flags =
 		ETextureCreateFlags::ShaderResource | 
@@ -1556,6 +1556,10 @@ static void MaterialCacheFinalizePages(FSceneRendererBase* Renderer, FRDGBuilder
 		PageWriteDataArray.NumBytes()
 	);
 
+	const uint32 BlockSH = 2;
+	const uint32 TexelThreadCountX = GetMaterialCacheTileWidth() + GetMaterialCacheTileBorderWidth();
+	const uint32 BlockThreadCountX = TexelThreadCountX >> BlockSH;
+	
 	FMaterialCacheABufferWritePagesCS::FParameters* PassParameters = GraphBuilder.AllocParameters<FMaterialCacheABufferWritePagesCS::FParameters>();
 	PassParameters->PageWriteData = GraphBuilder.CreateSRV(PageWriteData);
 	PassParameters->ABuffer0 = GraphBuilder.CreateSRV(RenderData.ABuffer.ABufferTextures[0]);
@@ -1564,8 +1568,9 @@ static void MaterialCacheFinalizePages(FSceneRendererBase* Renderer, FRDGBuilder
 	PassParameters->RWVTLayer0 = GraphBuilder.CreateUAV(GraphBuilder.RegisterExternalTexture(SceneExtension.PendingEntries[0].Setup.PhysicalRenderTargets[0], ERDGTextureFlags::ForceImmediateFirstBarrier));
 	PassParameters->RWVTLayer1 = GraphBuilder.CreateUAV(GraphBuilder.RegisterExternalTexture(SceneExtension.PendingEntries[0].Setup.PhysicalRenderTargets[1], ERDGTextureFlags::ForceImmediateFirstBarrier));
 	PassParameters->RWVTLayer2 = GraphBuilder.CreateUAV(GraphBuilder.RegisterExternalTexture(SceneExtension.PendingEntries[0].Setup.PhysicalRenderTargets[2], ERDGTextureFlags::ForceImmediateFirstBarrier));
-	PassParameters->TileParams = GetMaterialCacheTileParams();
-	PassParameters->TileOrderingParams = GetMaterialCacheTileOrderingParams();
+	PassParameters->TextureSampler = TStaticSamplerState<SF_Bilinear>::GetRHI();
+	PassParameters->BlockCount = BlockThreadCountX;
+	PassParameters->TexelSize = FVector2f(1.0f / RenderData.ABuffer.ABufferTextures[0]->Desc.Extent.X, 1.0f / RenderData.ABuffer.ABufferTextures[0]->Desc.Extent.Y);
 
 	GraphBuilder.AddPass(
 		RDG_EVENT_NAME("WritePages"),
@@ -1574,13 +1579,13 @@ static void MaterialCacheFinalizePages(FSceneRendererBase* Renderer, FRDGBuilder
 		[Renderer, &RenderData, PassParameters](FRDGAsyncTask, FRHIComputeCommandList& RHICmdList)
 	{
 		FComputeShaderUtils::Dispatch(
-			RHICmdList, 
+			RHICmdList,
 			Renderer->Views[0].ShaderMap->GetShader<FMaterialCacheABufferWritePagesCS>(),
 			*PassParameters,
 			FIntVector(
-				FMath::DivideAndRoundUp(GetMaterialCacheTileThreadCount(), 64u),
-				RenderData.ABuffer.Pages.Num(),
-				1
+				FMath::DivideAndRoundUp(PassParameters->BlockCount, 8u),
+				FMath::DivideAndRoundUp(PassParameters->BlockCount, 8u),
+				RenderData.ABuffer.Pages.Num()
 			)
 		);
 	});
diff --git a/Engine/Source/Runtime/Renderer/Private/MaterialCache/MaterialCacheSceneExtension.cpp b/Engine/Source/Runtime/Renderer/Private/MaterialCache/MaterialCacheSceneExtension.cpp
--- a/Engine/Source/Runtime/Renderer/Private/MaterialCache/MaterialCacheSceneExtension.cpp
+++ b/Engine/Source/Runtime/Renderer/Private/MaterialCache/MaterialCacheSceneExtension.cpp
@@ -14,10 +14,10 @@ IMPLEMENT_SCENE_EXTENSION(FMaterialCacheSceneExtension);
 
 BEGIN_SHADER_PARAMETER_STRUCT(FMaterialCacheTextureParameters, RENDERER_API)
 	SHADER_PARAMETER_TEXTURE(Texture2D<uint4>, PageTableTexture)
-	SHADER_PARAMETER_TEXTURE(Texture2D<float4>, PhysicalTexture0)
+	SHADER_PARAMETER_SRV(Texture2D<float4>, PhysicalTexture0)
+	SHADER_PARAMETER_SRV(Texture2D<float4>, PhysicalTexture1)
+	SHADER_PARAMETER_SRV(Texture2D<float4>, PhysicalTexture2)
 	SHADER_PARAMETER(FUintVector4, PackedUniform)
-	SHADER_PARAMETER_TEXTURE(Texture2D<float4>, PhysicalTexture1)
-	SHADER_PARAMETER_TEXTURE(Texture2D<float4>, PhysicalTexture2)
 END_SHADER_PARAMETER_STRUCT()
 
 DECLARE_SCENE_UB_STRUCT(FMaterialCacheTextureParameters, MaterialCache, RENDERER_API)
@@ -89,9 +89,9 @@ public:
 			const float  PackedSignBit          = bPageTableExtraBits ? 1.0f : -1.0f;
 			
 			Parameters.PageTableTexture = Texture->GetPageTableTexture(0);
-			Parameters.PhysicalTexture0 = Texture->GetPhysicalTexture(0);
-			Parameters.PhysicalTexture1 = Texture->GetPhysicalTexture(1);
-			Parameters.PhysicalTexture2 = Texture->GetPhysicalTexture(2);
+			Parameters.PhysicalTexture0 = Texture->GetPhysicalTextureSRV(0, true);
+			Parameters.PhysicalTexture1 = Texture->GetPhysicalTextureSRV(1, false);
+			Parameters.PhysicalTexture2 = Texture->GetPhysicalTextureSRV(2, false);
 			Parameters.PackedUniform.X  = 0xFFFFFFFF;
 			Parameters.PackedUniform.Y  = FMath::AsUInt(static_cast<float>(PageSize) * RcpPhysicalTextureSize);
 			Parameters.PackedUniform.Z  = FMath::AsUInt(static_cast<float>(PageBorderSize) * RcpPhysicalTextureSize);
@@ -236,8 +236,8 @@ void FMaterialCacheSceneExtension::Unregister(FPrimitiveComponentId PrimitiveCom
 static void GetDefaultMaterialCacheParameters(FMaterialCacheTextureParameters& OutParameters, FRDGBuilder&)
 {
 	OutParameters.PageTableTexture = GBlackUintTexture->TextureRHI;
-	OutParameters.PhysicalTexture0 = GBlackTextureWithSRV->TextureRHI;
-	OutParameters.PhysicalTexture1 = GBlackTextureWithSRV->TextureRHI;
-	OutParameters.PhysicalTexture2 = GBlackTextureWithSRV->TextureRHI;
+	OutParameters.PhysicalTexture0 = GBlackTextureWithSRV->ShaderResourceViewRHI;
+	OutParameters.PhysicalTexture1 = GBlackTextureWithSRV->ShaderResourceViewRHI;
+	OutParameters.PhysicalTexture2 = GBlackTextureWithSRV->ShaderResourceViewRHI;
 	OutParameters.PackedUniform       = FUintVector4(0, 0, 0, 0);
 }
diff --git a/Engine/Source/Runtime/Renderer/Private/MaterialCache/MaterialCacheShaders.cpp b/Engine/Source/Runtime/Renderer/Private/MaterialCache/MaterialCacheShaders.cpp
--- a/Engine/Source/Runtime/Renderer/Private/MaterialCache/MaterialCacheShaders.cpp
+++ b/Engine/Source/Runtime/Renderer/Private/MaterialCache/MaterialCacheShaders.cpp
@@ -134,6 +134,8 @@ bool FMaterialCacheABufferWritePagesCS::ShouldCompilePermutation(const FGlobalSh
 
 void FMaterialCacheABufferWritePagesCS::ModifyCompilationEnvironment(const FShaderPermutationParameters& Parameters, FShaderCompilerEnvironment& OutEnvironment)
 {
+	OutEnvironment.SetDefine(TEXT("BLOCK_COMPRESS_SRC_TEXTURE_ARRAY"), 1);
+	
 	OutEnvironment.CompilerFlags.Add(CFLAG_RootConstants);
 }
 
diff --git a/Engine/Source/Runtime/Renderer/Private/MaterialCache/MaterialCacheVirtualTexture.cpp b/Engine/Source/Runtime/Renderer/Private/MaterialCache/MaterialCacheVirtualTexture.cpp
--- a/Engine/Source/Runtime/Renderer/Private/MaterialCache/MaterialCacheVirtualTexture.cpp
+++ b/Engine/Source/Runtime/Renderer/Private/MaterialCache/MaterialCacheVirtualTexture.cpp
@@ -139,9 +139,11 @@ public:
 		, TileSize(InTileSize)
 		, TileBorderSize(InTileBorderSize)
 	{
+		TextureName = TEXT("MaterialCacheVirtualTexture");
+		
 		MaxLevel = FMath::CeilLogTwo(FMath::Max(InTileCount.X, InTileCount.Y));
 		
-		GetMaterialCacheABufferFormats({}, ABufferFormats);
+		GetMaterialCacheABufferPhysicalFormats({}, ABufferFormats);
 
 		// Share the page table across all physical textures
 		bSinglePhysicalSpace = true;
@@ -154,7 +156,7 @@ public:
 	
 	virtual EPixelFormat GetFormat(uint32 LayerIndex) const override
 	{
-		return ABufferFormats[LayerIndex];
+		return ABufferFormats[LayerIndex].Format;
 	}
 	
 	virtual uint32 GetTileSize() const override
@@ -211,9 +213,9 @@ public:
 
 		for (int32 LayerIndex = 0; LayerIndex < ABufferFormats.Num(); LayerIndex++)
 		{
-			ProducerDesc.LayerFormat[LayerIndex] = ABufferFormats[LayerIndex];
+			ProducerDesc.LayerFormat[LayerIndex] = ABufferFormats[LayerIndex].Format;
+			ProducerDesc.bIsLayerSRGB[LayerIndex] = ABufferFormats[LayerIndex].bIsSRGB;
 			ProducerDesc.PhysicalGroupIndex[LayerIndex] = 0;
-			ProducerDesc.bIsLayerSRGB[LayerIndex] = false;
 		}
 		
 		// Register producer on page feedback
@@ -229,7 +231,7 @@ private:
 	FPrimitiveComponentId PrimitiveComponentId;
 	
 	/** Physical formats */
-	TArray<EPixelFormat, TInlineAllocator<MaterialCacheMaxABuffers>> ABufferFormats;
+	TArray<EMaterialCacheFormat, TInlineAllocator<MaterialCacheMaxABuffers>> ABufferFormats;
 
 	/** Tiled properties */
 	FIntPoint TileCount;
