Stochastic Lighting - fixed tile categorization pass when using non-DBuffer decals. Decal pass has to run before the tile categorization as it changes GBuffer buffer.

#rnx

diff --git a/Engine/Source/Runtime/Renderer/Private/DeferredShadingRenderer.cpp b/Engine/Source/Runtime/Renderer/Private/DeferredShadingRenderer.cpp
--- a/Engine/Source/Runtime/Renderer/Private/DeferredShadingRenderer.cpp
+++ b/Engine/Source/Runtime/Renderer/Private/DeferredShadingRenderer.cpp
@@ -2891,12 +2891,19 @@ void FDeferredShadingSceneRenderer::Render(FRDGBuilder& GraphBuilder, const FSce
 			Substrate::AddSubstrateSampleMaterialPass(GraphBuilder, Scene, SceneTextures, Views);
 		}
 
+		FAsyncLumenIndirectLightingOutputs AsyncLumenIndirectLightingOutputs;
+
+		if (RequiresStochasticLightingPass())
 		{
+			// Decals may modify GBuffers so they need to be done first. Can decals read velocities and/or custom depth? If so, they need to be rendered earlier too.
+			CompositionLighting.ProcessAfterBasePass(GraphBuilder, InstanceCullingManager, FCompositionLighting::EProcessAfterBasePassMode::OnlyBeforeLightingDecals, Scene->SubstrateSceneData);
+			AsyncLumenIndirectLightingOutputs.bHasDrawnBeforeLightingDecals = true;
+
 			RDG_EVENT_SCOPE_STAT(GraphBuilder, RenderDeferredLighting, "StochasticLighting");
 			RDG_GPU_STAT_SCOPE(GraphBuilder, RenderDeferredLighting);
 
 			StochasticLightingTileClassificationMark(GraphBuilder, LumenFrameTemporaries, SceneTextures);
-		}		
+		}
 
 		// Copy lighting channels out of stencil before deferred decals which overwrite those values
 		TArray<FRDGTextureRef, TInlineAllocator<2>> NaniteShadingMask;
@@ -2916,8 +2923,6 @@ void FDeferredShadingSceneRenderer::Render(FRDGBuilder& GraphBuilder, const FSce
 			SingleLayerWaterPrePassResult = RenderSingleLayerWaterDepthPrepass(GraphBuilder, Views, SceneTextures, SingleLayerWaterPrepassLocation, NaniteRasterResults);
 		}
 
-		FAsyncLumenIndirectLightingOutputs AsyncLumenIndirectLightingOutputs;
-
 		GraphBuilder.FlushSetupQueue();
 
 		// Shadows, lumen and fog after base pass
@@ -2934,7 +2939,6 @@ void FDeferredShadingSceneRenderer::Render(FRDGBuilder& GraphBuilder, const FSce
 
 			DispatchAsyncLumenIndirectLightingWork(
 				GraphBuilder,
-				CompositionLighting,
 				SceneTextures,
 				InstanceCullingManager,
 				LumenFrameTemporaries,
diff --git a/Engine/Source/Runtime/Renderer/Private/DeferredShadingRenderer.h b/Engine/Source/Runtime/Renderer/Private/DeferredShadingRenderer.h
--- a/Engine/Source/Runtime/Renderer/Private/DeferredShadingRenderer.h
+++ b/Engine/Source/Runtime/Renderer/Private/DeferredShadingRenderer.h
@@ -672,7 +672,6 @@ private:
 	/** Dispatch async Lumen work if possible. */
 	void DispatchAsyncLumenIndirectLightingWork(
 		FRDGBuilder& GraphBuilder,
-		class FCompositionLighting& CompositionLighting,
 		FSceneTextures& SceneTextures,
 		FInstanceCullingManager& InstanceCullingManager,
 		FLumenSceneFrameTemporaries& LumenFrameTemporaries,
@@ -753,6 +752,7 @@ private:
 		ERDGPassFlags ComputePassFlags,
 		FLumenScreenSpaceBentNormalParameters& ScreenSpaceBentNormalParameters);
 
+	bool RequiresStochasticLightingPass();
 	void StochasticLightingTileClassificationMark(FRDGBuilder& GraphBuilder, FLumenSceneFrameTemporaries& FrameTemporaries, const FSceneTextures& SceneTextures);
 
 	/** Extract current frame opaque (no water) depth and normal scene textures to use as history data. */
diff --git a/Engine/Source/Runtime/Renderer/Private/IndirectLightRendering.cpp b/Engine/Source/Runtime/Renderer/Private/IndirectLightRendering.cpp
--- a/Engine/Source/Runtime/Renderer/Private/IndirectLightRendering.cpp
+++ b/Engine/Source/Runtime/Renderer/Private/IndirectLightRendering.cpp
@@ -882,7 +882,6 @@ void FDeferredShadingSceneRenderer::SetupCommonDiffuseIndirectParameters(
 
 void FDeferredShadingSceneRenderer::DispatchAsyncLumenIndirectLightingWork(
 	FRDGBuilder& GraphBuilder,
-	FCompositionLighting& CompositionLighting,
 	FSceneTextures& SceneTextures,
 	FInstanceCullingManager& InstanceCullingManager,
 	FLumenSceneFrameTemporaries& LumenFrameTemporaries,
@@ -902,10 +901,6 @@ void FDeferredShadingSceneRenderer::DispatchAsyncLumenIndirectLightingWork(
 		return;
 	}
 
-	// Decals may modify GBuffers so they need to be done first. Can decals read velocities and/or custom depth? If so, they need to be rendered earlier too.
-	CompositionLighting.ProcessAfterBasePass(GraphBuilder, InstanceCullingManager, FCompositionLighting::EProcessAfterBasePassMode::OnlyBeforeLightingDecals, Scene->SubstrateSceneData);
-	Outputs.bHasDrawnBeforeLightingDecals = true;
-
 	if (DynamicShadowsTaskData)
 	{
 		for (const FViewInfo& View : Views)
diff --git a/Engine/Source/Runtime/Renderer/Private/StochasticLighting/StochasticLighting.cpp b/Engine/Source/Runtime/Renderer/Private/StochasticLighting/StochasticLighting.cpp
--- a/Engine/Source/Runtime/Renderer/Private/StochasticLighting/StochasticLighting.cpp
+++ b/Engine/Source/Runtime/Renderer/Private/StochasticLighting/StochasticLighting.cpp
@@ -162,6 +162,24 @@ void StochasticLighting::TileClassificationMark(
 	}
 }
 
+bool FDeferredShadingSceneRenderer::RequiresStochasticLightingPass()
+{
+	for (int32 ViewIndex = 0; ViewIndex < Views.Num(); ++ViewIndex)
+	{
+		FViewInfo& View = Views[ViewIndex];
+		const FPerViewPipelineState& ViewPipelineState = GetViewPipelineState(View);
+
+		if (ViewPipelineState.DiffuseIndirectMethod == EDiffuseIndirectMethod::Lumen
+			|| ViewPipelineState.ReflectionsMethod == EReflectionsMethod::Lumen
+			|| MegaLights::IsEnabled(ViewFamily))
+		{
+			return true;
+		}
+	}
+
+	return false;
+}
+
 /**
  * Load GBuffer data once and transform it for subsequent lighting passes
  * This includes full res depth and normal copy for opaque before it gets overwritten by water or other translucency writing depth
