[SingleLayerWater] PR #13345: Fix not using correct DownsampledClearTileListDataBuffer for single layer water.

diff --git a/Engine/Source/Runtime/Renderer/Private/SingleLayerWaterRendering.cpp b/Engine/Source/Runtime/Renderer/Private/SingleLayerWaterRendering.cpp
--- a/Engine/Source/Runtime/Renderer/Private/SingleLayerWaterRendering.cpp
+++ b/Engine/Source/Runtime/Renderer/Private/SingleLayerWaterRendering.cpp
@@ -795,7 +795,7 @@ static FSingleLayerWaterTileClassification ClassifyTiles(FRDGBuilder& GraphBuild
 		if (bNeedDownsample && bNeedClearTiles)
 		{
 			DownsampledClearTileListDataBuffer = GraphBuilder.CreateBuffer(FRDGBufferDesc::CreateBufferDesc(sizeof(uint32), DownsampledViewResInTiles.X * DownsampledViewResInTiles.Y), TEXT("SLW.DownsampledClearTileListDataBuffer"));
-			Result.TiledReflection.DownsampledClearTileListDataBufferSRV = GraphBuilder.CreateSRV(ClearTileListDataBuffer, PF_R32_UINT);
+			Result.TiledReflection.DownsampledClearTileListDataBufferSRV = GraphBuilder.CreateSRV(DownsampledClearTileListDataBuffer, PF_R32_UINT);
 		}
 
 		FRDGBufferUAVRef DrawIndirectParametersBufferUAV = GraphBuilder.CreateUAV(Result.TiledReflection.DrawIndirectParametersBuffer);
