Fix for mesh pass not being waited on at the end of the scene render.
- Custom render passes being ignored when views are added to the Prereq array.
- Moved the task out of the mesh pass so we always wait once on the render thread; removes the possibility of a deadlock later on.

diff --git a/Engine/Source/Runtime/Renderer/Private/InstanceCulling/InstanceCullingOcclusionQuery.cpp b/Engine/Source/Runtime/Renderer/Private/InstanceCulling/InstanceCullingOcclusionQuery.cpp
--- a/Engine/Source/Runtime/Renderer/Private/InstanceCulling/InstanceCullingOcclusionQuery.cpp
+++ b/Engine/Source/Runtime/Renderer/Private/InstanceCulling/InstanceCullingOcclusionQuery.cpp
@@ -348,7 +348,7 @@ struct FInstanceCullingOcclusionQueryDeferredContext
 			return;
 		}
 
-		const FParallelMeshDrawCommandPass& MeshDrawCommandPass = *View->ParallelMeshDrawCommandPasses[MeshPass];
+		FParallelMeshDrawCommandPass& MeshDrawCommandPass = *View->ParallelMeshDrawCommandPasses[MeshPass];
 
 		// Execute() is expected to run late enough to not stall here.
 		// If it does happen, then we may have to move the render pass to later point in the frame.
diff --git a/Engine/Source/Runtime/Renderer/Private/MeshDrawCommands.cpp b/Engine/Source/Runtime/Renderer/Private/MeshDrawCommands.cpp
--- a/Engine/Source/Runtime/Renderer/Private/MeshDrawCommands.cpp
+++ b/Engine/Source/Runtime/Renderer/Private/MeshDrawCommands.cpp
@@ -1484,7 +1484,7 @@ void FParallelMeshDrawCommandPass::DispatchPassSetup(
 		// This work needs to be deferred until at least BuildRenderingCommands (to ensure the DynamicPrimitiveCollector is uploaded), so we use the async mechanism either way 
 		auto FinalizeInstanceCullingSetup = [this, Scene](FInstanceCullingContext& InstanceCullingContext)
 		{
-			WaitForMeshPassSetupTask();
+			WaitForSetupTask();
 
 #if DO_CHECK
 			for (const FVisibleMeshDrawCommand& VisibleMeshDrawCommand : TaskContext.MeshDrawCommands)
@@ -1511,22 +1511,20 @@ bool FParallelMeshDrawCommandPass::IsOnDemandShaderCreationEnabled()
 		(GRHISupportsMultithreadedShaderCreation || (bIsMobileRenderer && (!GSupportsParallelRenderingTasksWithSeparateRHIThread && IsRunningRHIInSeparateThread())));
 }
 
-void FParallelMeshDrawCommandPass::WaitForMeshPassSetupTask() const
+void FParallelMeshDrawCommandPass::WaitForSetupTask()
 {
 	if (TaskEventRef.IsValid())
 	{
 		// Need to wait on GetRenderThread_Local, as mesh pass setup task can wait on rendering thread inside InitResourceFromPossiblyParallelRendering().
 		QUICK_SCOPE_CYCLE_COUNTER(STAT_WaitForMeshPassSetupTask);
 		TaskEventRef->Wait(IsInActualRenderingThread() ? ENamedThreads::GetRenderThread_Local() : ENamedThreads::AnyThread);
+		TaskEventRef = nullptr;
 	}
 }
 
-void FParallelMeshDrawCommandPass::WaitForTasksAndEmpty()
+void FParallelMeshDrawCommandPass::Cleanup()
 {
-	// Need to wait in case if someone dispatched sort and draw merge task, but didn't draw it.
-	WaitForMeshPassSetupTask();
-	TaskEventRef = nullptr;
-
+	check(TaskEventRef == nullptr);
 	DumpInstancingStats();
 
 	if (TaskContext.MeshPassProcessor)
@@ -1666,13 +1664,6 @@ void FParallelMeshDrawCommandPass::BuildRenderingCommands(
 	OutInstanceCullingDrawParams.IndirectArgsByteOffset = 0U;
 }
 
-
-
-void FParallelMeshDrawCommandPass::WaitForSetupTask() const
-{
-	WaitForMeshPassSetupTask();
-}
-
 PRAGMA_DISABLE_DEPRECATION_WARNINGS
 
 void FParallelMeshDrawCommandPass::DispatchDraw(FParallelCommandListSet* ParallelCommandListSet, FRHICommandList& RHICmdList, const FInstanceCullingDrawParams* InstanceCullingDrawParams) const
@@ -1750,7 +1741,7 @@ void FParallelMeshDrawCommandPass::Draw(FRHICommandList& RHICmdList, const FInst
 
 	QUICK_SCOPE_CYCLE_COUNTER(STAT_MeshPassDrawImmediate);
 
-	WaitForMeshPassSetupTask();
+	check(TaskEventRef == nullptr);
 
 	if (TaskContext.bUseGPUScene)
 	{
diff --git a/Engine/Source/Runtime/Renderer/Private/MeshDrawCommands.h b/Engine/Source/Runtime/Renderer/Private/MeshDrawCommands.h
--- a/Engine/Source/Runtime/Renderer/Private/MeshDrawCommands.h
+++ b/Engine/Source/Runtime/Renderer/Private/MeshDrawCommands.h
@@ -157,7 +157,7 @@ public:
 	/**
 	 * Sync with setup task.
 	 */
-	void WaitForSetupTask() const;
+	void WaitForSetupTask();
 
 	/**
 	 * Dispatch visible mesh draw command draw task.
@@ -171,7 +171,7 @@ PRAGMA_DISABLE_DEPRECATION_WARNINGS
 	void DispatchDraw(FParallelCommandListSet* ParallelCommandListSet, FRHICommandList& RHICmdList, const FInstanceCullingDrawParams* InstanceCullingDrawParams = nullptr) const;
 PRAGMA_ENABLE_DEPRECATION_WARNINGS
 
-	void WaitForTasksAndEmpty();
+	void Cleanup();
 	void SetDumpInstancingStats(const FString& InPassName);
 	bool HasAnyDraw() const { return MaxNumDraws > 0; }
 
@@ -180,6 +180,8 @@ PRAGMA_ENABLE_DEPRECATION_WARNINGS
 	FInstanceCullingContext* GetInstanceCullingContext() { return &TaskContext.InstanceCullingContext; }
 	const FGraphEventRef& GetTaskEvent() const { return TaskEventRef; }
 
+	FGraphEventRef&& AcquireTaskEvent() { return MoveTemp(TaskEventRef); }
+
 	// NOTE: It is only safe to access mesh draw commands after the setup task is complete (use WaitForSetupTask). 
 	// Only access the data late in the frame to allow as much time as possible for async tasks to complete.
 	const FMeshCommandOneFrameArray& GetMeshDrawCommands() const
@@ -199,7 +201,6 @@ private:
 	int32 MaxNumDraws;
 
 	void DumpInstancingStats() const;
-	void WaitForMeshPassSetupTask() const;
 };
 
 inline bool HasAnyDraw(const FParallelMeshDrawCommandPass* Pass)
diff --git a/Engine/Source/Runtime/Renderer/Private/MobileShadingRenderer.cpp b/Engine/Source/Runtime/Renderer/Private/MobileShadingRenderer.cpp
--- a/Engine/Source/Runtime/Renderer/Private/MobileShadingRenderer.cpp
+++ b/Engine/Source/Runtime/Renderer/Private/MobileShadingRenderer.cpp
@@ -406,6 +406,11 @@ void FMobileSceneRenderer::SetupMobileBasePassAfterShadowInit(FExclusiveDepthSte
 		FViewInfo& View = *AllViews[ViewIndex];
 		FViewCommands& ViewCommands = ViewCommandsPerView[ViewIndex];
 
+		if (ViewCommands.MeshCommands[EMeshPass::BasePass].IsEmpty() && ViewCommands.MeshCommands[EMeshPass::MobileBasePassCSM].IsEmpty() && View.NumVisibleDynamicMeshElements[EMeshPass::BasePass] == 0 && ViewCommands.NumDynamicMeshCommandBuildRequestElements[EMeshPass::BasePass] == 0)
+		{
+			continue;
+		}
+
 		FMeshPassProcessor* MeshPassProcessor = FPassProcessorManager::CreateMeshPassProcessor(EShadingPath::Mobile, EMeshPass::BasePass, Scene->GetFeatureLevel(), Scene, &View, nullptr);
 
 		FMeshPassProcessor* BasePassCSMMeshPassProcessor = FPassProcessorManager::CreateMeshPassProcessor(EShadingPath::Mobile, EMeshPass::MobileBasePassCSM, Scene->GetFeatureLevel(), Scene, &View, nullptr);
diff --git a/Engine/Source/Runtime/Renderer/Private/Renderer/ViewSnapshotCache.cpp b/Engine/Source/Runtime/Renderer/Private/Renderer/ViewSnapshotCache.cpp
--- a/Engine/Source/Runtime/Renderer/Private/Renderer/ViewSnapshotCache.cpp
+++ b/Engine/Source/Runtime/Renderer/Private/Renderer/ViewSnapshotCache.cpp
@@ -94,7 +94,7 @@ public:
 			{
 				if (Pass)
 				{
-					Pass->WaitForTasksAndEmpty();
+					Pass->Cleanup();
 				}
 			}
 
diff --git a/Engine/Source/Runtime/Renderer/Private/SceneRenderBuilder.cpp b/Engine/Source/Runtime/Renderer/Private/SceneRenderBuilder.cpp
--- a/Engine/Source/Runtime/Renderer/Private/SceneRenderBuilder.cpp
+++ b/Engine/Source/Runtime/Renderer/Private/SceneRenderBuilder.cpp
@@ -295,7 +295,7 @@ static void DestroySplitScreenDebugViewFamilies(TConstArrayView<const FSceneView
 
 //////////////////////////////////////////////////////////////////////////
 
-static int32 GSceneRenderCleanUpMode = 0;
+static int32 GSceneRenderCleanUpMode = 1;
 static FAutoConsoleVariableRef CVarSceneRenderCleanUpMode(
 	TEXT("r.SceneRender.CleanUpMode"),
 	GSceneRenderCleanUpMode,
@@ -705,12 +705,12 @@ static void CleanupSceneRenderer(FSceneRenderer* Renderer)
 
 	for (auto* Pass : Renderer->DispatchedShadowDepthPasses)
 	{
-		Pass->WaitForTasksAndEmpty();
+		Pass->Cleanup();
 	}
 
 	for (FViewInfo* View : Renderer->AllViews)
 	{
-		View->WaitForTasks();
+		View->Cleanup();
 	}
 
 	ViewSnapshotCache::Destroy();
@@ -725,58 +725,67 @@ void FSceneRenderProcessor::Cleanup(FRHICommandListImmediate& RHICmdList, FScene
 
 	const ESceneRenderCleanUpMode SceneRenderCleanUpMode = GetSceneRenderCleanUpMode();
 
-	if (SceneRenderCleanUpMode == ESceneRenderCleanUpMode::Immediate)
+	if (SceneRenderCleanUpMode == ESceneRenderCleanUpMode::Async)
+	{
+		Renderer->GPUSceneDynamicContext.Release();
+	}
+	else // Immediate
 	{
 		WaitForAsyncDeleteTask(); // This is to handle cases where as switch from async to immediate.
 		RHICmdList.ImmediateFlush(EImmediateFlushType::WaitForOutstandingTasksOnly);
-		ViewSnapshotCache::Deallocate();
-		CleanupSceneRenderer(Renderer);
 	}
-	else
-	{
-		FGraphEventArray Prereqs;
-		Renderer->GPUSceneDynamicContext.Release();
 
-		// Wait on all setup tasks now to ensure that no additional render commands are enqueued which
-		// might mess with render state, since setup tasks are working with high-level render objects.
-		for (FParallelMeshDrawCommandPass* DispatchedShadowDepthPass : Renderer->DispatchedShadowDepthPasses)
+	FGraphEventArray TasksToWait;
+
+	// Wait on all setup tasks now to ensure that no additional render commands are enqueued which
+	// might mess with render state, since setup tasks are working with high-level render objects.
+	for (FParallelMeshDrawCommandPass* DispatchedShadowDepthPass : Renderer->DispatchedShadowDepthPasses)
+	{
+		if (FGraphEventRef Event = DispatchedShadowDepthPass->AcquireTaskEvent())
 		{
-			if (DispatchedShadowDepthPass->GetTaskEvent())
-			{
-				Prereqs.Add(DispatchedShadowDepthPass->GetTaskEvent());
-			}
+			TasksToWait.Emplace(MoveTemp(Event));
 		}
+	}
 
-		for (const FViewInfo& View : Renderer->Views)
+	for (const FViewInfo* View : Renderer->AllViews)
+	{
+		for (FParallelMeshDrawCommandPass* Pass : View->ParallelMeshDrawCommandPasses)
 		{
-			for (const FParallelMeshDrawCommandPass* Pass : View.ParallelMeshDrawCommandPasses)
+			if (Pass)
 			{
-				if (Pass && Pass->GetTaskEvent())
+				if (FGraphEventRef Event = Pass->AcquireTaskEvent())
 				{
-					Prereqs.Add(Pass->GetTaskEvent());
+					TasksToWait.Emplace(MoveTemp(Event));
 				}
 			}
 		}
+	}
 
-		// Wait for the last renderer's cleanup tasks so that snapshot deallocation and destruction don't overlap.
-		if (AsyncTasks.Cleanup)
-		{
-			Prereqs.Add(AsyncTasks.Cleanup);
-		}
+	// Wait for the last renderer's cleanup tasks so that snapshot deallocation and destruction don't overlap.
+	if (AsyncTasks.Cleanup)
+	{
+		TasksToWait.Emplace(MoveTemp(AsyncTasks.Cleanup));
+	}
 
-		if (!Prereqs.IsEmpty())
-		{
-			FTaskGraphInterface::Get().WaitUntilTasksComplete(Prereqs, ENamedThreads::GetRenderThread_Local());
-		}
+	if (!TasksToWait.IsEmpty())
+	{
+		FTaskGraphInterface::Get().WaitUntilTasksComplete(TasksToWait, ENamedThreads::GetRenderThread_Local());
+	}
 
-		ViewSnapshotCache::Deallocate();
+	ViewSnapshotCache::Deallocate();
 
+	if (SceneRenderCleanUpMode == ESceneRenderCleanUpMode::Async)
+	{
 		AsyncTasks.Cleanup = FFunctionGraphTask::CreateAndDispatchWhenReady([Renderer]
 		{
 			CleanupSceneRenderer(Renderer);
 
 		}, TStatId(), &GRHICommandList.WaitOutstandingTasks);
 	}
+	else // Immediate
+	{
+		CleanupSceneRenderer(Renderer);
+	}
 
 	GlobalDynamicBuffer::GarbageCollect();
 	GPrimitiveIdVertexBufferPool.DiscardAll();
diff --git a/Engine/Source/Runtime/Renderer/Private/SceneRendering.cpp b/Engine/Source/Runtime/Renderer/Private/SceneRendering.cpp
--- a/Engine/Source/Runtime/Renderer/Private/SceneRendering.cpp
+++ b/Engine/Source/Runtime/Renderer/Private/SceneRendering.cpp
@@ -1023,13 +1023,13 @@ FParallelMeshDrawCommandPass* FViewInfo::CreateMeshPass(EMeshPass::Type MeshPass
 	return ParallelMeshDrawCommandPasses[MeshPass] = Allocator.Create<FParallelMeshDrawCommandPass>();
 }
 
-void FViewInfo::WaitForTasks()
+void FViewInfo::Cleanup()
 {
 	for (int32 MeshDrawIndex = 0; MeshDrawIndex < EMeshPass::Num; MeshDrawIndex++)
 	{
 		if (auto* Pass = ParallelMeshDrawCommandPasses[MeshDrawIndex])
 		{
-			Pass->WaitForTasksAndEmpty();
+			Pass->Cleanup();
 		}
 	}
 }
diff --git a/Engine/Source/Runtime/Renderer/Private/SceneRendering.h b/Engine/Source/Runtime/Renderer/Private/SceneRendering.h
--- a/Engine/Source/Runtime/Renderer/Private/SceneRendering.h
+++ b/Engine/Source/Runtime/Renderer/Private/SceneRendering.h
@@ -1773,7 +1773,7 @@ public:
 	/** Create a snapshot of this view info on the scene allocator. */
 	FViewInfo* CreateSnapshot() const;
 
-	void WaitForTasks();
+	void Cleanup();
 
 	// Get the range in DynamicMeshElements[] for a given PrimitiveIndex
 	// @return range (start is inclusive, end is exclusive)
diff --git a/Engine/Source/Runtime/Renderer/Private/ShadowSetup.cpp b/Engine/Source/Runtime/Renderer/Private/ShadowSetup.cpp
--- a/Engine/Source/Runtime/Renderer/Private/ShadowSetup.cpp
+++ b/Engine/Source/Runtime/Renderer/Private/ShadowSetup.cpp
@@ -2908,7 +2908,7 @@ void FProjectedShadowInfo::ClearTransientArrays()
 	DynamicSubjectHeterogeneousVolumeMeshElements.Empty();
 
 	ShadowDepthPassVisibleCommands.Empty();
-	ShadowDepthPass.WaitForTasksAndEmpty();
+	ShadowDepthPass.Cleanup();
 
 	SubjectMeshCommandBuildRequests.Empty();
 	SubjectMeshCommandBuildFlags.Empty();
