View frozen matrices are patched before and after the visibility pass instead of directly inside GDME. This used to also be set for static relevance but that was refactored out long ago.

diff --git a/Engine/Source/Runtime/Engine/Private/SceneManagement.cpp b/Engine/Source/Runtime/Engine/Private/SceneManagement.cpp
--- a/Engine/Source/Runtime/Engine/Private/SceneManagement.cpp
+++ b/Engine/Source/Runtime/Engine/Private/SceneManagement.cpp
@@ -106,28 +106,6 @@ void FTemporalLODState::UpdateTemporalLODTransition(const FSceneView& View, floa
 	}
 }
 
-FFrozenSceneViewMatricesGuard::FFrozenSceneViewMatricesGuard(FSceneView& SV)
-	: SceneView(SV)
-{
-#if !(UE_BUILD_SHIPPING || UE_BUILD_TEST)
-	if (SceneView.State)
-	{
-		SceneView.State->ActivateFrozenViewMatrices(SceneView);
-	}
-#endif
-}
-
-FFrozenSceneViewMatricesGuard::~FFrozenSceneViewMatricesGuard()
-{
-#if !(UE_BUILD_SHIPPING || UE_BUILD_TEST)
-	if (SceneView.State)
-	{
-		SceneView.State->RestoreUnfrozenViewMatrices(SceneView);
-	}
-#endif
-}
-
-
 IMPLEMENT_STATIC_UNIFORM_BUFFER_SLOT(WorkingColorSpace);
 IMPLEMENT_STATIC_UNIFORM_BUFFER_STRUCT(FWorkingColorSpaceShaderParameters, "WorkingColorSpace", WorkingColorSpace);
 
diff --git a/Engine/Source/Runtime/Engine/Private/StaticMeshSceneProxy.cpp b/Engine/Source/Runtime/Engine/Private/StaticMeshSceneProxy.cpp
--- a/Engine/Source/Runtime/Engine/Private/StaticMeshSceneProxy.cpp
+++ b/Engine/Source/Runtime/Engine/Private/StaticMeshSceneProxy.cpp
@@ -1671,8 +1671,6 @@ void FStaticMeshSceneProxy::GetDynamicMeshElements(const TArray<const FSceneView
 
 			if (IsShown(View) && (VisibilityMap & (1 << ViewIndex)))
 			{
-				FFrozenSceneViewMatricesGuard FrozenMatricesGuard(*const_cast<FSceneView*>(Views[ViewIndex]));
-
 				FLODMask LODMask = GetLODMask(View);
 
 				for (int32 LODIndex = 0; LODIndex < RenderData->LODResources.Num(); LODIndex++)
diff --git a/Engine/Source/Runtime/Engine/Public/SceneManagement.h b/Engine/Source/Runtime/Engine/Public/SceneManagement.h
--- a/Engine/Source/Runtime/Engine/Public/SceneManagement.h
+++ b/Engine/Source/Runtime/Engine/Public/SceneManagement.h
@@ -275,17 +275,13 @@ private:
 	friend class FScene;
 };
 
-class FFrozenSceneViewMatricesGuard
+class UE_DEPRECATED(5.7, "FFrozenSceneViewMatricesGuard is no longer used.") FFrozenSceneViewMatricesGuard
 {
 public:
-	ENGINE_API FFrozenSceneViewMatricesGuard(FSceneView& SV);
-	ENGINE_API ~FFrozenSceneViewMatricesGuard();
-
-private:
-	FSceneView& SceneView;
+	FFrozenSceneViewMatricesGuard(FSceneView& SV) {}
+	~FFrozenSceneViewMatricesGuard() {}
 };
 
-
 /**
  * Global working color space shader parameters (color space conversion matrices).
  */
diff --git a/Engine/Source/Runtime/InteractiveToolsFramework/Private/BaseGizmos/ViewAdjustedStaticMeshGizmoComponent.cpp b/Engine/Source/Runtime/InteractiveToolsFramework/Private/BaseGizmos/ViewAdjustedStaticMeshGizmoComponent.cpp
--- a/Engine/Source/Runtime/InteractiveToolsFramework/Private/BaseGizmos/ViewAdjustedStaticMeshGizmoComponent.cpp
+++ b/Engine/Source/Runtime/InteractiveToolsFramework/Private/BaseGizmos/ViewAdjustedStaticMeshGizmoComponent.cpp
@@ -85,8 +85,6 @@ namespace ViewAdjustedStaticMeshGizmoComponentLocals
 					continue;
 				}
 
-				FFrozenSceneViewMatricesGuard FrozenMatricesGuard(*const_cast<FSceneView*>(Views[ViewIndex]));
-
 				// We can calculate our adjusted transform at this point, now that we have the view
 				const TUniformBuffer<FPrimitiveUniformShaderParameters>* AdjustedTransformBuffer = nullptr;
 				bool bAdjustedTransformDeterminantIsNegative = IsLocalToWorldDeterminantNegative();
diff --git a/Engine/Source/Runtime/Renderer/Private/SceneVisibility.cpp b/Engine/Source/Runtime/Renderer/Private/SceneVisibility.cpp
--- a/Engine/Source/Runtime/Renderer/Private/SceneVisibility.cpp
+++ b/Engine/Source/Runtime/Renderer/Private/SceneVisibility.cpp
@@ -3397,6 +3397,20 @@ FVisibilityTaskConfig::FVisibilityTaskConfig(const FScene& Scene, TConstArrayVie
 		{
 			Schedule = EVisibilityTaskSchedule::RenderThread;
 		}
+
+#if !(UE_BUILD_SHIPPING || UE_BUILD_TEST)
+		for (FViewInfo* View : Views)
+		{
+			auto* ViewState = static_cast<FSceneViewState*>(View->State);
+
+			// Frozen matrices require modifying the view matrices which does not support the async path.
+			if (ViewState && ViewState->bIsFrozen)
+			{
+				Schedule = EVisibilityTaskSchedule::RenderThread;
+				break;
+			}
+		}
+#endif
 	}
 
 	NumTestedPrimitives = uint32(Scene.Primitives.Num());
@@ -4611,6 +4625,13 @@ void FVisibilityTaskData::ProcessRenderThreadTasks()
 	{
 		for (FVisibilityViewPacket& ViewPacket : ViewPackets)
 		{
+		#if !(UE_BUILD_SHIPPING || UE_BUILD_TEST)
+			if (ViewPacket.ViewState)
+			{
+				ViewPacket.ViewState->ActivateFrozenViewMatrices(ViewPacket.View);
+			}
+		#endif
+
 			ViewPacket.BeginInitVisibility();
 			UpdatePrimitiveFading(Scene, ViewPacket.View, ViewPacket.ViewState, PrimitiveRange);
 		}
@@ -4654,6 +4675,16 @@ void FVisibilityTaskData::ProcessRenderThreadTasks()
 
 		check(DynamicMeshElements.PrimitiveViewMasks);
 		GatherDynamicMeshElements(*DynamicMeshElements.PrimitiveViewMasks);
+
+	#if !(UE_BUILD_SHIPPING || UE_BUILD_TEST)
+		for (FVisibilityViewPacket& ViewPacket : ViewPackets)
+		{
+			if (ViewPacket.ViewState)
+			{
+				ViewPacket.ViewState->RestoreUnfrozenViewMatrices(ViewPacket.View);
+			}
+		}
+	#endif
 	}
 	else
 	{
