Fix issue where InstanceCulling UB is not set correctly for parallel basepass EditorPrimitives.

diff --git a/Engine/Source/Runtime/Renderer/Private/BasePassRendering.cpp b/Engine/Source/Runtime/Renderer/Private/BasePassRendering.cpp
--- a/Engine/Source/Runtime/Renderer/Private/BasePassRendering.cpp
+++ b/Engine/Source/Runtime/Renderer/Private/BasePassRendering.cpp
@@ -1588,6 +1588,10 @@ void FDeferredShadingSceneRenderer::RenderBasePassInternal(
 						Pass->Dispatch(DispatchPassBuilder, &PassParameters->InstanceCullingDrawParams);
 					});
 				}
+				else
+				{
+					InstanceCullingManager.SetDummyCullingParams(GraphBuilder, PassParameters->InstanceCullingDrawParams);
+				}
 
 				const bool bShouldRenderViewForNanite = bNaniteEnabled
 					&& !View.bHasNoVisiblePrimitive
