[HWRT] Refactored per-frame initialization of RayTracingScene and RayTracingSBT.
- moved initalization earlier in FDeferredShadingSceneRenderer::Render(...)
- check whether bNeedsInstanceExtraDataBuffer or bTracingFeedbackEnabled are necessary once per batch of renderers
- these flags affect data gathered per instance which will be important once we cache instances in FRayTracingScene.
- also made those flags private and exposed setter/getter methods to be able to track when they change.

diff --git a/Engine/Source/Runtime/Renderer/Private/DeferredShadingRenderer.cpp b/Engine/Source/Runtime/Renderer/Private/DeferredShadingRenderer.cpp
--- a/Engine/Source/Runtime/Renderer/Private/DeferredShadingRenderer.cpp
+++ b/Engine/Source/Runtime/Renderer/Private/DeferredShadingRenderer.cpp
@@ -938,7 +938,7 @@ bool FDeferredShadingSceneRenderer::DispatchRayTracingWorldUpdates(FRDGBuilder&
 	RDG_GPU_MASK_SCOPE(GraphBuilder, GRayTracingMultiGpuTLASMask ? GraphBuilder.RHICmdList.GetGPUMask() : FRHIGPUMask::All());
 
 	FRayTracingDynamicGeometryUpdateManager* DynamicGeometryUpdateManager = Scene->GetRayTracingDynamicGeometryUpdateManager();
-	DynamicGeometryUpdateManager->AddDynamicGeometryUpdatePass(GraphBuilder, ComputePassFlags, GetSceneUniformBufferRef(GraphBuilder), RayTracingScene.bTracingFeedbackEnabled, ResourceAccessPipelines, OutDynamicGeometryScratchBuffer);
+	DynamicGeometryUpdateManager->AddDynamicGeometryUpdatePass(GraphBuilder, ComputePassFlags, GetSceneUniformBufferRef(GraphBuilder), RayTracingScene.IsRayTracingFeedbackEnabled(), ResourceAccessPipelines, OutDynamicGeometryScratchBuffer);
 
 	((FRayTracingGeometryManager*)GRayTracingGeometryManager)->ResetVisibleGeometries();
 
@@ -1633,6 +1633,17 @@ void FDeferredShadingSceneRenderer::Render(FRDGBuilder& GraphBuilder, const FSce
 
 		GRayTracingGeometryManager->ProcessBuildRequests(GraphBuilder.RHICmdList);
 	}
+
+	const int32 ReferenceViewIndex = 0;
+	FViewInfo& ReferenceView = Views[ReferenceViewIndex];
+
+	FRayTracingShaderBindingTable& RayTracingSBT = Scene->RayTracingSBT;
+	FRayTracingScene& RayTracingScene = Scene->RayTracingScene;
+	RayTracingSBT.ResetMissAndCallableShaders();
+
+	RayTracingScene.Reset(); // Resets the internal arrays, but does not release any resources.
+	RayTracingScene.InitPreViewTranslation(ReferenceView.ViewMatrices);
+
 #endif
 
 	FInitViewTaskDatas InitViewTaskDatas = OnRenderBegin(GraphBuilder, SceneUpdateInputs);
@@ -1798,9 +1809,6 @@ void FDeferredShadingSceneRenderer::Render(FRDGBuilder& GraphBuilder, const FSce
 #if RHI_RAYTRACING
 		if (FamilyPipelineState[&FFamilyPipelineState::bRayTracing])
 		{
-			const int32 ReferenceViewIndex = 0;
-			FViewInfo& ReferenceView = Views[ReferenceViewIndex];
-
 			InitViewTaskDatas.RayTracingGatherInstances = RayTracing::CreateGatherInstancesTaskData(
 				Allocator,
 				*Scene,
@@ -1924,21 +1932,11 @@ void FDeferredShadingSceneRenderer::Render(FRDGBuilder& GraphBuilder, const FSce
 		}
 	}
 
-#if RHI_RAYTRACING
-	const int32 ReferenceViewIndex = 0;
-	FViewInfo& ReferenceView = Views[ReferenceViewIndex];
-	FRayTracingScene& RayTracingScene = Scene->RayTracingScene;
-	FRayTracingShaderBindingTable& RayTracingSBT = Scene->RayTracingSBT;
-#endif
-
 	if (RendererOutput == ERendererOutput::FinalSceneColor)
 	{
-#if RHI_RAYTRACING
 		// Prepare the scene for rendering this frame.
-		RayTracingScene.Reset(IsRayTracingInstanceDebugDataEnabled(ReferenceView)); // Resets the internal arrays, but does not release any resources.
-		
-		RayTracingSBT.ResetMissAndCallableShaders();
 
+#if RHI_RAYTRACING
 		if (ViewFamily.EngineShowFlags.PathTracing)
 		{
 			if (ShouldPrepareRayTracingDecals(*Scene, ViewFamily))
diff --git a/Engine/Source/Runtime/Renderer/Private/HeterogeneousVolumes/HeterogeneousVolumesHardwareRayTracing.cpp b/Engine/Source/Runtime/Renderer/Private/HeterogeneousVolumes/HeterogeneousVolumesHardwareRayTracing.cpp
--- a/Engine/Source/Runtime/Renderer/Private/HeterogeneousVolumes/HeterogeneousVolumesHardwareRayTracing.cpp
+++ b/Engine/Source/Runtime/Renderer/Private/HeterogeneousVolumes/HeterogeneousVolumesHardwareRayTracing.cpp
@@ -196,7 +196,7 @@ void GenerateRayTracingScene(
 	FRayTracingScene& RayTracingScene
 )
 {
-	RayTracingScene.Reset(false);
+	RayTracingScene.Reset();
 
 	TArray<FRayTracingGeometryBuildParams> BuildParams;
 	uint32 BLASScratchSize = 0;
diff --git a/Engine/Source/Runtime/Renderer/Private/RayTracing/RayTracing.cpp b/Engine/Source/Runtime/Renderer/Private/RayTracing/RayTracing.cpp
--- a/Engine/Source/Runtime/Renderer/Private/RayTracing/RayTracing.cpp
+++ b/Engine/Source/Runtime/Renderer/Private/RayTracing/RayTracing.cpp
@@ -406,6 +406,27 @@ namespace RayTracing
 
 		FScene& Scene = *SceneUpdateInputs.Scene;
 
+		bool bAnyViewNeedsInstanceExtraDataBuffer = false;
+		bool bAnyViewNeedsRayTracingInstanceDebugData = false;
+
+		for (const FViewInfo* View : SceneUpdateInputs.Views)
+		{
+			bAnyViewNeedsInstanceExtraDataBuffer |= IsRayTracingInstanceOverlapEnabled(*View);
+			bAnyViewNeedsRayTracingInstanceDebugData |= IsRayTracingInstanceDebugDataEnabled(*View);
+		}
+
+		Scene.RayTracingScene.SetInstanceExtraDataBufferEnabled(bAnyViewNeedsInstanceExtraDataBuffer);
+		Scene.RayTracingScene.SetInstanceDebugDataEnabled(bAnyViewNeedsRayTracingInstanceDebugData);
+
+		bool bAnyViewFamilyUsingRayTracingFeedback = false;
+
+		for (const FViewFamilyInfo* ViewFamily : SceneUpdateInputs.ViewFamilies)
+		{
+			bAnyViewFamilyUsingRayTracingFeedback |= IsRayTracingFeedbackEnabled(*ViewFamily);
+		}
+
+		Scene.RayTracingScene.SetTracingFeedbackEnabled(bAnyViewFamilyUsingRayTracingFeedback);
+
 		if (CurrentType != Scene.CachedRayTracingMeshCommandsType
 			|| bNaniteCoarseMeshStreamingModeChanged
 			|| bNaniteRayTracingModeChanged
diff --git a/Engine/Source/Runtime/Renderer/Private/RayTracing/RayTracingScene.cpp b/Engine/Source/Runtime/Renderer/Private/RayTracing/RayTracingScene.cpp
--- a/Engine/Source/Runtime/Renderer/Private/RayTracing/RayTracingScene.cpp
+++ b/Engine/Source/Runtime/Renderer/Private/RayTracing/RayTracingScene.cpp
@@ -271,7 +271,7 @@ void FRayTracingScene::Create(FRDGBuilder& GraphBuilder, const FViewInfo& View,
 		}
 
 		FRDGBufferUAVRef InstanceExtraDataBufferUAV = nullptr;
-		if (bNeedsInstanceExtraDataBuffer || bTracingFeedbackEnabled || bInstanceDebugDataEnabled)
+		if (bInstanceExtraDataBufferEnabled || bTracingFeedbackEnabled || bInstanceDebugDataEnabled)
 		{
 			FRDGBufferDesc InstanceExtraDataBufferDesc;
 			InstanceExtraDataBufferDesc.Usage = EBufferUsageFlags::UnorderedAccess | EBufferUsageFlags::ShaderResource | EBufferUsageFlags::StructuredBuffer;
@@ -860,7 +860,7 @@ void FRayTracingScene::SetInstance(FInstanceRange InstanceRange, uint32 Instance
 	}
 }
 
-void FRayTracingScene::Reset(bool bInInstanceDebugDataEnabled)
+void FRayTracingScene::Reset()
 {
 	const uint8 NumLayers = uint8(ERayTracingSceneLayer::NUM);
 
@@ -868,6 +868,8 @@ void FRayTracingScene::Reset(bool bInInstanceDebugDataEnabled)
 	{
 		FLayer& Layer = Layers[LayerIndex];
 
+		Layer.InstanceBufferBuilder = {};
+
 		Layer.Instances.Reset();
 		Layer.InstancesDebugData.Reset();
 
@@ -883,7 +885,7 @@ void FRayTracingScene::Reset(bool bInInstanceDebugDataEnabled)
 		Layer.InstanceHitCountBuffer = nullptr;
 		Layer.AccelerationStructureIndexBuffer = nullptr;
 		Layer.GeometryHandleBuffer = nullptr;
-		Layer.GeometryHandles.Reset();		
+		Layer.GeometryHandles.Reset();
 	}
 
 	InstanceStatsBuffer = nullptr;
@@ -893,13 +895,11 @@ void FRayTracingScene::Reset(bool bInInstanceDebugDataEnabled)
 	Allocator.Flush();
 
 	bUsesLightingChannels = false;
-
-	bInstanceDebugDataEnabled = bInInstanceDebugDataEnabled;
 }
 
 void FRayTracingScene::EndFrame()
 {
-	Reset(false);
+	Reset();
 
 	// Release the resources if ray tracing wasn't used
 	if (!bUsedThisFrame)
@@ -921,6 +921,53 @@ void FRayTracingScene::EndFrame()
 	bInitializationDataBuilt = false;
 }
 
+bool FRayTracingScene::SetInstanceExtraDataBufferEnabled(bool bInEnabled)
+{
+	bInstanceExtraDataBufferEnabled = bInEnabled;
+
+	return false;
+}
+
+bool FRayTracingScene::SetTracingFeedbackEnabled(bool bInEnabled)
+{
+	bool bChanged = (bTracingFeedbackEnabled != bInEnabled);
+
+	bTracingFeedbackEnabled = bInEnabled;
+
+	if (bChanged)
+	{
+		const uint8 NumLayers = uint8(ERayTracingSceneLayer::NUM);
+
+		for (uint32 LayerIndex = 0; LayerIndex < NumLayers; ++LayerIndex)
+		{
+			FLayer& Layer = Layers[LayerIndex];
+			checkf(Layer.Instances.IsEmpty(), TEXT("Can't change bTracingFeedbackEnabled when there are already instances."));
+		}
+	}
+
+	return bChanged;
+}
+
+bool FRayTracingScene::SetInstanceDebugDataEnabled(bool bInEnabled)
+{
+	bool bChanged = (bInstanceDebugDataEnabled != bInEnabled);
+
+	bInstanceDebugDataEnabled = bInEnabled;
+
+	if (bChanged)
+	{
+		const uint8 NumLayers = uint8(ERayTracingSceneLayer::NUM);
+
+		for (uint32 LayerIndex = 0; LayerIndex < NumLayers; ++LayerIndex)
+		{
+			FLayer& Layer = Layers[LayerIndex];
+			checkf(Layer.Instances.IsEmpty(), TEXT("Can't change bInstanceDebugDataEnabled when there are already instances."));
+		}
+	}
+
+	return bChanged;
+}
+
 void FRayTracingScene::ReleaseReadbackBuffers()
 {
 	for (auto& ReadbackBuffer : StatsReadbackBuffers)
diff --git a/Engine/Source/Runtime/Renderer/Private/RayTracing/RayTracingScene.h b/Engine/Source/Runtime/Renderer/Private/RayTracing/RayTracingScene.h
--- a/Engine/Source/Runtime/Renderer/Private/RayTracing/RayTracingScene.h
+++ b/Engine/Source/Runtime/Renderer/Private/RayTracing/RayTracingScene.h
@@ -113,10 +113,16 @@ public:
 	void PostRender(FRDGBuilder& GraphBuilder, ERDGPassFlags ComputePassFlags = ERDGPassFlags::Compute);
 
 	// Resets the instance list and reserves memory for this frame.
-	void Reset(bool bInstanceDebugDataEnabled);
+	void Reset();
 
 	void EndFrame();
 
+	bool SetInstanceExtraDataBufferEnabled(bool bEnabled);
+	bool SetTracingFeedbackEnabled(bool bEnabled);
+	bool SetInstanceDebugDataEnabled(bool bEnabled);
+
+	bool IsRayTracingFeedbackEnabled() const { return bTracingFeedbackEnabled; }
+
 	// Allocates temporary memory that will be valid until the next Reset().
 	// Can be used to store temporary instance transforms, user data, etc.
 	template <typename T>
@@ -164,9 +170,6 @@ public:
 	// Geometries which still have a pending build request but are used this frame and require a force build.
 	TArray<const FRayTracingGeometry*> GeometriesToBuild;
 
-	bool bNeedsInstanceExtraDataBuffer = false;
-	bool bTracingFeedbackEnabled = false;
-
 	bool bUsesLightingChannels = false;
 
 	// Used for transforming to translated world space in which TLAS was built.
@@ -223,6 +226,8 @@ private:
 	// Transient memory allocator
 	FMemStackBase Allocator;
 
+	bool bInstanceExtraDataBufferEnabled = false;
+	bool bTracingFeedbackEnabled = false;
 	bool bInstanceDebugDataEnabled = false;
 
 	bool bInitializationDataBuilt = false;
diff --git a/Engine/Source/Runtime/Renderer/Private/SceneRendering.cpp b/Engine/Source/Runtime/Renderer/Private/SceneRendering.cpp
--- a/Engine/Source/Runtime/Renderer/Private/SceneRendering.cpp
+++ b/Engine/Source/Runtime/Renderer/Private/SceneRendering.cpp
@@ -1347,9 +1347,9 @@ void FViewInfo::SetupUniformBufferParameters(
 		}
 
 #if RHI_RAYTRACING
-		const FDFVector3& PreViewTranslation = Scene->RayTracingScene.PreViewTranslation;
-		ViewUniformShaderParameters.TLASPreViewTranslationHigh = PreViewTranslation.High;
-		ViewUniformShaderParameters.TLASPreViewTranslationLow = PreViewTranslation.Low;
+		const FDFVector3& RayTracingScenePreViewTranslation = Scene->RayTracingScene.PreViewTranslation;
+		ViewUniformShaderParameters.TLASPreViewTranslationHigh = RayTracingScenePreViewTranslation.High;
+		ViewUniformShaderParameters.TLASPreViewTranslationLow = RayTracingScenePreViewTranslation.Low;
 #endif
 
 		// Set default atmosphere lights parameters
diff --git a/Engine/Source/Runtime/Renderer/Private/SceneVisibility.cpp b/Engine/Source/Runtime/Renderer/Private/SceneVisibility.cpp
--- a/Engine/Source/Runtime/Renderer/Private/SceneVisibility.cpp
+++ b/Engine/Source/Runtime/Renderer/Private/SceneVisibility.cpp
@@ -4888,19 +4888,12 @@ void FSceneRenderer::PreVisibilityFrameSetup(FRDGBuilder& GraphBuilder)
 #endif
 
 #if RHI_RAYTRACING
-		if (Scene && Views.Num())
+		if (Scene)
 		{
-			const int32 ReferenceViewIndex = 0;
-			const FViewInfo& ReferenceView = Views[ReferenceViewIndex];
-
-			Scene->RayTracingScene.InitPreViewTranslation(ReferenceView.ViewMatrices);
-			Scene->RayTracingScene.bNeedsInstanceExtraDataBuffer = IsRayTracingInstanceOverlapEnabled(ReferenceView);
-			Scene->RayTracingScene.bTracingFeedbackEnabled = IsRayTracingFeedbackEnabled(ViewFamily);
-
 			for (int32 ViewIndex = 0; ViewIndex < Views.Num(); ViewIndex++)
 			{
 				FViewInfo& View = Views[ViewIndex];
-				View.bRayTracingFeedbackEnabled = Scene->RayTracingScene.bTracingFeedbackEnabled;
+				View.bRayTracingFeedbackEnabled = Scene->RayTracingScene.IsRayTracingFeedbackEnabled();
 			}
 		}
 #endif
