[HWRT] Decouple dynamic ray tracing instance mask/flags caching from CachedRayTracingInstance.

Context:
- CachedRayTracingInstance will be moved to FRayTracingScene soon, so can't use it to store dynamic ray tracing instance data.
- bCachedRayTracingInstance[...] flags are available in FRayTracingLODData for cached static instances.

Changes:
- Cached dynamic ray tracing instance data in a new struct FDynamicRayTracingInstanceCachedData.
- Get flags from FRayTracingLODData in the cached static instances code path.

diff --git a/Engine/Source/Runtime/Renderer/Private/PrimitiveSceneInfo.cpp b/Engine/Source/Runtime/Renderer/Private/PrimitiveSceneInfo.cpp
--- a/Engine/Source/Runtime/Renderer/Private/PrimitiveSceneInfo.cpp
+++ b/Engine/Source/Runtime/Renderer/Private/PrimitiveSceneInfo.cpp
@@ -1041,7 +1041,6 @@ struct FDeferredRayTracingMeshCommandData
 {
 	FPrimitiveSceneInfo* SceneInfo;
 	TArray<TArray<int32, TInlineAllocator<2>>, TInlineAllocator<1>> CommandIndicesPerLOD;
-
 };
 
 template<class T>
@@ -1316,7 +1315,7 @@ void FPrimitiveSceneInfo::UpdateCachedRayTracingInstance(FPrimitiveSceneInfo* Sc
 
 		if (RayTracingInstance.GetMaterials().IsEmpty())
 		{
-		 	// If the material list is empty, explicitly set the mask to 0 so it will not be added in the raytracing scene
+			// If the material list is empty, explicitly set the mask to 0 so it will not be added in the raytracing scene
 			InstanceMaskAndFlags.Mask = 0;
 		}
 		else
@@ -1324,21 +1323,26 @@ void FPrimitiveSceneInfo::UpdateCachedRayTracingInstance(FPrimitiveSceneInfo* Sc
 			InstanceMaskAndFlags = BuildRayTracingInstanceMaskAndFlags(RayTracingInstance, *SceneProxy);
 		}
 
-		SceneInfo->UpdateCachedRayTracingInstanceMaskAndFlags(InstanceMaskAndFlags);
-	}
-	else
-	{
-		SceneInfo->CachedRayTracingInstance.Mask = 0xFF;
-		SceneInfo->CachedRayTracingInstance.Flags = ERayTracingInstanceFlags::None;
-		SceneInfo->bCachedRayTracingInstanceAnySegmentsDecal = false;
-		SceneInfo->bCachedRayTracingInstanceAllSegmentsDecal = false;
-		SceneInfo->bCachedRayTracingInstanceAllSegmentsTranslucent = false;
+		SceneInfo->CachedRayTracingInstance.Mask = InstanceMaskAndFlags.Mask;
 
-		SceneInfo->bCachedRayTracingInstanceMaskAndFlagsDirty = true;
+		if (InstanceMaskAndFlags.bForceOpaque)
+		{
+			SceneInfo->CachedRayTracingInstance.Flags |= ERayTracingInstanceFlags::ForceOpaque;
+		}
+
+		if (InstanceMaskAndFlags.bDoubleSided)
+		{
+			SceneInfo->CachedRayTracingInstance.Flags |= ERayTracingInstanceFlags::TriangleCullDisable;
+		}
+
+		if (InstanceMaskAndFlags.bReverseCulling)
+		{
+			SceneInfo->CachedRayTracingInstance.Flags |= ERayTracingInstanceFlags::TriangleCullReverse;
+		}
 	}
 }
 
-void FPrimitiveSceneInfo::SetCachedRayTracingInstanceGeometryRHI(FRHIRayTracingGeometry* Geometry)
+void FPrimitiveSceneInfo::SetCachedRayTracingInstanceGeometryRHI(FRHIRayTracingGeometry* GeometryRHI)
 {
 	// no cached RT LOD data?
 	if (RayTracingLODData.IsEmpty())
@@ -1361,38 +1365,12 @@ void FPrimitiveSceneInfo::SetCachedRayTracingInstanceGeometryRHI(FRHIRayTracingG
 		check(CachedRayTracingInstance.GeometryRHI == nullptr);
 	}
 
-	CachedRayTracingInstance.GeometryRHI = Geometry;
+	CachedRayTracingInstance.GeometryRHI = GeometryRHI;
 
 	AllocateRayTracingSBT();
 	CacheRayTracingShaderBindingData();
 }
 
-void FPrimitiveSceneInfo::UpdateCachedRayTracingInstanceMaskAndFlags(FRayTracingMaskAndFlags& InstanceMaskAndFlags)
-{	
-	CachedRayTracingInstance.Mask = InstanceMaskAndFlags.Mask; // When no cached command is found, InstanceMask == 0 and the instance is effectively filtered out
-
-	if (InstanceMaskAndFlags.bForceOpaque)
-	{
-		CachedRayTracingInstance.Flags |= ERayTracingInstanceFlags::ForceOpaque;
-	}
-
-	if (InstanceMaskAndFlags.bDoubleSided)
-	{
-		CachedRayTracingInstance.Flags |= ERayTracingInstanceFlags::TriangleCullDisable;
-	}
-
-	if (InstanceMaskAndFlags.bReverseCulling)
-	{
-		CachedRayTracingInstance.Flags |= ERayTracingInstanceFlags::TriangleCullReverse;
-	}
-
-	bCachedRayTracingInstanceAnySegmentsDecal = InstanceMaskAndFlags.bAnySegmentsDecal;
-	bCachedRayTracingInstanceAllSegmentsDecal = InstanceMaskAndFlags.bAllSegmentsDecal;
-	bCachedRayTracingInstanceAllSegmentsTranslucent = InstanceMaskAndFlags.bAllSegmentsTranslucent;
-
-	bCachedRayTracingInstanceMaskAndFlagsDirty = false;
-}
-
 void FPrimitiveSceneInfo::RemoveCachedRayTracingPrimitives()
 {
 	if (IsRayTracingAllowed())
diff --git a/Engine/Source/Runtime/Renderer/Private/RayTracing/RayTracing.cpp b/Engine/Source/Runtime/Renderer/Private/RayTracing/RayTracing.cpp
--- a/Engine/Source/Runtime/Renderer/Private/RayTracing/RayTracing.cpp
+++ b/Engine/Source/Runtime/Renderer/Private/RayTracing/RayTracing.cpp
@@ -834,7 +834,7 @@ namespace RayTracing
 						continue;
 					}
 
-					if (Instance.bInstanceMaskAndFlagsDirty || SceneInfo->bCachedRayTracingInstanceMaskAndFlagsDirty)
+					if (Instance.bInstanceMaskAndFlagsDirty || SceneInfo->bDynamicRayTracingInstanceCachedDataDirty)
 					{
 						// Build InstanceMaskAndFlags since the data in SceneInfo is not up to date
 
@@ -850,24 +850,31 @@ namespace RayTracing
 							InstanceMaskAndFlags = BuildRayTracingInstanceMaskAndFlags(Instance, *SceneProxy);
 						}
 
-						SceneInfo->UpdateCachedRayTracingInstanceMaskAndFlags(InstanceMaskAndFlags);
+						SceneInfo->DynamicRayTracingInstanceCachedData.Mask = InstanceMaskAndFlags.Mask;
+						SceneInfo->DynamicRayTracingInstanceCachedData.bForceOpaque = InstanceMaskAndFlags.bForceOpaque;
+						SceneInfo->DynamicRayTracingInstanceCachedData.bDoubleSided = InstanceMaskAndFlags.bDoubleSided;
+						SceneInfo->DynamicRayTracingInstanceCachedData.bReverseCulling = InstanceMaskAndFlags.bReverseCulling;
+						SceneInfo->DynamicRayTracingInstanceCachedData.bAnySegmentsDecal = InstanceMaskAndFlags.bAnySegmentsDecal;
+						SceneInfo->DynamicRayTracingInstanceCachedData.bAllSegmentsDecal = InstanceMaskAndFlags.bAllSegmentsDecal;
+						SceneInfo->DynamicRayTracingInstanceCachedData.bAllSegmentsTranslucent = InstanceMaskAndFlags.bAllSegmentsTranslucent;
+						SceneInfo->bDynamicRayTracingInstanceCachedDataDirty = false;
 					}
 
-					if (!SceneOptions.bTranslucentGeometry && SceneInfo->bCachedRayTracingInstanceAllSegmentsTranslucent)
+					if (!SceneOptions.bTranslucentGeometry && SceneInfo->DynamicRayTracingInstanceCachedData.bAllSegmentsTranslucent)
 					{
 						continue;
 					}
 
 					// TODO: Do we want to support dynamic instances in far field?
-					const bool bNeedMainInstance = !SceneInfo->bCachedRayTracingInstanceAllSegmentsDecal;
+					const bool bNeedMainInstance = !SceneInfo->DynamicRayTracingInstanceCachedData.bAllSegmentsDecal;
 
 					// if primitive has mixed decal and non-decal segments we need to have two ray tracing instances
 					// one containing non-decal segments and the other with decal segments
 					// masking of segments is done using "hidden" hitgroups
 					// TODO: Debug Visualization to highlight primitives using this?
-					const bool bNeedDecalInstance = SceneInfo->bCachedRayTracingInstanceAnySegmentsDecal && !ShouldExcludeDecals();
+					const bool bNeedDecalInstance = SceneInfo->DynamicRayTracingInstanceCachedData.bAnySegmentsDecal && !ShouldExcludeDecals();
 
-					if (ShouldExcludeDecals() && SceneInfo->bCachedRayTracingInstanceAllSegmentsDecal)
+					if (ShouldExcludeDecals() && SceneInfo->DynamicRayTracingInstanceCachedData.bAllSegmentsDecal)
 					{
 						continue;
 					}
@@ -890,8 +897,19 @@ namespace RayTracing
 					RayTracingInstance.DefaultUserData = InstanceSceneDataOffset;
 					RayTracingInstance.bIncrementUserDataPerInstance = true;
 					RayTracingInstance.bApplyLocalBoundsTransform = Instance.bApplyLocalBoundsTransform;
-					RayTracingInstance.Mask = SceneInfo->GetCachedRayTracingInstance().Mask;
-					RayTracingInstance.Flags = SceneInfo->GetCachedRayTracingInstance().Flags;
+					RayTracingInstance.Mask = SceneInfo->DynamicRayTracingInstanceCachedData.Mask;
+					if (SceneInfo->DynamicRayTracingInstanceCachedData.bForceOpaque)
+					{
+						RayTracingInstance.Flags |= ERayTracingInstanceFlags::ForceOpaque;
+					}
+					if (SceneInfo->DynamicRayTracingInstanceCachedData.bDoubleSided)
+					{
+						RayTracingInstance.Flags |= ERayTracingInstanceFlags::TriangleCullDisable;
+					}
+					if (SceneInfo->DynamicRayTracingInstanceCachedData.bReverseCulling)
+					{
+						RayTracingInstance.Flags |= ERayTracingInstanceFlags::TriangleCullReverse;
+					}
 					AddRayTracingInstanceFlagOverrides(RayTracingInstance.Flags, bUsesLightingChannels);
 
 					if (!Instance.GetPrimitiveInstanceIndices().IsEmpty())
@@ -1193,6 +1211,8 @@ namespace RayTracing
 	{
 		TRACE_CPUPROFILER_EVENT_SCOPE(GatherRayTracingRelevantStaticPrimitives);
 
+		const bool bExcludeDecals = ShouldExcludeDecals();
+
 		struct FRelevantStaticPrimitivesContext
 		{
 			FRelevantStaticPrimitivesContext(int32 InContextIndex) : ContextIndex(InContextIndex) {}
@@ -1218,7 +1238,7 @@ namespace RayTracing
 			TaskData.StaticPrimitivesIndices.Num(),
 			[](int32 ContextIndex, int32 NumContexts) { return ContextIndex; },
 			[&Scene = TaskData.Scene, &View = TaskData.View, GlobalLODScale, ForcedLODLevel, &StaticPrimitivesIndices = TaskData.StaticPrimitivesIndices,
-			bUsingReferenceBasedResidency, bLightingChannelsUsingAHS = TaskData.SceneOptions.bLightingChannelsUsingAHS](FRelevantStaticPrimitivesContext& Context, int32 ItemIndex)
+			bUsingReferenceBasedResidency, bLightingChannelsUsingAHS = TaskData.SceneOptions.bLightingChannelsUsingAHS, bExcludeDecals](FRelevantStaticPrimitivesContext& Context, int32 ItemIndex)
 			{
 				const int32 PrimitiveIndex = StaticPrimitivesIndices[ItemIndex];
 
@@ -1227,6 +1247,8 @@ namespace RayTracing
 				const ERayTracingPrimitiveFlags Flags = Scene.PrimitiveRayTracingFlags[PrimitiveIndex];
 				const FScene::FPrimitiveRayTracingData& PrimitiveRayTracingData = Scene.PrimitiveRayTracingDatas[PrimitiveIndex];
 
+				ensureMsgf(!PrimitiveRayTracingData.bCachedRaytracingDataDirty, TEXT("Cached ray tracing instances must be up-to-date at this point"));
+
 				const bool bUsingNaniteRayTracing = (Nanite::GetRayTracingMode() != Nanite::ERayTracingMode::Fallback) && SceneProxy->IsNaniteMesh();
 
 				if (bUsingNaniteRayTracing)
@@ -1267,12 +1289,6 @@ namespace RayTracing
 					{
 						// Currently IsCachedRayTracingGeometryValid() can only be called for non-nanite geometries
 						checkf(SceneInfo->IsCachedRayTracingGeometryValid(), TEXT("Cached ray tracing instance is expected to be valid. Was mesh LOD streamed but cached data was not invalidated?"));
-						checkf(CachedRayTracingInstance.GeometryRHI, TEXT("Ray tracing instance must have a valid geometry."));
-					}
-
-					if (ShouldExcludeDecals() && SceneInfo->bCachedRayTracingInstanceAllSegmentsDecal)
-					{
-						return;
 					}
 
 					checkf(CachedRayTracingInstance.GeometryRHI, TEXT("Ray tracing instance must have a valid geometry."));
@@ -1283,6 +1299,12 @@ namespace RayTracing
 					LODIndex = 0;
 
 					const FPrimitiveSceneInfo::FRayTracingLODData& RTLODData = SceneInfo->GetRayTracingLODData(LODIndex);
+
+					if (bExcludeDecals && RTLODData.CachedMeshCommandFlags.bAllSegmentsDecal)
+					{
+						return;
+					}
+
 					ensure(RTLODData.SBTAllocation);
 					if (RTLODData.SBTAllocation == nullptr)
 					{
@@ -1305,8 +1327,8 @@ namespace RayTracing
 					ensureMsgf(!PrimitiveRayTracingData.bCachedRaytracingDataDirty, TEXT("Cached ray tracing instances must be up-to-date at this point"));
 
 					RelevantPrimitive->CachedRayTracingInstance = &CachedRayTracingInstance;
-					RelevantPrimitive->CachedMeshCommandFlags.bAnySegmentsDecal = SceneInfo->bCachedRayTracingInstanceAnySegmentsDecal;
-					RelevantPrimitive->CachedMeshCommandFlags.bAllSegmentsDecal = SceneInfo->bCachedRayTracingInstanceAllSegmentsDecal;
+					RelevantPrimitive->CachedMeshCommandFlags.bAnySegmentsDecal = RTLODData.CachedMeshCommandFlags.bAnySegmentsDecal;
+					RelevantPrimitive->CachedMeshCommandFlags.bAllSegmentsDecal = RTLODData.CachedMeshCommandFlags.bAllSegmentsDecal;
 
 					// CacheInstances expects to have one ray tracing mesh command per BLAS segment.
 					// If that's not the case in the future, other logic such as NumCachedStaticVisibleMeshCommands calculation needs to be updated.
@@ -1325,7 +1347,7 @@ namespace RayTracing
 					// one containing non-decal segments and the other with decal segments
 					// masking of segments is done using "hidden" hitgroups
 					// TODO: Debug Visualization to highlight primitives using this?
-					const bool bNeedDecalInstance = !bNeedFarFieldInstance && RelevantPrimitive->CachedMeshCommandFlags.bAnySegmentsDecal && !ShouldExcludeDecals();
+					const bool bNeedDecalInstance = !bNeedFarFieldInstance && RelevantPrimitive->CachedMeshCommandFlags.bAnySegmentsDecal && !bExcludeDecals;
 					checkf(bNeedMainInstance || bNeedDecalInstance || bNeedFarFieldInstance, TEXT("FRelevantPrimitive is expected to have a main instance, decal instance or both.")); // TODO
 
 					// For now store offsets relative to current context
@@ -1396,7 +1418,7 @@ namespace RayTracing
 						if (RTLODData.SBTAllocation == nullptr)
 						{
 							// No SBT allocation if no valid segments (see logic FRayTracingShaderBindingTable::AllocateStaticRange)
-							ensure(RTLODData.CachedMeshCommandFlags.bAllSegmentsDecal && !(RTLODData.CachedMeshCommandFlags.bAnySegmentsDecal && !RayTracing::ShouldExcludeDecals()));
+							ensure(RTLODData.CachedMeshCommandFlags.bAllSegmentsDecal && !(RTLODData.CachedMeshCommandFlags.bAnySegmentsDecal && !bExcludeDecals));
 							return;
 						}
 						const bool bUsesLightingChannels = bLightingChannelsUsingAHS && PrimitiveRayTracingData.bUsesLightingChannels;
diff --git a/Engine/Source/Runtime/Renderer/Public/PrimitiveSceneInfo.h b/Engine/Source/Runtime/Renderer/Public/PrimitiveSceneInfo.h
--- a/Engine/Source/Runtime/Renderer/Public/PrimitiveSceneInfo.h
+++ b/Engine/Source/Runtime/Renderer/Public/PrimitiveSceneInfo.h
@@ -643,11 +643,28 @@ public:
 
 #if RHI_RAYTRACING
 	bool bIsRelevantToPathTracing : 1;
+	UE_DEPRECATED(5.7, "Use DynamicRayTracingInstanceCachedData or GetRayTracingLODData(...) instead.")
 	bool bCachedRayTracingInstanceMaskAndFlagsDirty : 1;
+	UE_DEPRECATED(5.7, "Use DynamicRayTracingInstanceCachedData or GetRayTracingLODData(...) instead.")
 	bool bCachedRayTracingInstanceAnySegmentsDecal : 1;
+	UE_DEPRECATED(5.7, "Use DynamicRayTracingInstanceCachedData or GetRayTracingLODData(...) instead.")
 	bool bCachedRayTracingInstanceAllSegmentsDecal : 1;
+	UE_DEPRECATED(5.7, "Use DynamicRayTracingInstanceCachedData or GetRayTracingLODData(...) instead.")
 	bool bCachedRayTracingInstanceAllSegmentsTranslucent : 1;
 
+	struct FDynamicRayTracingInstanceCachedData
+	{
+		uint8 Mask = 0;
+		uint8 bForceOpaque : 1 = false;
+		uint8 bDoubleSided : 1 = false;
+		uint8 bReverseCulling : 1 = false;
+		uint8 bAnySegmentsDecal : 1 = false;
+		uint8 bAllSegmentsDecal : 1 = false;
+		uint8 bAllSegmentsTranslucent : 1 = false;
+	};
+	bool bDynamicRayTracingInstanceCachedDataDirty : 1 = true;
+	FDynamicRayTracingInstanceCachedData DynamicRayTracingInstanceCachedData;
+
 	struct FRayTracingLODData
 	{
 		FRayTracingSBTAllocation* SBTAllocation = nullptr;
@@ -682,8 +699,10 @@ public:
 		RayTracingLODData = MoveTemp(LODData);
 	}
 
-	void SetCachedRayTracingInstanceGeometryRHI(FRHIRayTracingGeometry* Geometry);
-	void UpdateCachedRayTracingInstanceMaskAndFlags(FRayTracingMaskAndFlags& InstanceMaskAndFlags);
+	void SetCachedRayTracingInstanceGeometryRHI(FRHIRayTracingGeometry* GeometryRHI);
+
+	UE_DEPRECATED(5.7, "No longer necessary.")
+	void UpdateCachedRayTracingInstanceMaskAndFlags(FRayTracingMaskAndFlags& InstanceMaskAndFlags) {}
 
 private:
 
