[Substrate] Fix Gbuffer Picking to work with Substrate when BlendableGbuffer is enabled.

diff --git a/Engine/Source/Runtime/Renderer/Private/PostProcess/PostProcessing.cpp b/Engine/Source/Runtime/Renderer/Private/PostProcess/PostProcessing.cpp
--- a/Engine/Source/Runtime/Renderer/Private/PostProcess/PostProcessing.cpp
+++ b/Engine/Source/Runtime/Renderer/Private/PostProcess/PostProcessing.cpp
@@ -3824,7 +3824,7 @@ public:
 			ShaderPrint::IsSupported(Platform) && 
 			!IsHlslccShaderPlatform(Platform) &&
 			!IsMobilePlatform(Platform) && 
-			!Substrate::IsSubstrateEnabled();
+			(!Substrate::IsSubstrateEnabled() || Substrate::IsSubstrateBlendableGBufferEnabled(Platform));
 	}
 	static bool ShouldCompilePermutation(const FGlobalShaderPermutationParameters& Parameters) 
 	{ 
