Modifications made to how objects are culled in SceneVisibility.cpp

Addition of cvar r.ExpandAllOcclusionTestedBBoxesScreenSpace and r.ExpandNewlyOcclusionTestedBBoxesScreenSpace, where the size of a bounding box will increase in screen space amounts, rather than in world space amounts.

Extension of existing r.VisualizeOccludedPrimitives to include value (2) which will depict the expanded bboxes of occluded objects.

Addition of r.VisualizePrimitiveBBoxes, which visualizes all bboxes and which has the option of depicting only expanded bboxes or both before and after expansion

diff --git a/Engine/Source/Runtime/Renderer/Private/SceneVisibility.cpp b/Engine/Source/Runtime/Renderer/Private/SceneVisibility.cpp
--- a/Engine/Source/Runtime/Renderer/Private/SceneVisibility.cpp
+++ b/Engine/Source/Runtime/Renderer/Private/SceneVisibility.cpp
@@ -139,9 +139,23 @@ static int32 GVisualizeOccludedPrimitives = 0;
 static FAutoConsoleVariableRef CVarVisualizeOccludedPrimitives(
 	TEXT("r.VisualizeOccludedPrimitives"),
 	GVisualizeOccludedPrimitives,
-	TEXT("Draw boxes for all occluded primitives"),
+	TEXT("Draw boxes for all occluded primitives.\n")
+	TEXT(" 0: Do not draw bounding boxes for occluded primitives (default).\n")
+	TEXT(" 1: Draw the primitive bounding boxes for all occluded primitives.\n")
+	TEXT(" 2: Draw the primitive bounding boxes for all occluded primitives after expansion. See r.Expand*BBoxes* for expansion.\n"),
 	ECVF_RenderThreadSafe | ECVF_Cheat
-	);
+);
+
+static int32 GVisualizePrimitiveBBoxes = 0;
+static FAutoConsoleVariableRef CVarVisualizePrimitiveBBoxes(
+	TEXT("r.VisualizePrimitiveBBoxes"),
+	GVisualizePrimitiveBBoxes,
+	TEXT("Draw boxes for all primitives.\n")
+	TEXT(" 0: Do not draw bounding boxes for all primitives (default).\n")
+	TEXT(" 1: Draw bounding boxes only after expansion. See r.Expand*BBoxes for expansion.\n")
+	TEXT(" 2: Draw bounding boxes for both before expansion and after expansion."),
+	ECVF_RenderThreadSafe | ECVF_Cheat
+);
 
 static int32 GAllowSubPrimitiveQueries = 1;
 static FAutoConsoleVariableRef CVarAllowSubPrimitiveQueries(
@@ -222,7 +236,31 @@ static float GExpandAllTestedBBoxesAmount = 0.0f;
 static FAutoConsoleVariableRef CVarExpandAllTestedBBoxesAmount(
 	TEXT("r.ExpandAllOcclusionTestedBBoxesAmount"),
 	GExpandAllTestedBBoxesAmount,
-	TEXT("Amount to expand all occlusion test bounds by."),
+	TEXT("Amount to expand all occlusion test bounds by.\n")
+	TEXT("This variable can be applied in conjunction with r.ExpandNewlyOcclusionTestedBBoxesScreenSpace and r.ExpandAllOcclusionTestedBBoxesScreenSpace. This variable always applies before those two.\n"),
+	ECVF_RenderThreadSafe
+);
+
+// TODO: UE-273143 This is the start of a significant refactor in the way that objects get culled.
+// The previous implementation took cvars to determine how much all bounding boxes should expand in world space.
+// The current implementation extends this to expand by a screen space amount (in pixels).
+// The idealized ultimate solution would take into account camera positional and/or angular velocity, over the previous n frames,
+// and extrapolate to determine the size at which bboxes should expand (in the event that they are likely to become unoccluded in the next m frames).
+
+static float GExpandNewlyOcclusionTestedBBoxesScreenSpace = 0.0f;
+static FAutoConsoleVariableRef CVarExpandNewlyOcclusionTestedBBoxesScreenSpace(
+	TEXT("r.ExpandNewlyOcclusionTestedBBoxesScreenSpace"),
+	GExpandNewlyOcclusionTestedBBoxesScreenSpace,
+	TEXT("If we don't occlusion test a primitive for r.GFramesNotOcclusionTestedToExpandBBoxes frames, then we expand the BBox when we do occlusion test it for a few frames by this amount. See also r.FramesToExpandNewlyOcclusionTestedBBoxes, r.GFramesNotOcclusionTestedToExpandBBoxes."),
+	ECVF_RenderThreadSafe
+);
+
+static float GExpandAllTestedBBoxesScreenSpace = 0.0f;
+static FAutoConsoleVariableRef CVarExpandAllTestedBBoxesScreenSpace(
+	TEXT("r.ExpandAllOcclusionTestedBBoxesScreenSpace"),
+	GExpandAllTestedBBoxesScreenSpace,
+	TEXT("Screen space amount in pixels to increase the occlusion test of BBoxes by.\n")
+	TEXT("This variable can be applied in conjunction with r.ExpandNewlyOcclusionTestedBBoxesAmount and r.ExpandAllOcclusionTestedBBoxesAmount. This variable always applies after those two.\n"),
 	ECVF_RenderThreadSafe
 );
 
@@ -2584,7 +2622,9 @@ FGPUOcclusionPacket::FGPUOcclusionPacket(FVisibilityViewPacket& InViewPacket, co
 	, CurrentRealTime(View.Family->Time.GetRealTimeSeconds())
 	, NeverOcclusionTestDistanceSquared(GNeverOcclusionTestDistance* GNeverOcclusionTestDistance)
 	, bUseOcclusionFeedback(View.FeatureLevel == ERHIFeatureLevel::ES3_1 && OcclusionFeedback.IsInitialized())
-	, bNewlyConsideredBBoxExpandActive(GExpandNewlyOcclusionTestedBBoxesAmount > 0.0f && GFramesToExpandNewlyOcclusionTestedBBoxes > 0 && GFramesNotOcclusionTestedToExpandBBoxes > 0)
+	, bNewlyConsideredBBoxExpandActive(
+		GExpandNewlyOcclusionTestedBBoxesAmount > 0.0f && GExpandNewlyOcclusionTestedBBoxesScreenSpace > 0.0f &&
+		GFramesToExpandNewlyOcclusionTestedBBoxes > 0 && GFramesNotOcclusionTestedToExpandBBoxes > 0)
 {}
 
 template <bool bIsParallel, typename VisitorType>
@@ -2641,6 +2681,26 @@ bool FGPUOcclusionPacket::OcclusionCullPrimitive(VisitorType& Visitor, FOcclusio
 	bool bAllSubOccluded = true;
 	FPrimitiveComponentId PrimitiveId = Scene.PrimitiveComponentIds[Index];
 
+	// Presolve for screen space bbox expansion
+	FVector4 ExpansionWorldSpaceNewlyConsidered;
+	FVector4 ExpansionWorldSpaceNotNewlyConsidered;
+	FVector4 ViewProjectionHomogeneousVector;
+	if (GExpandAllTestedBBoxesScreenSpace > 0.0f || GExpandNewlyOcclusionTestedBBoxesScreenSpace > 0.0f)
+	{
+		const FMatrix& InvViewProjectionMatrix = View.ViewMatrices.GetInvViewProjectionMatrix();
+
+		float ExtentNewly = GExpandAllTestedBBoxesScreenSpace + GExpandNewlyOcclusionTestedBBoxesScreenSpace;
+		FVector4 ExpansionScreenSpaceNewlyConsidered = FVector4(ExtentNewly / View.ViewRect.Width(), ExtentNewly / View.ViewRect.Height(), 0.0f, 0.0f);
+		ExpansionWorldSpaceNewlyConsidered = InvViewProjectionMatrix.TransformFVector4(ExpansionScreenSpaceNewlyConsidered);
+
+		float ExtentNotNewly = GExpandAllTestedBBoxesScreenSpace;
+		FVector4 ExpansionScreenSpaceNotNewlyConsidered = FVector4(ExtentNotNewly / View.ViewRect.Width(), ExtentNotNewly / View.ViewRect.Height(), 0.0f, 0.0f);
+		ExpansionWorldSpaceNotNewlyConsidered = InvViewProjectionMatrix.TransformFVector4(ExpansionScreenSpaceNotNewlyConsidered);
+
+		const FMatrix& ViewProjectionMatrix = View.ViewMatrices.GetViewProjectionMatrix();
+		ViewProjectionHomogeneousVector = FVector4(ViewProjectionMatrix.M[0][3], ViewProjectionMatrix.M[1][3], ViewProjectionMatrix.M[2][3], ViewProjectionMatrix.M[3][3]);
+	}
+
 	for (int32 SubQuery = 0; SubQuery < NumSubQueries; SubQuery++)
 	{
 		FPrimitiveOcclusionHistory* PrimitiveOcclusionHistory = PrimitiveOcclusionHistorySet.Find(FPrimitiveOcclusionHistoryKey(PrimitiveId, SubQuery));
@@ -2737,10 +2797,11 @@ bool FGPUOcclusionPacket::OcclusionCullPrimitive(VisitorType& Visitor, FOcclusio
 					}
 				}
 
-				if (GVisualizeOccludedPrimitives && bIsOccluded)
+				if (GVisualizeOccludedPrimitives == 1 && bIsOccluded)
 				{
 					const FBoxSphereBounds& Bounds = bSubQueries ? (*SubBounds)[SubQuery] : Scene.PrimitiveOcclusionBounds[Index];
-					Visitor.AddVisualizeQuery(Bounds.GetBox());
+					FViewElementPDI PDI(&View, nullptr, &View.DynamicPrimitiveCollector);
+					DrawWireBox(&PDI, Bounds.GetBox(), FColor(50, 255, 50), SDPG_Foreground);
 				}
 			}
 		}
@@ -2764,8 +2825,51 @@ bool FGPUOcclusionPacket::OcclusionCullPrimitive(VisitorType& Visitor, FOcclusio
 				}
 			}
 
+			const FBoxSphereBounds PrimitiveOcclusionBounds = (bSubQueries ? (*SubBounds)[SubQuery] : Scene.PrimitiveOcclusionBounds[Index]);
+
+			// Print unexpanded bounding boxes in green if not already visualized by GVisualizeOccludedPrimitives
+			if (GVisualizePrimitiveBBoxes == 2 && !(GVisualizeOccludedPrimitives == 1 && bIsOccluded))
+			{
+				FViewElementPDI PDI(&View, nullptr, &View.DynamicPrimitiveCollector);
+				DrawWireBox(&PDI, PrimitiveOcclusionBounds.GetBox(), FColor(0, 155, 0), SDPG_Foreground);
+			}
+
+			// scale occlusion bounds by GExpand*TestedBBoxesAmount in world space
+			FBoxSphereBounds WorldSpaceExpandedOcclusionBounds = PrimitiveOcclusionBounds;
+			if (GExpandAllTestedBBoxesAmount > 0.0f || GExpandNewlyOcclusionTestedBBoxesAmount > 0.0f)
+			{
+				WorldSpaceExpandedOcclusionBounds = WorldSpaceExpandedOcclusionBounds.ExpandBy(GExpandAllTestedBBoxesAmount + (bSkipNewlyConsidered ? GExpandNewlyOcclusionTestedBBoxesAmount : 0.0));
+			}
+			
+			// scale occlusion bounds by GExpand*TestedBBoxesScreenSpace in screen space
+			FBoxSphereBounds FinalOcclusionBounds = WorldSpaceExpandedOcclusionBounds;
+			if (GExpandAllTestedBBoxesScreenSpace > 0.0f || GExpandNewlyOcclusionTestedBBoxesScreenSpace > 0.0f)
+			{
+				FVector ExpansionWorldSpace = (bSkipNewlyConsidered) ? ExpansionWorldSpaceNewlyConsidered : ExpansionWorldSpaceNotNewlyConsidered;
+
+				const float BoundsScreenSpace_W = FVector::DotProduct(FinalOcclusionBounds.Origin, ViewProjectionHomogeneousVector) + ViewProjectionHomogeneousVector.W;
+
+				ExpansionWorldSpace *= BoundsScreenSpace_W;
+				ExpansionWorldSpace = FVector(abs(ExpansionWorldSpace.X), abs(ExpansionWorldSpace.Y), abs(ExpansionWorldSpace.Z));
+
+				FinalOcclusionBounds.BoxExtent = FinalOcclusionBounds.BoxExtent + ExpansionWorldSpace;
+				FinalOcclusionBounds.SphereRadius = FinalOcclusionBounds.BoxExtent.Length();
+			}
+
+			const FBoxSphereBounds OcclusionBounds = FinalOcclusionBounds;
+
+			if (GVisualizeOccludedPrimitives == 2 && bIsOccluded)
+			{
+				FViewElementPDI PDI(&View, nullptr, &View.DynamicPrimitiveCollector);
+				DrawWireBox(&PDI, OcclusionBounds.GetBox(), FColor(50, 255, 50), SDPG_Foreground);
+			}
+			else if (GVisualizePrimitiveBBoxes > 0)
+			{
+				FViewElementPDI PDI(&View, nullptr, &View.DynamicPrimitiveCollector);
+				DrawWireBox(&PDI, OcclusionBounds.GetBox(), FColor(155, 0, 0), SDPG_Foreground);
+			}
+
 			bool bAllowBoundsTest;
-			const FBoxSphereBounds OcclusionBounds = (bSubQueries ? (*SubBounds)[SubQuery] : Scene.PrimitiveOcclusionBounds[Index]).ExpandBy(GExpandAllTestedBBoxesAmount + (bSkipNewlyConsidered ? GExpandNewlyOcclusionTestedBBoxesAmount : 0.0));
 			if (FVector::DistSquared(ViewOrigin, OcclusionBounds.Origin) < NeverOcclusionTestDistanceSquared)
 			{
 				bAllowBoundsTest = false;
@@ -3393,7 +3497,7 @@ FVisibilityTaskConfig::FVisibilityTaskConfig(const FScene& Scene, TConstArrayVie
 
 	if (Schedule == EVisibilityTaskSchedule::Parallel)
 	{
-		if (!FApp::ShouldUseThreadingForPerformance() || !GIsThreadedRendering || !GSupportsParallelOcclusionQueries || GVisualizeOccludedPrimitives > 0 || IsMobilePlatform(Scene.GetShaderPlatform()))
+		if (!FApp::ShouldUseThreadingForPerformance() || !GIsThreadedRendering || !GSupportsParallelOcclusionQueries || GVisualizeOccludedPrimitives > 0 || GVisualizePrimitiveBBoxes > 0 || IsMobilePlatform(Scene.GetShaderPlatform()))
 		{
 			Schedule = EVisibilityTaskSchedule::RenderThread;
 		}
