Ensure VRS supports EDisplayOutputFormat::SDR_Rec709 and  EDisplayOutputFormat::SDR_ExplicitGammaMapping.
Setting r.TonemapperGamma > 0.0 forces the viewfamily output format to  EDisplayOutputFormat::SDR_ExplicitGammaMapping, which would previously cause VRS to be silently disabled.

diff --git a/Engine/Source/Runtime/Renderer/Private/VariableRateShading/ContrastAdaptiveImageGenerator.cpp b/Engine/Source/Runtime/Renderer/Private/VariableRateShading/ContrastAdaptiveImageGenerator.cpp
--- a/Engine/Source/Runtime/Renderer/Private/VariableRateShading/ContrastAdaptiveImageGenerator.cpp
+++ b/Engine/Source/Runtime/Renderer/Private/VariableRateShading/ContrastAdaptiveImageGenerator.cpp
@@ -369,12 +369,6 @@ static EDisplayOutputFormat GetDisplayOutputFormat(const FSceneView& View)
 	return (EDisplayOutputFormat)Parameters.OutputDevice;
 }
 
-static bool IsHDR10(const EDisplayOutputFormat& OutputFormat)
-{
-	return OutputFormat == EDisplayOutputFormat::HDR_ACES_1000nit_ST2084 ||
-		OutputFormat == EDisplayOutputFormat::HDR_ACES_2000nit_ST2084;
-}
-
 void AddCreateShadingRateImagePass(
 	FRDGBuilder& GraphBuilder,
 	const FSceneViewFamily& ViewFamily,
@@ -401,14 +395,20 @@ void AddCreateShadingRateImagePass(
 
 			TShaderMapRef<FCalculateShadingRateImageCS> ComputeShader(ViewInfo.ShaderMap, PermutationVector);
 			auto* PassParameters = GraphBuilder.AllocParameters<FCalculateShadingRateImageCS::FParameters>();
-			EDisplayOutputFormat OutputDisplayFormat = GetDisplayOutputFormat(ViewInfo);
+
+			const FTonemapperOutputDeviceParameters OutputDeviceParameters = GetTonemapperOutputDeviceParameters(*ViewInfo.Family);
+			const EDisplayOutputFormat OutputFormat = (EDisplayOutputFormat) OutputDeviceParameters.OutputDevice;
+			const bool bIsHDR10 = OutputFormat == EDisplayOutputFormat::HDR_ACES_1000nit_ST2084 
+				|| OutputFormat == EDisplayOutputFormat::HDR_ACES_2000nit_ST2084;
+
 			FCalculateShadingRateImageCS::InitParameters(
 				*PassParameters,
 				LuminanceTexture,
 				LuminanceRect,
-				IsHDR10(OutputDisplayFormat),
+				bIsHDR10,
 				bCreateHardwareImages ? GraphBuilder.CreateUAV(ImageData.HardwareImages.ConstructedSRIArray[Index]) : nullptr,
 				bCreateSoftwareImages ? GraphBuilder.CreateUAV(ImageData.SoftwareImages.ConstructedSRIArray[Index]) : nullptr);
+
 			FComputeShaderUtils::AddPass(
 				GraphBuilder,
 				RDG_EVENT_NAME("CreateShadingRateImage"),
@@ -540,9 +540,8 @@ bool FContrastAdaptiveImageGenerator::IsEnabled() const
 
 bool FContrastAdaptiveImageGenerator::IsSupportedByView(const FSceneView& View) const
 {
-	EDisplayOutputFormat DisplayOutputFormat = GetDisplayOutputFormat(View);
-	const bool bCompatibleWithOutputType = (DisplayOutputFormat == EDisplayOutputFormat::SDR_sRGB) || IsHDR10(DisplayOutputFormat);
-	return !View.bIsSceneCapture && bCompatibleWithOutputType && !IsMobilePlatform(View.GetShaderPlatform());
+	// See also the requirments under FVariableRateShadingImageManager::IsVRSCompatibleWithView
+	return !IsMobilePlatform(View.GetShaderPlatform());
 }
 
 FRDGTextureRef FContrastAdaptiveImageGenerator::GetDebugImage(FRDGBuilder& GraphBuilder, const FViewInfo& ViewInfo, FVariableRateShadingImageManager::EVRSImageType ImageType, bool bGetSoftwareImage)
diff --git a/Engine/Source/Runtime/Renderer/Private/VariableRateShading/VariableRateShadingImageManager.cpp b/Engine/Source/Runtime/Renderer/Private/VariableRateShading/VariableRateShadingImageManager.cpp
--- a/Engine/Source/Runtime/Renderer/Private/VariableRateShading/VariableRateShadingImageManager.cpp
+++ b/Engine/Source/Runtime/Renderer/Private/VariableRateShading/VariableRateShadingImageManager.cpp
@@ -397,28 +397,22 @@ void FVariableRateShadingImageManager::SetNeedStaticMeshUpdate(bool bInNeedStati
 	bNeedStaticMeshUpdate = bInNeedStaticMeshUpdate;
 }
 
-static EDisplayOutputFormat GetDisplayOutputFormat(const FViewInfo& View)
-{
-	FTonemapperOutputDeviceParameters Parameters = GetTonemapperOutputDeviceParameters(*View.Family);
-	return (EDisplayOutputFormat)Parameters.OutputDevice;
-}
-
 bool FVariableRateShadingImageManager::IsVRSCompatibleWithOutputType(const EDisplayOutputFormat& OutputFormat)
 {
-	// VRS texture generation is currently only compatible with SDR and HDR10
 	return OutputFormat == EDisplayOutputFormat::SDR_sRGB
+		|| OutputFormat == EDisplayOutputFormat::SDR_Rec709
+		|| OutputFormat == EDisplayOutputFormat::SDR_ExplicitGammaMapping
 		|| OutputFormat == EDisplayOutputFormat::HDR_ACES_1000nit_ST2084
 		|| OutputFormat == EDisplayOutputFormat::HDR_ACES_2000nit_ST2084;
 }
 
 bool FVariableRateShadingImageManager::IsVRSCompatibleWithView(const FViewInfo& ViewInfo)
 {
-	// TODO: Investigate if it's worthwhile getting scene captures working. Things that we'll need to take care of
-	// is to associate shading rate texture image with main scene, and scene capture.  But what if there is
-	// more than 1 scene capture?  Is there a unique identifier that connects two frames of scene capture.
+	FTonemapperOutputDeviceParameters Parameters = GetTonemapperOutputDeviceParameters(*ViewInfo.Family);
+
 	return !ViewInfo.bIsSceneCapture
 		&& ViewInfo.Family->bRealtimeUpdate
-		&& IsVRSCompatibleWithOutputType(GetDisplayOutputFormat(ViewInfo));
+		&& IsVRSCompatibleWithOutputType((EDisplayOutputFormat)Parameters.OutputDevice);
 }
 
 FIntPoint FVariableRateShadingImageManager::GetSRITileSize(bool bSoftwareVRS)
