Moved mesh pass task wait to happen in OnRenderFinish, which is just prior to RDG execution. This avoids the need to wait when building render commands with instance culling, and avoids sync issues if instance culling is skipped.

diff --git a/Engine/Source/Runtime/Renderer/Private/SceneRenderBuilder.cpp b/Engine/Source/Runtime/Renderer/Private/SceneRenderBuilder.cpp
--- a/Engine/Source/Runtime/Renderer/Private/SceneRenderBuilder.cpp
+++ b/Engine/Source/Runtime/Renderer/Private/SceneRenderBuilder.cpp
@@ -728,54 +728,15 @@ void FSceneRenderProcessor::Cleanup(FRHICommandListImmediate& RHICmdList, FScene
 	if (SceneRenderCleanUpMode == ESceneRenderCleanUpMode::Async)
 	{
 		Renderer->GPUSceneDynamicContext.Release();
-	}
-	else // Immediate
-	{
-		WaitForAsyncDeleteTask(); // This is to handle cases where as switch from async to immediate.
-		RHICmdList.ImmediateFlush(EImmediateFlushType::WaitForOutstandingTasksOnly);
-	}
 
-	FGraphEventArray TasksToWait;
-
-	// Wait on all setup tasks now to ensure that no additional render commands are enqueued which
-	// might mess with render state, since setup tasks are working with high-level render objects.
-	for (FParallelMeshDrawCommandPass* DispatchedShadowDepthPass : Renderer->DispatchedShadowDepthPasses)
-	{
-		if (FGraphEventRef Event = DispatchedShadowDepthPass->AcquireTaskEvent())
+		// Wait for the last renderer's cleanup tasks so that snapshot deallocation and destruction don't overlap.
+		if (AsyncTasks.Cleanup)
 		{
-			TasksToWait.Emplace(MoveTemp(Event));
+			AsyncTasks.Cleanup->Wait(ENamedThreads::GetRenderThread_Local());
 		}
-	}
 
-	for (const FViewInfo* View : Renderer->AllViews)
-	{
-		for (FParallelMeshDrawCommandPass* Pass : View->ParallelMeshDrawCommandPasses)
-		{
-			if (Pass)
-			{
-				if (FGraphEventRef Event = Pass->AcquireTaskEvent())
-				{
-					TasksToWait.Emplace(MoveTemp(Event));
-				}
-			}
-		}
-	}
-
-	// Wait for the last renderer's cleanup tasks so that snapshot deallocation and destruction don't overlap.
-	if (AsyncTasks.Cleanup)
-	{
-		TasksToWait.Emplace(MoveTemp(AsyncTasks.Cleanup));
-	}
-
-	if (!TasksToWait.IsEmpty())
-	{
-		FTaskGraphInterface::Get().WaitUntilTasksComplete(TasksToWait, ENamedThreads::GetRenderThread_Local());
-	}
+		ViewSnapshotCache::Deallocate();
 
-	ViewSnapshotCache::Deallocate();
-
-	if (SceneRenderCleanUpMode == ESceneRenderCleanUpMode::Async)
-	{
 		AsyncTasks.Cleanup = FFunctionGraphTask::CreateAndDispatchWhenReady([Renderer]
 		{
 			CleanupSceneRenderer(Renderer);
@@ -784,6 +745,9 @@ void FSceneRenderProcessor::Cleanup(FRHICommandListImmediate& RHICmdList, FScene
 	}
 	else // Immediate
 	{
+		WaitForAsyncDeleteTask(); // This is to handle cases where as switch from async to immediate.
+		RHICmdList.ImmediateFlush(EImmediateFlushType::WaitForOutstandingTasksOnly);
+		ViewSnapshotCache::Deallocate();
 		CleanupSceneRenderer(Renderer);
 	}
 
diff --git a/Engine/Source/Runtime/Renderer/Private/SceneRendering.cpp b/Engine/Source/Runtime/Renderer/Private/SceneRendering.cpp
--- a/Engine/Source/Runtime/Renderer/Private/SceneRendering.cpp
+++ b/Engine/Source/Runtime/Renderer/Private/SceneRendering.cpp
@@ -4639,6 +4639,33 @@ void FSceneRenderer::OnRenderFinish(FRDGBuilder& GraphBuilder, FRDGTextureRef Vi
 			Scene->DumpMeshDrawCommandMemoryStats();
 		}
 	});
+
+	// Sync all mesh pass setup tasks prior to starting RDG execution.
+	FGraphEventArray TasksToWait;
+
+	for (FParallelMeshDrawCommandPass* DispatchedShadowDepthPass : DispatchedShadowDepthPasses)
+	{
+		if (FGraphEventRef Event = DispatchedShadowDepthPass->AcquireTaskEvent())
+		{
+			TasksToWait.Emplace(MoveTemp(Event));
+		}
+	}
+
+	for (const FViewInfo* View : AllViews)
+	{
+		for (FParallelMeshDrawCommandPass* Pass : View->ParallelMeshDrawCommandPasses)
+		{
+			if (Pass)
+			{
+				if (FGraphEventRef Event = Pass->AcquireTaskEvent())
+				{
+					TasksToWait.Emplace(MoveTemp(Event));
+				}
+			}
+		}
+	}
+
+	FTaskGraphInterface::Get().WaitUntilTasksComplete(TasksToWait, ENamedThreads::GetRenderThread_Local());
 }
 
 void FSceneRenderer::SetupMeshPass(FViewInfo& View, FExclusiveDepthStencil::Type BasePassDepthStencilAccess, FViewCommands& ViewCommands, FInstanceCullingManager& InstanceCullingManager)
