Material Cache, fixed crash on dangling virtual resource ptr
- Instead pool through the scene extension, safe lifetimes

diff --git a/Engine/Source/Runtime/Engine/Internal/MaterialCache/MaterialCacheVirtualTextureInterface.h b/Engine/Source/Runtime/Engine/Internal/MaterialCache/MaterialCacheVirtualTextureInterface.h
--- a/Engine/Source/Runtime/Engine/Internal/MaterialCache/MaterialCacheVirtualTextureInterface.h
+++ b/Engine/Source/Runtime/Engine/Internal/MaterialCache/MaterialCacheVirtualTextureInterface.h
@@ -5,6 +5,8 @@
 #include "Engine/Texture.h"
 #include "MaterialCacheVirtualTextureInterface.generated.h"
 
+class FSceneInterface;
+
 class FMaterialCacheVirtualTextureRenderProxyInterface
 {
 public:
@@ -18,7 +20,7 @@ public:
 	}
 
 	/** Flush all pages */
-	virtual void Flush()
+	virtual void Flush(FSceneInterface* Scene)
 	{
 		checkNoEntry();
 	}
diff --git a/Engine/Source/Runtime/Renderer/Internal/MaterialCache/MaterialCacheVirtualTextureRenderProxy.h b/Engine/Source/Runtime/Renderer/Internal/MaterialCache/MaterialCacheVirtualTextureRenderProxy.h
--- a/Engine/Source/Runtime/Renderer/Internal/MaterialCache/MaterialCacheVirtualTextureRenderProxy.h
+++ b/Engine/Source/Runtime/Renderer/Internal/MaterialCache/MaterialCacheVirtualTextureRenderProxy.h
@@ -16,10 +16,10 @@ public: /** UMaterialCacheVirtualTextureRenderProxyInterface */
 		return TextureDescriptor;
 	}
 
-	virtual void Flush() override;
+	virtual void Flush(FSceneInterface* Scene) override;
 
-	/** Virtual texture resource */
-	FVirtualTexture2DResource* TextureResource = nullptr;
+	/** Component owning the virtual texture */
+	FPrimitiveComponentId PrimitiveCID{};
 
 	/** Optional, owned stack provider render proxy */
 	TUniquePtr<FMaterialCacheStackProviderRenderProxy> StackProviderRenderProxy;
diff --git a/Engine/Source/Runtime/Renderer/Private/MaterialCache/MaterialCacheVirtualTexture.cpp b/Engine/Source/Runtime/Renderer/Private/MaterialCache/MaterialCacheVirtualTexture.cpp
--- a/Engine/Source/Runtime/Renderer/Private/MaterialCache/MaterialCacheVirtualTexture.cpp
+++ b/Engine/Source/Runtime/Renderer/Private/MaterialCache/MaterialCacheVirtualTexture.cpp
@@ -301,7 +301,7 @@ FMaterialCacheVirtualTextureRenderProxyInterface* UMaterialCacheVirtualTexture::
 	}
 	
 	UMaterialCacheVirtualTextureRenderProxy* Proxy = new UMaterialCacheVirtualTextureRenderProxy();
-	Proxy->TextureResource   = VTResource;
+	Proxy->PrimitiveCID      = OwningComponent->GetSceneData().PrimitiveSceneId;
 	Proxy->UVCoordinateIndex = UVCoordinateIndex;
 
 	// Stack providers are optional
@@ -311,9 +311,9 @@ FMaterialCacheVirtualTextureRenderProxyInterface* UMaterialCacheVirtualTexture::
 	}
 
 	// Render thread initialization
-	ENQUEUE_RENDER_COMMAND(GetBackbufferFormatCmd)([Proxy](FRHICommandListImmediate&)
+	ENQUEUE_RENDER_COMMAND(GetBackbufferFormatCmd)([Proxy, VTResource](FRHICommandListImmediate&)
 	{
-		Proxy->TextureDescriptor = PackMaterialCacheTextureDescriptor(Proxy->TextureResource, Proxy->UVCoordinateIndex);
+		Proxy->TextureDescriptor = PackMaterialCacheTextureDescriptor(VTResource, Proxy->UVCoordinateIndex);
 	});
 
 	return Proxy;
@@ -516,10 +516,23 @@ FScene* UMaterialCacheVirtualTexture::GetRenderScene()
 	return Scene->GetRenderScene();
 }
 
-void UMaterialCacheVirtualTextureRenderProxy::Flush()
+void UMaterialCacheVirtualTextureRenderProxy::Flush(FSceneInterface* Scene)
 {
-	if (IAllocatedVirtualTexture* AllocatedVT = TextureResource->GetAllocatedVT())
+	// May not exist if headless
+	FScene* RenderScene = Scene->GetRenderScene();
+	if (!RenderScene)
+	{
+		return;
+	}
+
+	FMaterialCacheSceneExtension& SceneExtension = RenderScene->SceneExtensions.GetExtension<FMaterialCacheSceneExtension>();
+
+	// Ignore missing registrations, not fatal
+	if (FMaterialCachePrimitiveData* Data = SceneExtension.GetPrimitiveData(PrimitiveCID))
 	{
-		GetRendererModule().FlushVirtualTextureCache(AllocatedVT, FVector2f(0, 0), FVector2f(1, 1));
+		if (IAllocatedVirtualTexture* AllocatedVT = Data->Provider.Texture)
+		{
+			GetRendererModule().FlushVirtualTextureCache(AllocatedVT, FVector2f(0, 0), FVector2f(1, 1));
+		}
 	}
 }
diff --git a/Engine/Source/Runtime/Renderer/Private/PrimitiveSceneInfo.cpp b/Engine/Source/Runtime/Renderer/Private/PrimitiveSceneInfo.cpp
--- a/Engine/Source/Runtime/Renderer/Private/PrimitiveSceneInfo.cpp
+++ b/Engine/Source/Runtime/Renderer/Private/PrimitiveSceneInfo.cpp
@@ -2082,7 +2082,7 @@ void FPrimitiveSceneInfo::FlushRuntimeVirtualTexture()
 
 	if (FMaterialCacheVirtualTextureRenderProxyInterface* MaterialCacheProxy = Proxy->GetMaterialCacheRenderProxy())
 	{
-		MaterialCacheProxy->Flush();
+		MaterialCacheProxy->Flush(Scene);
 	}
 }
 
