Added support for SceneColor sampling in translucent materials on the mobile renderer using deferred shading

#rnx

diff --git a/Engine/Shaders/Private/MaterialTemplate.ush b/Engine/Shaders/Private/MaterialTemplate.ush
--- a/Engine/Shaders/Private/MaterialTemplate.ush
+++ b/Engine/Shaders/Private/MaterialTemplate.ush
@@ -3372,8 +3372,7 @@ float4 SceneTextureFetchFunc(FMaterialPixelParameters Parameters, int SceneTextu
 }
 #endif // NEEDS_SCENE_TEXTURES
 
-#if SHADING_PATH_DEFERRED
-
+#if SHADING_PATH_DEFERRED || MOBILE_DEFERRED_SHADING
 
 /** Applies an offset to the scene texture lookup and decodes the HDR linear space color and alpha. */
 float4 DecodeSceneColorAndAlpharForMaterialNode(float2 ScreenUV)
@@ -3399,7 +3398,7 @@ float3 DecodeSceneColorForMaterialNode(float2 ScreenUV)
 	return DecodeSceneColorAndAlpharForMaterialNode(ScreenUV).rgb;
 }
 
-#endif // SHADING_PATH_DEFERRED
+#endif // SHADING_PATH_DEFERRED || MOBILE_DEFERRED_SHADING
 
 float4 MaterialExpressionDBufferTextureLookup(FMaterialPixelParameters Parameters, float2 BufferUV, int DBufferTextureIndex)
 {
diff --git a/Engine/Shaders/Private/MobileBasePassPixelShader.usf b/Engine/Shaders/Private/MobileBasePassPixelShader.usf
--- a/Engine/Shaders/Private/MobileBasePassPixelShader.usf
+++ b/Engine/Shaders/Private/MobileBasePassPixelShader.usf
@@ -89,7 +89,6 @@
 
 #if MATERIAL_SHADINGMODEL_THIN_TRANSLUCENT || SUBSTRATE_BLENDING_TRANSLUCENT_COLOREDTRANSMITTANCE
 	#define MOBILE_TRANSLUCENT_COLOR_TRANSMITTANCE 1
-	
 #else
 	#define MOBILE_TRANSLUCENT_COLOR_TRANSMITTANCE 0
 	#undef MOBILE_TRANSLUCENT_COLOR_TRANSMITTANCE_DUAL_SRC_BLENDING
@@ -100,6 +99,12 @@
 	#define MOBILE_TRANSLUCENT_COLOR_TRANSMITTANCE_SINGLE_SRC_BLENDING 0
 #endif
 
+
+#if MOBILE_DEFERRED_SHADING && (MATERIALBLENDING_TRANSLUCENT || MATERIALBLENDING_ADDITIVE || MATERIALBLENDING_MODULATE || MATERIAL_SHADINGMODEL_SINGLELAYERWATER)
+	#define SceneColorCopyTexture	MobileBasePass.SceneColorCopyTexture
+	#define SceneColorCopySampler	MobileBasePass.SceneColorCopySampler
+#endif
+
 // SceneDepthAux always enabled for IOS, on other platforms only enabled when in forward and MobileHDR=true
 // Sky does not write to depth
 #define HAS_SCENE_DEPTH_AUX_OUTPUT (USE_SCENE_DEPTH_AUX && (MATERIALBLENDING_SOLID || MATERIALBLENDING_MASKED) && !MATERIAL_SHADINGMODEL_SINGLELAYERWATER && !MATERIAL_IS_SKY)
diff --git a/Engine/Source/Runtime/Engine/Private/Materials/HLSLMaterialTranslator.cpp b/Engine/Source/Runtime/Engine/Private/Materials/HLSLMaterialTranslator.cpp
--- a/Engine/Source/Runtime/Engine/Private/Materials/HLSLMaterialTranslator.cpp
+++ b/Engine/Source/Runtime/Engine/Private/Materials/HLSLMaterialTranslator.cpp
@@ -8537,9 +8537,9 @@ int32 FHLSLMaterialTranslator::SceneColor(int32 Offset, int32 ViewportUV, bool b
 		Errorf(TEXT("SceneColor lookups are only available when MaterialDomain = Surface."));
 	}
 
-	if (ErrorUnlessFeatureLevelSupported(ERHIFeatureLevel::SM5) == INDEX_NONE)
+	if (IsMobilePlatform(Platform) && !MobileSupportsSM5MaterialNodes(Platform))
 	{
-		return INDEX_NONE;
+		Errorf(TEXT("Cannot read SceneColor with mobile shading (platform id %d)."), Platform);
 	}
 
 	MaterialCompilationOutput.SetIsSceneTextureUsed(PPI_SceneColor);
diff --git a/Engine/Source/Runtime/RenderCore/Private/Shader.cpp b/Engine/Source/Runtime/RenderCore/Private/Shader.cpp
--- a/Engine/Source/Runtime/RenderCore/Private/Shader.cpp
+++ b/Engine/Source/Runtime/RenderCore/Private/Shader.cpp
@@ -2587,6 +2587,12 @@ void ShaderMapAppendKey(EShaderPlatform Platform, FShaderKeyGenerator& KeyGen)
 					break;
 			}
 		}
+
+		if (MobileSupportsSM5MaterialNodes(Platform))
+		{
+			KeyGen.AppendSeparator();
+			KeyGen.Append(TEXT("MATMOBSM5"));
+		}
 	}
 
 	{
diff --git a/Engine/Source/Runtime/Renderer/Private/MobileBasePassRendering.cpp b/Engine/Source/Runtime/Renderer/Private/MobileBasePassRendering.cpp
--- a/Engine/Source/Runtime/Renderer/Private/MobileBasePassRendering.cpp
+++ b/Engine/Source/Runtime/Renderer/Private/MobileBasePassRendering.cpp
@@ -264,8 +264,8 @@ void SetupMobileBasePassUniformParameters(
 	const FViewInfo& View, 
 	EMobileBasePass BasePass,
 	EMobileSceneTextureSetupMode SetupMode,
-	const FMobileBasePassTextures& MobileBasePassTextures,
 	FMobileBasePassUniformParameters& BasePassParameters,
+	const FMobileBasePassTextures& MobileBasePassTextures,
 	bool bForRealtimeSkyCapture)
 {
 	const FViewInfo* InstancedView = View.GetInstancedView();
@@ -375,34 +375,37 @@ void SetupMobileBasePassUniformParameters(
 	SetupReflectionUniformParameters(GraphBuilder, View, BasePassParameters.ReflectionsParameters);
 
 	SetupMobileSSRParameters(GraphBuilder, View, BasePassParameters.SSRParams);
-}
 
-TRDGUniformBufferRef<FMobileBasePassUniformParameters> CreateMobileBasePassUniformBuffer(
-	FRDGBuilder& GraphBuilder,
-	const FViewInfo& View,
-	EMobileBasePass BasePass,
-	EMobileSceneTextureSetupMode SetupMode,
-	const FMobileBasePassTextures& MobileBasePassTextures,
-	bool bForRealtimeSkyCapture)
-{
-	FMobileBasePassUniformParameters* BasePassParameters = GraphBuilder.AllocParameters<FMobileBasePassUniformParameters>();
-	SetupMobileBasePassUniformParameters(GraphBuilder, View, BasePass, SetupMode, MobileBasePassTextures, *BasePassParameters, bForRealtimeSkyCapture);
+	BasePassParameters.SceneColorCopyTexture = SystemTextures.White;
+	BasePassParameters.SceneColorCopySampler = TStaticSamplerState<SF_Bilinear, AM_Clamp, AM_Clamp, AM_Clamp>::GetRHI();;
+
 #if WITH_DEBUG_VIEW_MODES
 	if (View.Family->UseDebugViewPS())
 	{
-		SetupDebugViewModePassUniformBufferConstants(View, BasePassParameters->DebugViewMode);
+		SetupDebugViewModePassUniformBufferConstants(View, BasePassParameters.DebugViewMode);
 	}
 #endif
 
 	// QuadOverdraw is a UAV so it needs to be initialized even if not used
-	const FSceneTextures* SceneTextures = View.GetSceneTexturesChecked();
 	FRDGTextureRef QuadOverdrawTexture = SceneTextures ? SceneTextures->QuadOverdraw : nullptr;
 	if (!QuadOverdrawTexture)
 	{
 		QuadOverdrawTexture = GraphBuilder.CreateTexture(FRDGTextureDesc::Create2D(FIntPoint(1, 1), PF_R32_UINT, FClearValueBinding::None, TexCreate_UAV), TEXT("DummyOverdrawUAV"));
 	}
-	BasePassParameters->QuadOverdraw = GraphBuilder.CreateUAV(QuadOverdrawTexture);
-	
+	BasePassParameters.QuadOverdraw = GraphBuilder.CreateUAV(QuadOverdrawTexture);
+}
+
+TRDGUniformBufferRef<FMobileBasePassUniformParameters> CreateMobileBasePassUniformBuffer(
+	FRDGBuilder& GraphBuilder,
+	const FViewInfo& View,
+	EMobileBasePass BasePass,
+	EMobileSceneTextureSetupMode SetupMode,
+	const FMobileBasePassTextures& MobileBasePassTextures,
+	bool bForRealtimeSkyCapture)
+{
+	FMobileBasePassUniformParameters* BasePassParameters = GraphBuilder.AllocParameters<FMobileBasePassUniformParameters>();
+	SetupMobileBasePassUniformParameters(GraphBuilder, View, BasePass, SetupMode, *BasePassParameters, MobileBasePassTextures, bForRealtimeSkyCapture);
+
 	return GraphBuilder.CreateUniformBuffer(BasePassParameters);
 }
 
diff --git a/Engine/Source/Runtime/Renderer/Private/MobileBasePassRendering.h b/Engine/Source/Runtime/Renderer/Private/MobileBasePassRendering.h
--- a/Engine/Source/Runtime/Renderer/Private/MobileBasePassRendering.h
+++ b/Engine/Source/Runtime/Renderer/Private/MobileBasePassRendering.h
@@ -65,6 +65,8 @@ BEGIN_GLOBAL_SHADER_PARAMETER_STRUCT(FMobileBasePassUniformParameters, )
 	SHADER_PARAMETER_SAMPLER(SamplerState, HalfResLocalFogVolumeViewSampler)
 	SHADER_PARAMETER_STRUCT_INCLUDE(FDBufferParameters, DBuffer)
 	SHADER_PARAMETER_STRUCT(FMobileScreenSpaceReflectionParams, SSRParams)
+	SHADER_PARAMETER_RDG_TEXTURE(Texture2D, SceneColorCopyTexture)
+	SHADER_PARAMETER_SAMPLER(SamplerState, SceneColorCopySampler)
 END_GLOBAL_SHADER_PARAMETER_STRUCT()
 
 enum class EMobileBasePass
@@ -78,9 +80,10 @@ extern void SetupMobileBasePassUniformParameters(
 	FRDGBuilder& GraphBuilder,
 	const FViewInfo& View,
 	EMobileBasePass BasePass,
-	const FMobileBasePassTextures& MobileBasePassTextures,
-	FMobileBasePassUniformParameters& BasePassParameters,
-	bool bForRealtimeSkyCapture);
+	EMobileSceneTextureSetupMode SetupMode,
+	FMobileBasePassUniformParameters& BasePassParameters, 
+	const FMobileBasePassTextures& MobileBasePassTextures = {},
+	bool bForRealtimeSkyCapture = false);
 
 extern TRDGUniformBufferRef<FMobileBasePassUniformParameters> CreateMobileBasePassUniformBuffer(
 	FRDGBuilder& GraphBuilder,
diff --git a/Engine/Source/Runtime/Renderer/Private/MobileShadingRenderer.cpp b/Engine/Source/Runtime/Renderer/Private/MobileShadingRenderer.cpp
--- a/Engine/Source/Runtime/Renderer/Private/MobileShadingRenderer.cpp
+++ b/Engine/Source/Runtime/Renderer/Private/MobileShadingRenderer.cpp
@@ -2422,7 +2422,11 @@ void FMobileSceneRenderer::RenderDeferredMultiPass(FRDGBuilder& GraphBuilder, FS
 			EMobileSceneTextureSetupMode::SceneDepth | 
 			EMobileSceneTextureSetupMode::CustomDepth | 
 			EMobileSceneTextureSetupMode::GBuffers;
-		PassParameters->MobileBasePass = CreateMobileBasePassUniformBuffer(GraphBuilder, View, EMobileBasePass::Translucent, SetupMode);
+
+		FMobileBasePassUniformParameters* BasePassParameters = GraphBuilder.AllocParameters<FMobileBasePassUniformParameters>();
+		SetupMobileBasePassUniformParameters(GraphBuilder, View, EMobileBasePass::Translucent, SetupMode, *BasePassParameters);
+		PassParameters->MobileBasePass = GraphBuilder.CreateUniformBuffer(BasePassParameters);
+
 		PassParameters->ReflectionCapture = View.MobileReflectionCaptureUniformBuffer;
 		PassParameters->LocalFogVolumeInstances = View.LocalFogVolumeViewData.GPUInstanceDataBufferSRV;
 		PassParameters->LocalFogVolumeTileDrawIndirectBuffer = View.LocalFogVolumeViewData.GPUTileDrawIndirectBuffer;
@@ -2442,20 +2446,47 @@ void FMobileSceneRenderer::RenderDeferredMultiPass(FRDGBuilder& GraphBuilder, FS
 
 		const EMobileSSRQuality MobileSSRQuality = ActiveMobileSSRQuality(View, bShouldRenderVelocities);
 
+		bool bSeparateTranslucyPass = MobileSupportsSM5MaterialNodes(ShaderPlatform) && !View.IsUnderwater();
+
 		GraphBuilder.AddPass(
-			RDG_EVENT_NAME("LightingAndTranslucency"),
+			bSeparateTranslucyPass ? RDG_EVENT_NAME("LightingAndFog") : RDG_EVENT_NAME("LightingAndTranslucency"),
 			PassParameters,
 			ERDGPassFlags::Raster,
-			[this, PassParameters, ViewContext, MobileSSRQuality, &SceneTextures, &SortedLightSet](FRHICommandList& RHICmdList)
+			[this, PassParameters, ViewContext, MobileSSRQuality, bSeparateTranslucyPass, &SortedLightSet](FRHICommandList& RHICmdList)
 			{
 				FViewInfo& View = *ViewContext.ViewInfo;
 
 				MobileDeferredShadingPass(RHICmdList, ViewContext.ViewIndex, Views.Num(), View, *Scene, SortedLightSet, VisibleLightInfos, MobileSSRQuality, PassParameters->BentNormalAOTexture);
 				RenderFog(RHICmdList, View);
-
-				// Draw translucency.
-				RenderTranslucency(RHICmdList, View, Views, StandardTranslucencyPass, StandardTranslucencyMeshPass, &PassParameters->InstanceCullingDrawParams);
+				
+				if (!bSeparateTranslucyPass)
+				{
+					RenderTranslucency(RHICmdList, View, Views, StandardTranslucencyPass, StandardTranslucencyMeshPass, &PassParameters->InstanceCullingDrawParams);
+				}
 			});
+
+		if (bSeparateTranslucyPass)
+		{
+			const bool bShouldRenderTranslucency = ShouldRenderTranslucency(StandardTranslucencyPass, Views) && ViewFamily.EngineShowFlags.Translucency && (View.ParallelMeshDrawCommandPasses[StandardTranslucencyMeshPass] != nullptr);
+
+			if (bShouldRenderTranslucency)
+			{
+				FRDGTextureRef SceneColorCopyTexture = AddCopySceneColorPass(GraphBuilder, Views, SceneTextures.Color, false);
+				if (SceneColorCopyTexture != nullptr)
+				{
+					BasePassParameters->SceneColorCopyTexture = SceneColorCopyTexture;
+				}
+
+				GraphBuilder.AddPass(
+					RDG_EVENT_NAME("Translucency"),
+					PassParameters,
+					ERDGPassFlags::Raster,
+					[this, &View, PassParameters](FRHICommandList& RHICmdList)
+					{
+						RenderTranslucency(RHICmdList, View, Views, StandardTranslucencyPass, StandardTranslucencyMeshPass, &PassParameters->InstanceCullingDrawParams);
+					});
+			}
+		}
 	}
 }
 
diff --git a/Engine/Source/Runtime/Renderer/Private/TranslucentRendering.cpp b/Engine/Source/Runtime/Renderer/Private/TranslucentRendering.cpp
--- a/Engine/Source/Runtime/Renderer/Private/TranslucentRendering.cpp
+++ b/Engine/Source/Runtime/Renderer/Private/TranslucentRendering.cpp
@@ -312,13 +312,13 @@ public:
 
 	static bool ShouldCompilePermutation(const FGlobalShaderPermutationParameters& Parameters)
 	{
-		return IsFeatureLevelSupported(Parameters.Platform, ERHIFeatureLevel::SM5);
+		return IsFeatureLevelSupported(Parameters.Platform, ERHIFeatureLevel::SM5) || MobileSupportsSM5MaterialNodes(Parameters.Platform);
 	}
 };
 
 IMPLEMENT_GLOBAL_SHADER(FCopySceneColorPS, "/Engine/Private/TranslucentLightingShaders.usf", "CopySceneColorMain", SF_Pixel);
 
-static FRDGTextureRef AddCopySceneColorPass(FRDGBuilder& GraphBuilder, TArrayView<const FViewInfo> Views, FRDGTextureMSAA SceneColor,bool WithAlpha = false)
+FRDGTextureRef AddCopySceneColorPass(FRDGBuilder& GraphBuilder, TArrayView<const FViewInfo> Views, FRDGTextureMSAA SceneColor, bool WithAlpha)
 {
 	FRDGTextureRef SceneColorCopyTexture = nullptr;
 	ERenderTargetLoadAction LoadAction = ERenderTargetLoadAction::ENoAction;
diff --git a/Engine/Source/Runtime/Renderer/Private/TranslucentRendering.h b/Engine/Source/Runtime/Renderer/Private/TranslucentRendering.h
--- a/Engine/Source/Runtime/Renderer/Private/TranslucentRendering.h
+++ b/Engine/Source/Runtime/Renderer/Private/TranslucentRendering.h
@@ -35,6 +35,9 @@ struct FSeparateTranslucencyDimensions
 
 DECLARE_GPU_DRAWCALL_STAT_EXTERN(Translucency);
 
+/** Add Copy SceneColor Pass. */
+FRDGTextureRef AddCopySceneColorPass(FRDGBuilder& GraphBuilder, TArrayView<const FViewInfo> Views, FRDGTextureMSAA SceneColor, bool WithAlpha = false);
+
 /** Converts the the translucency pass into the respective mesh pass. */
 EMeshPass::Type TranslucencyPassToMeshPass(ETranslucencyPass::Type TranslucencyPass);
 
