[FastGeo] Fixed regression where proxy could be added twice in Lumen (FScene::LumenAddPrimitive).

diff --git a/Engine/Source/Runtime/Renderer/Private/RendererScene.cpp b/Engine/Source/Runtime/Renderer/Private/RendererScene.cpp
--- a/Engine/Source/Runtime/Renderer/Private/RendererScene.cpp
+++ b/Engine/Source/Runtime/Renderer/Private/RendererScene.cpp
@@ -1815,7 +1815,13 @@ void FScene::UpdatePrimitivesDrawnInGame_RenderThread(TArrayView<FPrimitiveScene
 			if (bDrawnInGame && bPrimitiveIndexValid)
 			{
 				DistanceFieldSceneData.AddPrimitive(PrimitiveSceneInfo);
-				LumenAddPrimitive(PrimitiveSceneInfo);
+				// 5.6.1 Temporary Fix
+				// IsDrawInGame() is not sufficient to know if primitive wasn't already added in Lumen (see UpdateVisibleInLumenScene).
+				// Validate that proxy was not already added.
+				if (PrimitiveSceneProxy->IsVisibleInLumenScene() && PrimitiveSceneInfo->LumenPrimitiveGroupIndices.IsEmpty())
+				{
+					LumenAddPrimitive(PrimitiveSceneInfo);
+				}
 			}
 				
 #if RHI_RAYTRACING
