[Light Grid] Fix crash when turning on light grid visualization in non-editor builds.
- necessary shaders are not included in builds so just show a warning instead.

diff --git a/Engine/Source/Runtime/Renderer/Private/DeferredShadingRenderer.cpp b/Engine/Source/Runtime/Renderer/Private/DeferredShadingRenderer.cpp
--- a/Engine/Source/Runtime/Renderer/Private/DeferredShadingRenderer.cpp
+++ b/Engine/Source/Runtime/Renderer/Private/DeferredShadingRenderer.cpp
@@ -319,7 +319,6 @@ namespace Nanite
 	extern bool IsStatFilterActive(const FString& FilterName);
 	extern void ListStatFilters(FSceneRenderer* SceneRenderer);
 }
-extern bool ShouldVisualizeLightGrid();
 
 DECLARE_CYCLE_STAT(TEXT("InitViews Intentional Stall"), STAT_InitViews_Intentional_Stall, STATGROUP_InitViews);
 
diff --git a/Engine/Source/Runtime/Renderer/Private/LightGridInjection.cpp b/Engine/Source/Runtime/Renderer/Private/LightGridInjection.cpp
--- a/Engine/Source/Runtime/Renderer/Private/LightGridInjection.cpp
+++ b/Engine/Source/Runtime/Renderer/Private/LightGridInjection.cpp
@@ -185,9 +185,16 @@ FAutoConsoleVariableRef CVarLightCullingMaxDistanceOverride(
 
 DECLARE_GPU_STAT(LightGrid);
 
-bool ShouldVisualizeLightGrid()
+bool ShouldVisualizeLightGrid(EShaderPlatform InShaderPlatform)
 {
-	return GForwardLightGridDebug > 0;
+	const bool bEnabled = GForwardLightGridDebug > 0;
+
+	if (bEnabled && !GIsEditor)
+	{
+		GEngine->AddOnScreenDebugMessage(uint64(0xA2985F1D), 1.f, FColor::Red, TEXT("Light grid visualization is only support in editor builds."));
+	}
+
+	return bEnabled && ShaderPrint::IsSupported(InShaderPlatform) && GIsEditor;
 }
 
 // If this is changed, the LIGHT_GRID_USES_16BIT_BUFFERS define from LightGridCommon.ush should also be updated.
@@ -1679,36 +1686,35 @@ IMPLEMENT_GLOBAL_SHADER(FDebugLightGridPS, "/Engine/Private/LightGridInjection.u
 
 FScreenPassTexture AddVisualizeLightGridPass(FRDGBuilder& GraphBuilder, const FViewInfo& View, FScreenPassTexture ScreenPassSceneColor, FScreenPassTexture SceneDepthTexture)
 {
-	if (ShaderPrint::IsSupported(View.Family->GetShaderPlatform()))
-	{
-		RDG_EVENT_SCOPE(GraphBuilder, "VisualizeLightGrid");
-
-		// Force ShaderPrint on.
-		ShaderPrint::SetEnabled(true);
-
-		ShaderPrint::RequestSpaceForLines(128);
-		ShaderPrint::RequestSpaceForCharacters(128);
-
-		FDebugLightGridPS::FPermutationDomain PermutationVector;
-		TShaderMapRef<FDebugLightGridPS> PixelShader(View.ShaderMap, PermutationVector);
-		FDebugLightGridPS::FParameters* PassParameters = GraphBuilder.AllocParameters<FDebugLightGridPS::FParameters>();
-		PassParameters->ViewUniformBuffer = View.ViewUniformBuffer;
-		PassParameters->ForwardLightStruct = View.ForwardLightingResources.ForwardLightUniformBuffer;
-		ShaderPrint::SetParameters(GraphBuilder, View.ShaderPrintData, PassParameters->ShaderPrintParameters);
-		PassParameters->DepthTexture = SceneDepthTexture.IsValid() ? SceneDepthTexture.Texture : GSystemTextures.GetMaxFP16Depth(GraphBuilder);
-		PassParameters->ScreenToPrimaryScreenPos = (
-			FScreenTransform::ChangeTextureBasisFromTo(ScreenPassSceneColor, FScreenTransform::ETextureBasis::TexelPosition, FScreenTransform::ETextureBasis::ViewportUV) *
-			FScreenTransform::ChangeTextureBasisFromTo(SceneDepthTexture, FScreenTransform::ETextureBasis::ViewportUV, FScreenTransform::ETextureBasis::TexelPosition));
-		PassParameters->MiniFontTexture = GetMiniFontTexture();
-		PassParameters->RenderTargets[0] = FRenderTargetBinding(ScreenPassSceneColor.Texture, ERenderTargetLoadAction::ELoad);
-		PassParameters->DebugMode = GForwardLightGridDebug;
-		PassParameters->MaxThreshold = GForwardLightGridDebugMaxThreshold;
-
-		FRHIBlendState* PreMultipliedColorTransmittanceBlend = TStaticBlendState<CW_RGB, BO_Add, BF_One, BF_SourceAlpha, BO_Add, BF_Zero, BF_One>::GetRHI();
-
-		FPixelShaderUtils::AddFullscreenPass<FDebugLightGridPS>(GraphBuilder, View.ShaderMap, RDG_EVENT_NAME("DebugLightGridCS"), PixelShader, PassParameters,
-			ScreenPassSceneColor.ViewRect, PreMultipliedColorTransmittanceBlend);
-	}
+	checkf(ShouldVisualizeLightGrid(View.Family->GetShaderPlatform()), TEXT("Must check ShouldVisualizeLightGrid(...) before calling AddVisualizeLightGridPass(...)."));
+
+	RDG_EVENT_SCOPE(GraphBuilder, "VisualizeLightGrid");
+
+	// Force ShaderPrint on.
+	ShaderPrint::SetEnabled(true);
+
+	ShaderPrint::RequestSpaceForLines(128);
+	ShaderPrint::RequestSpaceForCharacters(128);
+
+	FDebugLightGridPS::FPermutationDomain PermutationVector;
+	TShaderMapRef<FDebugLightGridPS> PixelShader(View.ShaderMap, PermutationVector);
+	FDebugLightGridPS::FParameters* PassParameters = GraphBuilder.AllocParameters<FDebugLightGridPS::FParameters>();
+	PassParameters->ViewUniformBuffer = View.ViewUniformBuffer;
+	PassParameters->ForwardLightStruct = View.ForwardLightingResources.ForwardLightUniformBuffer;
+	ShaderPrint::SetParameters(GraphBuilder, View.ShaderPrintData, PassParameters->ShaderPrintParameters);
+	PassParameters->DepthTexture = SceneDepthTexture.IsValid() ? SceneDepthTexture.Texture : GSystemTextures.GetMaxFP16Depth(GraphBuilder);
+	PassParameters->ScreenToPrimaryScreenPos = (
+		FScreenTransform::ChangeTextureBasisFromTo(ScreenPassSceneColor, FScreenTransform::ETextureBasis::TexelPosition, FScreenTransform::ETextureBasis::ViewportUV) *
+		FScreenTransform::ChangeTextureBasisFromTo(SceneDepthTexture, FScreenTransform::ETextureBasis::ViewportUV, FScreenTransform::ETextureBasis::TexelPosition));
+	PassParameters->MiniFontTexture = GetMiniFontTexture();
+	PassParameters->RenderTargets[0] = FRenderTargetBinding(ScreenPassSceneColor.Texture, ERenderTargetLoadAction::ELoad);
+	PassParameters->DebugMode = GForwardLightGridDebug;
+	PassParameters->MaxThreshold = GForwardLightGridDebugMaxThreshold;
+
+	FRHIBlendState* PreMultipliedColorTransmittanceBlend = TStaticBlendState<CW_RGB, BO_Add, BF_One, BF_SourceAlpha, BO_Add, BF_Zero, BF_One>::GetRHI();
+
+	FPixelShaderUtils::AddFullscreenPass<FDebugLightGridPS>(GraphBuilder, View.ShaderMap, RDG_EVENT_NAME("DebugLightGridCS"), PixelShader, PassParameters,
+		ScreenPassSceneColor.ViewRect, PreMultipliedColorTransmittanceBlend);
 
 	return MoveTemp(ScreenPassSceneColor);
 }
diff --git a/Engine/Source/Runtime/Renderer/Private/PostProcess/PostProcessing.cpp b/Engine/Source/Runtime/Renderer/Private/PostProcess/PostProcessing.cpp
--- a/Engine/Source/Runtime/Renderer/Private/PostProcess/PostProcessing.cpp
+++ b/Engine/Source/Runtime/Renderer/Private/PostProcess/PostProcessing.cpp
@@ -88,7 +88,7 @@ bool IsMobileEyeAdaptationEnabled(const FViewInfo& View);
 bool IsValidBloomSetupVariation(bool bUseBloom, bool bUseSun, bool bUseDof, bool bUseEyeAdaptation);
 
 extern FScreenPassTexture AddVisualizeLightGridPass(FRDGBuilder& GraphBuilder, const FViewInfo& View, FScreenPassTexture ScreenPassSceneColor, FScreenPassTexture SceneDepthTexture);
-extern bool ShouldVisualizeLightGrid();
+extern bool ShouldVisualizeLightGrid(EShaderPlatform ShaderPlatform);
 
 namespace
 {
@@ -555,7 +555,7 @@ void AddPostProcessingPasses(
 #endif
 	PassSequence.SetEnabled(EPass::VisualizeLumenScene, LumenVisualizeMode >= 0 && LumenVisualizeMode != VISUALIZE_MODE_OVERVIEW && LumenVisualizeMode != VISUALIZE_MODE_PERFORMANCE_OVERVIEW && bPostProcessingEnabled);
 	PassSequence.SetEnabled(EPass::VisualizeSubstrate, Substrate::ShouldRenderSubstrateDebugPasses(View));
-	PassSequence.SetEnabled(EPass::VisualizeLightGrid, ShouldVisualizeLightGrid());
+	PassSequence.SetEnabled(EPass::VisualizeLightGrid, ShouldVisualizeLightGrid(View.Family->GetShaderPlatform()));
 
 #if WITH_EDITOR
 	PassSequence.SetEnabled(EPass::VisualizeSkyAtmosphere, Scene&& View.Family && View.Family->EngineShowFlags.VisualizeSkyAtmosphere && ShouldRenderSkyAtmosphereDebugPasses(Scene, View.Family->EngineShowFlags));
@@ -2525,7 +2525,7 @@ void AddMobilePostProcessingPasses(FRDGBuilder& GraphBuilder, FScene* Scene, con
 	PassSequence.SetEnabled(EPass::SecondaryUpscale, View.Family->GetSecondarySpatialUpscalerInterface() != nullptr);
 
 	PassSequence.SetEnabled(EPass::Visualize, View.Family->EngineShowFlags.ShaderComplexity);
-	PassSequence.SetEnabled(EPass::VisualizeLightGrid, ShouldVisualizeLightGrid());
+	PassSequence.SetEnabled(EPass::VisualizeLightGrid, ShouldVisualizeLightGrid(View.Family->GetShaderPlatform()));
 	PassSequence.SetEnabled(EPass::HMDDistortion, View.Family->EngineShowFlags.StereoRendering && View.Family->EngineShowFlags.HMDDistortion);
 
 	// TODO: Add support for EPass::AlphaInvert if the AlphaInvert show flag is enabled, but EPass::PrimaryUpscale is disabled
