Refactor VSM ID allocation so that we can control the order more consistently.
The immediate goal is to guarantee that directional lights get allocated first so that we can easily allocate receiver masks just for directional lights without the memory overhead of doing it for all local lights as well.

Tested on QAGame, EngineTest, ElectricDreams, NaniteTest

diff --git a/Engine/Shaders/Private/VirtualShadowMaps/VirtualShadowMapPageAccessCommon.ush b/Engine/Shaders/Private/VirtualShadowMaps/VirtualShadowMapPageAccessCommon.ush
--- a/Engine/Shaders/Private/VirtualShadowMaps/VirtualShadowMapPageAccessCommon.ush
+++ b/Engine/Shaders/Private/VirtualShadowMaps/VirtualShadowMapPageAccessCommon.ush
@@ -369,13 +369,13 @@ uint GetVirtualShadowMapHZBArrayIndex(bool bUseStaticOcclusion)
 }
 
 // Gather a 2x2 footprint from a page table texture.
-uint4 GatherPageTable(Texture2D<uint> PageTableTexture, uint2 TexelCoord, uint HMipLevel, float UVScaleFactor)
+uint4 GatherPageTable(Texture2D<uint> PageTableTexture, uint2 TexelCoord, uint HMipLevel, float2 UVScaleFactor)
 {
 	uint2 TexelCoordMip = TexelCoord >> HMipLevel;
 #if COMPILER_SUPPORTS_GATHER_LOD_RED
 	// Offset to 2x2 footprint center and scale to UV space
 	float MipScale = float(1U << HMipLevel);
-	float2 UV = float2(TexelCoordMip + uint2(1U, 1U)) * VirtualShadowMap.PageTableTextureSizeInvSize.zw * MipScale * UVScaleFactor;
+	float2 UV = float2(TexelCoordMip + uint2(1U, 1U)) * MipScale * UVScaleFactor;
 	uint4 Result2x2 = PageTableTexture.GatherLODRed(VirtualShadowMap.PageTableSampler, UV, HMipLevel);
 #else
 	uint4 TexelRectMip = uint4(TexelCoordMip.xy, TexelCoordMip.xy + uint2(1U, 1U));
diff --git a/Engine/Shaders/Private/VirtualShadowMaps/VirtualShadowMapPageManagement.usf b/Engine/Shaders/Private/VirtualShadowMaps/VirtualShadowMapPageManagement.usf
--- a/Engine/Shaders/Private/VirtualShadowMaps/VirtualShadowMapPageManagement.usf
+++ b/Engine/Shaders/Private/VirtualShadowMaps/VirtualShadowMapPageManagement.usf
@@ -153,7 +153,7 @@ void ProcessMaskMip4x4(RWTexture2D<uint> OutPageReceiverMaskMip, inout uint Mask
 // Gather a 2x2 footprint given the upper-left coordinate
 uint4 GatherReceiverMask(uint2 TexelCoord)
 {
-	return GatherPageTable(InPageReceiverMasks, TexelCoord, 0u, 0.5f);
+	return GatherPageTable(InPageReceiverMasks, TexelCoord, 0u, 0.5f * VirtualShadowMap.PageTableTextureSizeInvSize.zw);
 }
 
 /**
diff --git a/Engine/Shaders/Private/VirtualShadowMaps/VirtualShadowMapPageOverlap.ush b/Engine/Shaders/Private/VirtualShadowMaps/VirtualShadowMapPageOverlap.ush
--- a/Engine/Shaders/Private/VirtualShadowMaps/VirtualShadowMapPageOverlap.ush
+++ b/Engine/Shaders/Private/VirtualShadowMaps/VirtualShadowMapPageOverlap.ush
@@ -63,13 +63,13 @@ FScreenRect VirtualShadowMapGetUncachedScreenRect(FScreenRect ScreenRect, FVirtu
 
 uint4 GatherPageFlags(uint2 TexelCoord, uint HMipLevel)
 {
-	return GatherPageTable(VirtualShadowMap.PageFlags, TexelCoord, HMipLevel, 1.0f);
+	return GatherPageTable(VirtualShadowMap.PageFlags, TexelCoord, HMipLevel, VirtualShadowMap.PageTableTextureSizeInvSize.zw);
 }
 
 uint4 GatherReceiverMask(uint2 TexelCoord, uint HMipLevel)
 {
 	// Note: UV scaled by extra 0.5f to account for 2x size of mask texture
-	return GatherPageTable(VirtualShadowMap.PageReceiverMasks, TexelCoord, HMipLevel, 0.5f);
+	return GatherPageTable(VirtualShadowMap.PageReceiverMasks, TexelCoord, HMipLevel, 0.5f * VirtualShadowMap.PageTableTextureSizeInvSize.zw);
 }
 
 // Returns true if ANY flags in FlagMask are set on at least one overlapped page (see GetPageFlagMaskForRendering below)
diff --git a/Engine/Shaders/Private/VirtualShadowMaps/VirtualShadowMapPhysicalPageManagement.usf b/Engine/Shaders/Private/VirtualShadowMaps/VirtualShadowMapPhysicalPageManagement.usf
--- a/Engine/Shaders/Private/VirtualShadowMaps/VirtualShadowMapPhysicalPageManagement.usf
+++ b/Engine/Shaders/Private/VirtualShadowMaps/VirtualShadowMapPhysicalPageManagement.usf
@@ -179,11 +179,11 @@ void UpdatePhysicalPageAddresses(uint3 Index : SV_DispatchThreadID)
 			PrevGlobalPageOffset = CalcPageOffset(PrevVirtualShadowMapHandle, PrevMetaData.MipLevel, PrevMetaData.PageAddress);
 
 			FNextVirtualShadowMapData NextData = NextVirtualShadowMapData[PrevVirtualShadowMapHandle.GetDataIndex()];
-			FVirtualShadowMapHandle VirtualShadowMapHandle = FVirtualShadowMapHandle::MakeFromId(NextData.NextVirtualShadowMapId);
-
+			
 			// Check if it maps to a valid virtual shadow map this frame
-			if (VirtualShadowMapHandle.IsValid())
+			if (NextData.Flags & VSM_NEXT_FLAG_VALID)
 			{
+				FVirtualShadowMapHandle VirtualShadowMapHandle = FVirtualShadowMapHandle::MakeFromId(NextData.NextVirtualShadowMapId);
 				// Clipmap panning; zeroed otherwise so safe
 				int2 TestPageAddress = int2(PrevMetaData.PageAddress) + NextData.PageAddressOffset;
 				if (IsVirtualShadowMapPageAddressValid(TestPageAddress, PrevMetaData.MipLevel))
diff --git a/Engine/Shaders/Private/VirtualShadowMaps/VirtualShadowMapThrottle.usf b/Engine/Shaders/Private/VirtualShadowMaps/VirtualShadowMapThrottle.usf
--- a/Engine/Shaders/Private/VirtualShadowMaps/VirtualShadowMapThrottle.usf
+++ b/Engine/Shaders/Private/VirtualShadowMaps/VirtualShadowMapThrottle.usf
@@ -41,9 +41,10 @@ void ProcessPrevFramePerfDataCS(uint DispatchThreadId : SV_DispatchThreadID)
 		FVirtualShadowMapHandle PrevVSMHandle = FVirtualShadowMapHandle::MakeFromId(DispatchThreadId);
 		const FNextVirtualShadowMapData NextVSMData = NextVirtualShadowMapData[PrevVSMHandle.GetDataIndex()];
 
-		FVirtualShadowMapHandle NextVSMHandle = FVirtualShadowMapHandle::MakeFromId(NextVSMData.NextVirtualShadowMapId);
-		if (NextVSMHandle.IsValid())
+		if (NextVSMData.Flags & VSM_NEXT_FLAG_VALID)
 		{
+			FVirtualShadowMapHandle NextVSMHandle = FVirtualShadowMapHandle::MakeFromId(NextVSMData.NextVirtualShadowMapId);
+
 			const uint NPFOffset = VSM_NPF_SIZEOF_HEADER + VSM_NPF_SIZEOF_ENTRY * PrevVSMHandle.GetDataIndex();
 			const uint ClustersHW = PrevNanitePerformanceFeedback[NPFOffset + VSM_NPF_ENTRY_HW_CLUSTERS];
 			const uint ClustersSW = PrevNanitePerformanceFeedback[NPFOffset + VSM_NPF_ENTRY_SW_CLUSTERS];
diff --git a/Engine/Shaders/Shared/VirtualShadowMapDefinitions.h b/Engine/Shaders/Shared/VirtualShadowMapDefinitions.h
--- a/Engine/Shaders/Shared/VirtualShadowMapDefinitions.h
+++ b/Engine/Shaders/Shared/VirtualShadowMapDefinitions.h
@@ -120,6 +120,9 @@ VirtualShadowMapDefinitions.h: used in virtual shadow map shaders and C++ code t
 #define VSM_TB_ENTRY_THROTTLE						2
 #define VSM_TB_SIZEOF_ENTRY							3
 
+// Next data
+#define VSM_NEXT_FLAG_VALID		1
+
 #ifdef __cplusplus
 #include "HLSLTypeAliases.h"
 #include "HLSLMathAliases.h"
@@ -143,9 +146,9 @@ struct FVSMCullingBatchInfo
 
 struct FNextVirtualShadowMapData
 {
+	uint Flags;
 	int NextVirtualShadowMapId;
 	int2 PageAddressOffset;
-	int _Padding;
 };
 
 //ExtraBias must >= 0
diff --git a/Engine/Source/Runtime/Renderer/Private/LightRendering.cpp b/Engine/Source/Runtime/Renderer/Private/LightRendering.cpp
--- a/Engine/Source/Runtime/Renderer/Private/LightRendering.cpp
+++ b/Engine/Source/Runtime/Renderer/Private/LightRendering.cpp
@@ -1674,7 +1674,6 @@ void FDeferredShadingSceneRenderer::RenderLights(
 						!(bDirectLighting && bDrawLightFunction) &&
 						!bDrawPreviewIndicator &&
 						SortedLightInfo.SortKey.Fields.LightType != LightType_Directional &&
-						VisibleLightInfo.VirtualShadowMapId != INDEX_NONE &&	// Not a directional light, so no per-view clipmaps
 						VisibleLightInfo.ContainsOnlyVirtualShadowMaps() &&
 						!bDrawFirstPersonSelfShadow;
 					bElideScreenShadowMaskSubPixel = bElideScreenShadowMask && CVarHairStrandsAllowOneTransmittancePass.GetValueOnRenderThread() > 0;
diff --git a/Engine/Source/Runtime/Renderer/Private/SceneRendering.cpp b/Engine/Source/Runtime/Renderer/Private/SceneRendering.cpp
--- a/Engine/Source/Runtime/Renderer/Private/SceneRendering.cpp
+++ b/Engine/Source/Runtime/Renderer/Private/SceneRendering.cpp
@@ -66,6 +66,7 @@
 #include "HairStrands/HairStrandsVisibility.h"
 #include "SystemTextures.h"
 #include "VirtualShadowMaps/VirtualShadowMapClipmap.h"
+#include "VirtualShadowMaps/VirtualShadowMapCacheManager.h"
 #include "Misc/AutomationTest.h"
 #include "Engine/TextureCube.h"
 #include "GPUSkinCacheVisualizationData.h"
@@ -555,7 +556,7 @@ FASTVRAM_CVAR(DistanceFieldAOScreenGridResources, 1);
 FASTVRAM_CVAR(ForwardLightingCullingResources, 1);
 FASTVRAM_CVAR(GlobalDistanceFieldCullGridBuffers, 1);
 
-static TSharedPtr<FVirtualShadowMapClipmap> FindShadowClipmapForView(const TArray<TSharedPtr<FVirtualShadowMapClipmap>,SceneRenderingAllocator>& VirtualShadowMapClipmaps, const FViewInfo* View)
+TSharedPtr<FVirtualShadowMapClipmap> FVisibleLightInfo::FindVirtualShadowMapShadowClipmapForView(const FViewInfo* View) const
 {
 	for (const auto& Clipmap : VirtualShadowMapClipmaps)
 	{
@@ -579,11 +580,25 @@ int32 FVisibleLightInfo::GetVirtualShadowMapId(const FViewInfo* View) const
 {
 	if (VirtualShadowMapClipmaps.Num())
 	{
-		return FindShadowClipmapForView(VirtualShadowMapClipmaps, View)->GetVirtualShadowMapId();
+		TSharedPtr<FVirtualShadowMapClipmap> Clipmap = FindVirtualShadowMapShadowClipmapForView(View);
+		return Clipmap.IsValid() ? Clipmap->GetVirtualShadowMapId() : INDEX_NONE;
 	}
 	else
 	{
-		return VirtualShadowMapId;
+		// Look through local light projected shadows
+		// TODO: With ForceOnlyVSM there should only be one
+		for (int32 ShadowIndex = 0; ShadowIndex < AllProjectedShadows.Num(); ShadowIndex++)
+		{
+			const FProjectedShadowInfo* ProjectedShadowInfo = AllProjectedShadows[ShadowIndex];
+			if (ProjectedShadowInfo->VirtualShadowMapPerLightCacheEntry.IsValid())
+			{
+				// VSM ID should be allocated by the time anyone calls this function
+				check(ProjectedShadowInfo->VirtualShadowMapPerLightCacheEntry->GetVirtualShadowMapId() >= 0);
+				return ProjectedShadowInfo->VirtualShadowMapPerLightCacheEntry->GetVirtualShadowMapId();
+			}
+		}
+
+		return INDEX_NONE;
 	}
 }
 
diff --git a/Engine/Source/Runtime/Renderer/Private/SceneRendering.h b/Engine/Source/Runtime/Renderer/Private/SceneRendering.h
--- a/Engine/Source/Runtime/Renderer/Private/SceneRendering.h
+++ b/Engine/Source/Runtime/Renderer/Private/SceneRendering.h
@@ -229,10 +229,7 @@ public:
 	*/
 	bool ContainsOnlyVirtualShadowMaps() const;
 
-	/**
-	 * Returns true if there are any virtual shadow maps for any views for this light.
-	 */
-	bool HasVirtualShadowMap() const { return VirtualShadowMapId != INDEX_NONE; }
+	TSharedPtr<FVirtualShadowMapClipmap> FindVirtualShadowMapShadowClipmapForView(const FViewInfo* View) const;
 
 	/**
 	* Prefer this to direct access of the VirtualShadowMapId member when a view is known.
@@ -240,7 +237,6 @@ public:
 	* while the VirtualShadowMapId variable will simply be an arbitrary one of them if multiple exist.
 	*/
 	int32 GetVirtualShadowMapId( const FViewInfo* View ) const;
-	int32 VirtualShadowMapId = INDEX_NONE;	
 };
 
 // Stores the primitive count of each translucency pass (redundant, could be computed after sorting but this way we touch less memory)
diff --git a/Engine/Source/Runtime/Renderer/Private/ShadowRendering.h b/Engine/Source/Runtime/Renderer/Private/ShadowRendering.h
--- a/Engine/Source/Runtime/Renderer/Private/ShadowRendering.h
+++ b/Engine/Source/Runtime/Renderer/Private/ShadowRendering.h
@@ -438,8 +438,8 @@ public:
 
 	EShadowMeshSelection MeshSelectionMask = EShadowMeshSelection::All;
 
-	int32 VirtualShadowMapId = INDEX_NONE;
 	TSharedPtr<FVirtualShadowMapPerLightCacheEntry> VirtualShadowMapPerLightCacheEntry;
+	TSharedPtr<FVirtualShadowMapClipmap> VirtualShadowMapClipmap;
 
 	/** View projection matrices for each cubemap face, used by one pass point light shadows. */
 	TArray<FMatrix> OnePassShadowViewProjectionMatrices;
@@ -470,7 +470,7 @@ public:
 	int32 SubjectPrimitiveComponentIndex = -1;
 
 	TArray<int32, TInlineAllocator<6> > ViewIds;
-	TSharedPtr<FVirtualShadowMapClipmap> VirtualShadowMapClipmap;
+	
 public:
 
 	// default constructor
@@ -806,7 +806,7 @@ public:
 		return ::GetShadowMeshPassType(GetShadowDepthType());
 	}
 
-	bool HasVirtualShadowMap() const { return VirtualShadowMapId != INDEX_NONE || VirtualShadowMapClipmap.IsValid(); }
+	bool HasVirtualShadowMap() const { return VirtualShadowMapPerLightCacheEntry.IsValid(); }
 
 	FParallelMeshDrawCommandPass& GetShadowDepthPass() { return ShadowDepthPass; }
 
diff --git a/Engine/Source/Runtime/Renderer/Private/ShadowSetup.cpp b/Engine/Source/Runtime/Renderer/Private/ShadowSetup.cpp
--- a/Engine/Source/Runtime/Renderer/Private/ShadowSetup.cpp
+++ b/Engine/Source/Runtime/Renderer/Private/ShadowSetup.cpp
@@ -4332,7 +4332,6 @@ void FSceneRenderer::CreateWholeSceneProjectedShadow(
 						AddInteractingPrimitives(LightSceneInfo->GetDynamicInteractionStaticPrimitiveList(), ProjectedShadowInfo, bContainsNaniteSubjects, EmptyHull, VirtualSmPerLightCacheEntry);
 						ProjectedShadowInfo->bContainsNaniteSubjects = bContainsNaniteSubjects;
 					}
-					VisibleLightInfo.VirtualShadowMapId = ProjectedShadowInfo->VirtualShadowMapId;
 					VisibleLightInfo.AllProjectedShadows.Add(ProjectedShadowInfo);
 				}
 				
diff --git a/Engine/Source/Runtime/Renderer/Private/Shadows/ShadowSceneRenderer.cpp b/Engine/Source/Runtime/Renderer/Private/Shadows/ShadowSceneRenderer.cpp
--- a/Engine/Source/Runtime/Renderer/Private/Shadows/ShadowSceneRenderer.cpp
+++ b/Engine/Source/Runtime/Renderer/Private/Shadows/ShadowSceneRenderer.cpp
@@ -266,7 +266,7 @@ void FShadowSceneRenderer::PreInitViews(FRDGBuilder& GraphBuilder)
 		for (const FHeapPair &HeapPair : DistantLightUpdateQueue)
 		{
 			// Mark frame it was scheduled, this is picked up later in AddLocalLightShadow to trigger invalidation 
-			HeapPair.CacheEntry->Current.ScheduledFrameNumber = SceneFrameNumber;
+			HeapPair.CacheEntry->MarkScheduled(SceneFrameNumber);
 		}
 	});
 }
@@ -429,9 +429,9 @@ TSharedPtr<FVirtualShadowMapPerLightCacheEntry> FShadowSceneRenderer::AddLocalLi
 
 		}
 	}
-
-
+	
 	const int32 NumMaps = ProjectedShadowInitializer.bOnePassPointLightShadow ? 6 : 1;
+
 	TSharedPtr<FVirtualShadowMapPerLightCacheEntry> PerLightCacheEntry = CacheManager->FindCreateLightCacheEntry(LightId, 0, NumMaps);
 	LocalLightSetup.PerLightCacheEntry = PerLightCacheEntry;
 
@@ -444,15 +444,12 @@ TSharedPtr<FVirtualShadowMapPerLightCacheEntry> FShadowSceneRenderer::AddLocalLi
 		!bShouldForceTimeSliceDistantUpdate,
 		IsVirtualShadowMapLocalReceiverMaskEnabled());
 
-	if (bIsDistantLight && PerLightCacheEntry->Prev.ScheduledFrameNumber == Scene.GetFrameNumber())
+	if (bIsDistantLight && PerLightCacheEntry->GetLastScheduledFrameNumber() == Scene.GetFrameNumber())
 	{
 		PerLightCacheEntry->Invalidate();
 	}
 
 	// Update info on the ProjectionShadowInfo; eventually this should all move into local data structures here
-	const int32 VirtualShadowMapId = VirtualShadowMapArray.Allocate(bIsDistantLight, NumMaps);
-	LocalLightSetup.VirtualShadowMapId = VirtualShadowMapId;
-	ProjectedShadowInfo->VirtualShadowMapId = VirtualShadowMapId;
 	ProjectedShadowInfo->VirtualShadowMapPerLightCacheEntry = PerLightCacheEntry;
 	ProjectedShadowInfo->bShouldRenderVSM = !PerLightCacheEntry->IsFullyCached();
 
@@ -468,9 +465,8 @@ TSharedPtr<FVirtualShadowMapPerLightCacheEntry> FShadowSceneRenderer::AddLocalLi
 
 		for (int32 Index = 0; Index < NumMaps; ++Index)
 		{
-			const int32 FaceVirtualShadowMapId = VirtualShadowMapId + Index;
 			FVirtualShadowMapCacheEntry& VirtualSmCacheEntry = PerLightCacheEntry->ShadowMapEntries[Index];
-			VirtualSmCacheEntry.Update(VirtualShadowMapArray, *PerLightCacheEntry, FaceVirtualShadowMapId);
+			VirtualSmCacheEntry.Update(*PerLightCacheEntry);
 
 			FVirtualShadowMapProjectionShaderData& ProjectionData = VirtualSmCacheEntry.ProjectionData;
 
@@ -563,8 +559,6 @@ void FShadowSceneRenderer::AddDirectionalLightShadow(FLightSceneInfo& LightScene
 
 		// NOTE: only contains "regular" clipmaps, the alternate types are internal to the system and needs to be queried for.
 		VisibleLightInfo.VirtualShadowMapClipmaps.Add(VirtualShadowMapClipmap);
-		// NOTE: If there are multiple camera views this will simply be associated with "one of them"
-		VisibleLightInfo.VirtualShadowMapId = VirtualShadowMapClipmap->GetVirtualShadowMapId();
 
 		RegularInfo.ProjectedShadowInfo = AddLegacySetup(VirtualShadowMapClipmap, true);
 		RegularInfo.Clipmap = VirtualShadowMapClipmap;
@@ -613,8 +607,42 @@ void FShadowSceneRenderer::AddDirectionalLightShadow(FLightSceneInfo& LightScene
 	}
 }
 
+void FShadowSceneRenderer::AllocateVirtualShadowMapIds()
+{
+	// TODO: Move the bReferencedThisRender setting stuff to here as well just for organization?
+
+	// Directional lights allocated first
+	for (const FDirectionalLightShadowFrameSetup& FrameSetup : DirectionalLights)
+	{
+		for (const auto& ClipmapInfo : FrameSetup.ClipmapInfos)
+		{
+			if (!ClipmapInfo.Clipmap.IsValid())
+			{
+				continue;
+			}
+
+			const TSharedPtr<FVirtualShadowMapPerLightCacheEntry>& CacheEntry = ClipmapInfo.Clipmap->GetCacheEntry();
+			const int32 NextVirtualShadowMapId = VirtualShadowMapArray.Allocate(false, CacheEntry->ShadowMapEntries.Num());
+			CacheEntry->UpdateVirtualShadowMapId(NextVirtualShadowMapId);
+		}
+	}
+
+	// Then local lights
+	for (const FLocalLightShadowFrameSetup& FrameSetup : LocalLights)
+	{
+		const TSharedPtr<FVirtualShadowMapPerLightCacheEntry>& CacheEntry = FrameSetup.PerLightCacheEntry;
+		const int32 NextVirtualShadowMapId = VirtualShadowMapArray.Allocate(CacheEntry->bIsDistantLight, CacheEntry->ShadowMapEntries.Num());
+		CacheEntry->UpdateVirtualShadowMapId(NextVirtualShadowMapId);
+	}
+
+	// Then add any unreferenced lights that we still may keep around their cached pages at the end, and allocate VSM IDs for everything
+	VirtualShadowMapArray.CacheManager->UpdateUnreferencedCacheEntries(VirtualShadowMapArray);
+}
+
 void FShadowSceneRenderer::PostInitDynamicShadowsSetup()
 {
+	AllocateVirtualShadowMapIds();
+
 	// Dispatch async Nanite culling job if appropriate
 	if (CVarVSMMaterialVisibility.GetValueOnRenderThread() != 0)
 	{
@@ -935,28 +963,21 @@ void FShadowSceneRenderer::ApplyVirtualShadowMapProjectionForLight(
 		return;
 	}
 
-	FLightSceneInfo::FPersistentId LightId = LightSceneInfo->Id;
-
-	const FLightCommonFrameSetup& CommonSetup = CommonSetups[LightId];
-
-	// No VSM set up this frame
-	if (!CommonSetup.bHasVirtualShadowMap)
+	// Some lights can elide the screen shadow mask entirely, in which case they will be sampled directly in the lighting shader
+	if (!OutputScreenShadowMaskTexture)
 	{
 		return;
 	}
-	const FVisibleLightInfo& VisibleLightInfo = SceneRenderer.VisibleLightInfos[LightId];
 
-	if (!VisibleLightInfo.HasVirtualShadowMap())
-	{
-		return;
-	}
+	FLightSceneInfo::FPersistentId LightId = LightSceneInfo->Id;
+	const FLightCommonFrameSetup& CommonSetup = CommonSetups[LightId];
 
-	// Some lights can elide the screen shadow mask entirely, in which case they will be sampled directly in the lighting shader
-	if (!OutputScreenShadowMaskTexture)
+	// No VSM set up this frame
+	if (!CommonSetup.bHasVirtualShadowMap)
 	{
 		return;
 	}
-
+	
 	for (int32 ViewIndex = 0; ViewIndex < SceneRenderer.Views.Num(); ++ViewIndex)
 	{
 		RDG_EVENT_SCOPE_CONDITIONAL(GraphBuilder, SceneRenderer.Views.Num() > 1, "View%d", ViewIndex);
@@ -1014,7 +1035,7 @@ void FShadowSceneRenderer::ApplyVirtualShadowMapProjectionForLight(
 						ScissorRect,
 						VirtualShadowMapArray,
 						InputType,
-						LocalLightSetup.VirtualShadowMapId,
+						LocalLightSetup.PerLightCacheEntry->GetVirtualShadowMapId(),
 						InputType == EVirtualShadowMapProjectionInputType::HairStrands ? VirtualShadowMapMaskBitsHairStrands : VirtualShadowMapMaskBits,
 						OutputScreenShadowMaskTexture);
 				}
@@ -1031,7 +1052,7 @@ void FShadowSceneRenderer::ApplyVirtualShadowMapProjectionForLight(
 						ScissorRect,
 						InputType,
 						*LightSceneInfo,
-						LocalLightSetup.VirtualShadowMapId,
+						LocalLightSetup.PerLightCacheEntry->GetVirtualShadowMapId(),
 						OutputScreenShadowMaskTexture);
 				}
 			}			
diff --git a/Engine/Source/Runtime/Renderer/Private/Shadows/ShadowSceneRenderer.h b/Engine/Source/Runtime/Renderer/Private/Shadows/ShadowSceneRenderer.h
--- a/Engine/Source/Runtime/Renderer/Private/Shadows/ShadowSceneRenderer.h
+++ b/Engine/Source/Runtime/Renderer/Private/Shadows/ShadowSceneRenderer.h
@@ -216,6 +216,7 @@ private:
 		float ShadowsLODScaleFactor,
 		FNaniteVirtualShadowMapRenderPass& InOutRenderPass);
 
+	void AllocateVirtualShadowMapIds();
 
 	struct FLightCommonFrameSetup
 	{
@@ -230,7 +231,6 @@ private:
 
 	struct FLocalLightShadowFrameSetup
 	{
-		int32 VirtualShadowMapId = INDEX_NONE;
 		TSharedPtr<FVirtualShadowMapPerLightCacheEntry> PerLightCacheEntry;
 		// link to legacy system stuff, to be removed in due time
 		FProjectedShadowInfo* ProjectedShadowInfo = nullptr;
diff --git a/Engine/Source/Runtime/Renderer/Private/SkyAtmosphereRendering.cpp b/Engine/Source/Runtime/Renderer/Private/SkyAtmosphereRendering.cpp
--- a/Engine/Source/Runtime/Renderer/Private/SkyAtmosphereRendering.cpp
+++ b/Engine/Source/Runtime/Renderer/Private/SkyAtmosphereRendering.cpp
@@ -22,6 +22,7 @@
 #include "ScreenPass.h"
 #include "UnrealEngine.h"
 #include "PostProcess/PostProcessing.h" // IsPostProcessingWithAlphaChannelSupported
+#include "VirtualShadowMaps/VirtualShadowMapClipmap.h"
 
 
 //PRAGMA_DISABLE_OPTIMIZATION
@@ -344,13 +345,27 @@ bool ShouldSkySampleAtmosphereLightsOpaqueShadow(const FScene& Scene, const TArr
 		LightShadowData.ProjectedShadowInfo0 = GetFirstWholeSceneShadowMap(Light0);	
 		// NOTE: This will be an arbitrary clipmap if multiple exist, but this is similar to the CSM select above,
 		// which also does not take multiple views into account.
-		LightShadowData.VirtualShadowMapId0 = Light0.VirtualShadowMapId;
+		if (Light0.VirtualShadowMapClipmaps.Num() > 0)
+		{
+			LightShadowData.VirtualShadowMapId0 = Light0.VirtualShadowMapClipmaps[0]->GetVirtualShadowMapId();
+		}
+		else
+		{
+			LightShadowData.VirtualShadowMapId0 = INDEX_NONE;
+		}
 	}
 	if (LightShadowData.LightVolumetricShadowSceneinfo1 && LightShadowData.LightVolumetricShadowSceneinfo1->Proxy && LightShadowData.LightVolumetricShadowSceneinfo1->Proxy->GetCastShadowsOnAtmosphere())
 	{
 		const FVisibleLightInfo& Light1 = VisibleLightInfos[LightShadowData.LightVolumetricShadowSceneinfo1->Id];
 		LightShadowData.ProjectedShadowInfo1 = GetFirstWholeSceneShadowMap(Light1);
-		LightShadowData.VirtualShadowMapId1 = Light1.VirtualShadowMapId;
+		if (Light1.VirtualShadowMapClipmaps.Num() > 0)
+		{
+			LightShadowData.VirtualShadowMapId1 = Light1.VirtualShadowMapClipmaps[0]->GetVirtualShadowMapId();
+		}
+		else
+		{
+			LightShadowData.VirtualShadowMapId1 = INDEX_NONE;
+		}
 	}
 
 	return CVarSkyAtmosphereSampleLightShadowmap.GetValueOnRenderThread() > 0 &&
diff --git a/Engine/Source/Runtime/Renderer/Private/VirtualShadowMaps/VirtualShadowMapArray.cpp b/Engine/Source/Runtime/Renderer/Private/VirtualShadowMaps/VirtualShadowMapArray.cpp
--- a/Engine/Source/Runtime/Renderer/Private/VirtualShadowMaps/VirtualShadowMapArray.cpp
+++ b/Engine/Source/Runtime/Renderer/Private/VirtualShadowMaps/VirtualShadowMapArray.cpp
@@ -621,23 +621,15 @@ FVirtualShadowMapArray::FVirtualShadowMapArray(FScene& InScene)
 {
 }
 
-void FVirtualShadowMapArray::UpdateNextData(int32 PrevVirtualShadowMapId, int32 NextVirtualShadowMapId, FInt32Point PageOffset)
+void FVirtualShadowMapArray::UpdateNextData(int32 PrevVirtualShadowMapId, const FNextVirtualShadowMapData& Data)
 {
-	// Fill in any slots with empty mappings
-	FNextVirtualShadowMapData EmptyData;
-	EmptyData.NextVirtualShadowMapId = INDEX_NONE;
-	EmptyData.PageAddressOffset = FIntVector2(0, 0);
-
-	// TODO: Some ways to optimize this
-	// Can't use SetNum because we need the empty item initializer which doesn't fit nicely with our shared HLSL definition right now
-	NextData.Reserve(PrevVirtualShadowMapId);
-	while (PrevVirtualShadowMapId >= NextData.Num())
+	if (PrevVirtualShadowMapId >= NextData.Num())
 	{
-		NextData.Add(EmptyData);
+		// Will set flags to 0 which means "invalid/no caching"
+		NextData.SetNumZeroed(PrevVirtualShadowMapId + 1);
 	}
 
-	NextData[PrevVirtualShadowMapId].NextVirtualShadowMapId = NextVirtualShadowMapId;
-	NextData[PrevVirtualShadowMapId].PageAddressOffset = FIntVector2(PageOffset.X, PageOffset.Y);
+	NextData[PrevVirtualShadowMapId] = Data;
 }
 
 static FVirtualShadowMapPerViewParameters MakeEmptyVirtualShadowMapPerViewParameters(FRDGBuilder& GraphBuilder)
@@ -1757,40 +1749,33 @@ void FVirtualShadowMapArray::UpdateVisualizeLight(
 	for (int32 LightId = 0; LightId < VisibleLightInfos.Num(); ++LightId)
 	{
 		const FVisibleLightInfo& VisibleLightInfo = VisibleLightInfos[LightId];
-		for (const TSharedPtr<FVirtualShadowMapClipmap>& Clipmap : VisibleLightInfo.VirtualShadowMapClipmaps)
+
+		if (VisibleLightInfo.VirtualShadowMapClipmaps.Num() > 0)
 		{
+			// Directional light
 			for (int32 ViewIndex = 0; ViewIndex < Views.Num(); ViewIndex++)
 			{
-				if (Clipmap->GetDependentView() == Views[ViewIndex].GetPrimaryView())
+				TSharedPtr<FVirtualShadowMapClipmap> Clipmap = VisibleLightInfo.FindVirtualShadowMapShadowClipmapForView(&Views[ViewIndex]);
+				if (Clipmap.IsValid())
 				{
 					VisualizeLight[ViewIndex].CheckLight(Clipmap->GetLightSceneInfo().Proxy, Clipmap->GetVirtualShadowMapId());
 				}
 			}
 		}
-
-		for (FProjectedShadowInfo* ProjectedShadowInfo : VisibleLightInfo.AllProjectedShadows)
+		else
 		{
-			// NOTE: Specifically checking the VirtualShadowMapId vs HasVirtualShadowMap() here as clipmaps are handled above
-			if (ProjectedShadowInfo->VirtualShadowMapId != INDEX_NONE)
+			// Local light
+			int32 VirtualShadowMapId = VisibleLightInfo.GetVirtualShadowMapId(&Views[0]);
+			if (VirtualShadowMapId != INDEX_NONE)
 			{
-				check(ProjectedShadowInfo->CascadeSettings.ShadowSplitIndex == INDEX_NONE);		// We use clipmaps for virtual shadow maps, not cascades
-
-				// NOTE: Virtual shadow maps are never atlased, but verify our assumptions
-				{
-					const FVector4f ClipToShadowUV = ProjectedShadowInfo->GetClipToShadowBufferUvScaleBias();
-					check(ProjectedShadowInfo->BorderSize == 0);
-					check(ProjectedShadowInfo->X == 0);
-					check(ProjectedShadowInfo->Y == 0);
-					const FIntRect ShadowViewRect = ProjectedShadowInfo->GetInnerViewRect();
-					check(ShadowViewRect.Min.X == 0);
-					check(ShadowViewRect.Min.Y == 0);
-					check(ShadowViewRect.Max.X == FVirtualShadowMap::VirtualMaxResolutionXY);
-					check(ShadowViewRect.Max.Y == FVirtualShadowMap::VirtualMaxResolutionXY);
-				}
-
-				for (int32 ViewIndex = 0; ViewIndex < Views.Num(); ViewIndex++)
+				// NOTE: Assumes all the light scene infos for the different projections are the same!
+				if (VisibleLightInfo.AllProjectedShadows.Num() > 0)
 				{
-					VisualizeLight[ViewIndex].CheckLight(ProjectedShadowInfo->GetLightSceneInfo().Proxy, ProjectedShadowInfo->VirtualShadowMapId);
+					const FLightSceneProxy* Proxy = VisibleLightInfo.AllProjectedShadows[0]->GetLightSceneInfo().Proxy;
+					for (int32 ViewIndex = 0; ViewIndex < Views.Num(); ViewIndex++)
+					{
+						VisualizeLight[ViewIndex].CheckLight(Proxy, VirtualShadowMapId);
+					}
 				}
 			}
 		}
@@ -1840,11 +1825,6 @@ void FVirtualShadowMapArray::UpdatePhysicalPageAddresses(FRDGBuilder& GraphBuild
 		return;
 	}
 
-	// First, let the cache manager update any that may not be referenced this frame but may still have cached pages.
-	// TODO: Store the number of active lights we have first this frame for GPU looping purposes.
-	// By construction unreferenced lights are at the end.
-	CacheManager->UpdateUnreferencedCacheEntries(*this);
-
 	// NOTE: This past MUST run on all GPUs, as we still need to propogate changes to the VSM IDs even if
 	// a given GPU may not do any rendering during this phase.
 	RDG_GPU_MASK_SCOPE(GraphBuilder, FRHIGPUMask::All());
@@ -1975,10 +1955,14 @@ void FVirtualShadowMapArray::BuildPageAllocations(
 		for (auto It = CacheManager->CreateConstEntryIterator(); It; ++It)
 		{
 			const TSharedPtr<FVirtualShadowMapPerLightCacheEntry>& CacheEntry = It.Value();
-			for (const auto& Entry : CacheEntry->ShadowMapEntries)
+			const int32 VirtualShadowMapId = CacheEntry->GetVirtualShadowMapId();
+			check(VirtualShadowMapId >= 0);
+			for (int32 Index = 0; Index < CacheEntry->ShadowMapEntries.Num(); ++Index)
 			{
-				Uploader.Add(Entry.CurrentVirtualShadowMapId, &Entry.ProjectionData);
-				Builder.Add(Entry.CurrentVirtualShadowMapId, Entry.ProjectionData.MinMipLevel);
+				const auto& Entry = CacheEntry->ShadowMapEntries[Index];
+				int32 EntryVirtualShadowMapId = VirtualShadowMapId + Index;
+				Uploader.Add(EntryVirtualShadowMapId, &Entry.ProjectionData);
+				Builder.Add(EntryVirtualShadowMapId, Entry.ProjectionData.MinMipLevel);
 			}
 		}
 		Uploader.ResourceUploadTo(GraphBuilder, ProjectionDataRDG);
@@ -4193,12 +4177,11 @@ void FVirtualShadowMapArray::AddRenderViewsClipmap(
 
 	Nanite::SetCullingViewOverrides(CullingView, BaseParams);
 
+	const int32 VirtualShadowMapId = Clipmap->GetVirtualShadowMapId();
 	for (int32 ClipmapLevelIndex = 0; ClipmapLevelIndex < Clipmap->GetLevelCount(); ++ClipmapLevelIndex)
 	{
-		int32 VirtualShadowMapId = Clipmap->GetVirtualShadowMapId(ClipmapLevelIndex);
-
 		Nanite::FPackedViewParams Params = BaseParams;
-		Params.TargetLayerIndex = VirtualShadowMapId;
+		Params.TargetLayerIndex = VirtualShadowMapId + ClipmapLevelIndex;
 		Params.ViewMatrices = Clipmap->GetViewMatrices(ClipmapLevelIndex);
 		Params.PrevTargetLayerIndex = INDEX_NONE;
 		Params.PrevViewMatrices = Params.ViewMatrices;
@@ -4277,30 +4260,24 @@ void FVirtualShadowMapArray::AddRenderViewsLocal(
 	BaseParams.Flags |= CacheEntry->ShouldUseReceiverMask() ? NANITE_VIEW_FLAG_USE_RECEIVER_MASK : 0U;
 
 	Nanite::SetCullingViewOverrides(&Views[ClosestCullingViewIndex], BaseParams);
-	int32 NumMaps = ProjectedShadowInfo->bOnePassPointLightShadow ? 6 : 1;
-	for (int32 Index = 0; Index < NumMaps; ++Index)
+	for (int32 Index = 0; Index < CacheEntry->ShadowMapEntries.Num(); ++Index)
 	{
-		int32 VirtualShadowMapId = ProjectedShadowInfo->VirtualShadowMapId + Index;
+		FVirtualShadowMapCacheEntry& LevelEntry = CacheEntry->ShadowMapEntries[Index];
 
 		Nanite::FPackedViewParams Params = BaseParams;
-		Params.TargetLayerIndex = VirtualShadowMapId;
+		Params.TargetLayerIndex = CacheEntry->GetVirtualShadowMapId() + Index;
 		Params.ViewMatrices = ProjectedShadowInfo->GetShadowDepthRenderingViewMatrices(Index, true);
 		Params.RangeBasedCullingDistance = ProjectedShadowInfo->GetLightSceneInfo().Proxy->GetRadius();
-
-		if (CacheEntry)
-		{
-			FVirtualShadowMapCacheEntry& LevelEntry = CacheEntry->ShadowMapEntries[Index];
 			
-			if (bSetHZBParams)
-			{
-				LevelEntry.SetHZBViewParams(Params);
-			}
+		if (bSetHZBParams)
+		{
+			LevelEntry.SetHZBViewParams(Params);
+		}
 
-			// If we're going to generate a new HZB this frame, save the associated metadata
-			if (bUpdateHZBMetaData)
-			{
-				LevelEntry.UpdateHZBMetadata(Params.ViewMatrices, Params.ViewRect, Params.TargetLayerIndex);
-			}
+		// If we're going to generate a new HZB this frame, save the associated metadata
+		if (bUpdateHZBMetaData)
+		{
+			LevelEntry.UpdateHZBMetadata(Params.ViewMatrices, Params.ViewRect, Params.TargetLayerIndex);
 		}
 
 		OutVirtualShadowViews.Add(CreateNanitePackedView(Params));
diff --git a/Engine/Source/Runtime/Renderer/Private/VirtualShadowMaps/VirtualShadowMapArray.h b/Engine/Source/Runtime/Renderer/Private/VirtualShadowMaps/VirtualShadowMapArray.h
--- a/Engine/Source/Runtime/Renderer/Private/VirtualShadowMaps/VirtualShadowMapArray.h
+++ b/Engine/Source/Runtime/Renderer/Private/VirtualShadowMaps/VirtualShadowMapArray.h
@@ -289,10 +289,9 @@ public:
 
 	// Returns the first in a continuously allocated range of new VirtualShadowMapIds
 	int32 Allocate(bool bSinglePageShadowMap, int32 Count);
-
-	// TODO: Can probably make this 1:1 with allocate directly
-	void UpdateNextData(int32 PrevVirtualShadowMapId, int32 CurrentVirtualShadowMapId, FInt32Point PageOffset = FInt32Point(0, 0));
-
+	
+	void UpdateNextData(int32 PrevVirtualShadowMapId, const FNextVirtualShadowMapData& Data);
+	
 	static bool IsSinglePage(int VirtualShadowMapId)
 	{
 		return (VirtualShadowMapId < VSM_MAX_SINGLE_PAGE_SHADOW_MAPS);
diff --git a/Engine/Source/Runtime/Renderer/Private/VirtualShadowMaps/VirtualShadowMapCacheManager.cpp b/Engine/Source/Runtime/Renderer/Private/VirtualShadowMaps/VirtualShadowMapCacheManager.cpp
--- a/Engine/Source/Runtime/Renderer/Private/VirtualShadowMaps/VirtualShadowMapCacheManager.cpp
+++ b/Engine/Source/Runtime/Renderer/Private/VirtualShadowMaps/VirtualShadowMapCacheManager.cpp
@@ -54,14 +54,6 @@ static TAutoConsoleVariable<int32> CVarCacheVsmUseHzb(
 	TEXT(" When enabled, instances invalidations are tested against the HZB. Instances that are fully occluded will not cause page invalidations."),
 	ECVF_RenderThreadSafe);
 
-int32 GClipmapPanning = 1;
-FAutoConsoleVariableRef CVarEnableClipmapPanning(
-	TEXT("r.Shadow.Virtual.Cache.ClipmapPanning"),
-	GClipmapPanning,
-	TEXT("Enable support for panning cached clipmap pages for directional lights, allowing re-use of cached data when the camera moves. Keep this enabled outside of debugging."),
-	ECVF_RenderThreadSafe
-);
-
 static int32 GVSMCacheDeformableMeshesInvalidate = 1;
 FAutoConsoleVariableRef CVarCacheDeformableMeshesInvalidate(
 	TEXT("r.Shadow.Virtual.Cache.DeformableMeshesInvalidate"),
@@ -203,31 +195,20 @@ namespace Nanite
 }
 
 void FVirtualShadowMapCacheEntry::UpdateClipmapLevel(
-	FVirtualShadowMapArray& VirtualShadowMapArray,
 	const FVirtualShadowMapPerLightCacheEntry& PerLightEntry,
-	int32 VirtualShadowMapId,
 	FInt64Point PageSpaceLocation,
 	double LevelRadius,
 	double ViewCenterZ,
 	double ViewRadiusZ,
 	double WPODistanceDisableThresholdSquared)
 {
-	int32 PrevVirtualShadowMapId = CurrentVirtualShadowMapId;
-	FInt64Point PrevPageSpaceLocation = Clipmap.PageSpaceLocation;
 	UpdatePrevHZBMetadata();
 
-	bool bCacheValid = (PrevVirtualShadowMapId != INDEX_NONE);
+	// Not valid if it was never rendered
+	bool bCacheValid = (PerLightEntry.RenderedFrameNumber >= 0);
 	
-	if (bCacheValid && GClipmapPanning == 0)
-	{
-		if (PageSpaceLocation.X != PrevPageSpaceLocation.X ||
-			PageSpaceLocation.Y != PrevPageSpaceLocation.Y)
-		{
-			bCacheValid = false;
-			//UE_LOG(LogRenderer, Display, TEXT("Invalidated clipmap level (VSM %d) with page space location %d,%d (Prev %d, %d)"),
-			//	VirtualShadowMapId, PageSpaceLocation.X, PageSpaceLocation.Y, PrevPageSpaceLocation.X, PrevPageSpaceLocation.Y);
-		}
-	}
+	// Not valid if radius has changed
+	bCacheValid = bCacheValid && (ViewRadiusZ == Clipmap.ViewRadiusZ);
 
 	// Invalidate if the new Z radius strayed too close/outside the guardband of the cached shadow map
 	if (bCacheValid)
@@ -240,72 +221,53 @@ void FVirtualShadowMapCacheEntry::UpdateClipmapLevel(
 		}
 	}
 
-	// Not valid if it was never rendered
-	bCacheValid = bCacheValid && (PerLightEntry.Prev.RenderedFrameNumber >= 0);
-
-	// Not valid if radius has changed
-	bCacheValid = bCacheValid && (ViewRadiusZ == Clipmap.ViewRadiusZ);
-
 	// Not valid if WPO threshold has changed
 	if (bCacheValid && CVarClipmapWPODisableDistanceInvalidate.GetValueOnRenderThread() != 0 && (WPODistanceDisableThresholdSquared != Clipmap.WPODistanceDisableThresholdSquared))
 	{
 		bCacheValid = false;
 		// Only warn once per change... when this changes it will hit all of them
-		if (PerLightEntry.ShadowMapEntries.Num() > 0 && PerLightEntry.ShadowMapEntries[0].CurrentVirtualShadowMapId == VirtualShadowMapId)
+		if (&PerLightEntry.ShadowMapEntries[0] == this)
 		{
-			UE_LOG(LogRenderer, Display, TEXT("Invalidated clipmap due to WPO threshold change. This can occur due to resolution or FOV changes."), VirtualShadowMapId);
+			UE_LOG(LogRenderer, Display, TEXT("Invalidated VSM clipmap due to WPO threshold change. This can occur due to resolution or FOV changes."));
 		}
 	}
 
-	if (!bCacheValid)
+	if (bCacheValid)
 	{
-		Clipmap.ViewCenterZ = ViewCenterZ;
-		Clipmap.ViewRadiusZ = ViewRadiusZ;
-		Clipmap.WPODistanceDisableThresholdSquared = WPODistanceDisableThresholdSquared;
+		// NOTE: Leave the view center and radius where they were previously for the cached page
+		FInt64Point CurrentToPreviousPageOffset(PageSpaceLocation - Clipmap.PageSpaceLocation);
+		NextData.Flags = VSM_NEXT_FLAG_VALID;
+		NextData.PageAddressOffset = FIntVector2(CurrentToPreviousPageOffset.X, CurrentToPreviousPageOffset.Y);
 	}
 	else
 	{
-		// NOTE: Leave the view center and radius where they were previously for the cached page
-		FInt64Point CurrentToPreviousPageOffset(PageSpaceLocation - PrevPageSpaceLocation);
-		VirtualShadowMapArray.UpdateNextData(PrevVirtualShadowMapId, VirtualShadowMapId, FInt32Point(CurrentToPreviousPageOffset));
+		NextData.Flags = 0;
+		Clipmap.ViewCenterZ = ViewCenterZ;
+		Clipmap.ViewRadiusZ = ViewRadiusZ;
+		Clipmap.WPODistanceDisableThresholdSquared = WPODistanceDisableThresholdSquared;		
 	}
 	
-	CurrentVirtualShadowMapId = VirtualShadowMapId;
 	Clipmap.PageSpaceLocation = PageSpaceLocation;
 }
 
 void FVirtualShadowMapCacheEntry::Update(
-	FVirtualShadowMapArray& VirtualShadowMapArray,
-	const FVirtualShadowMapPerLightCacheEntry& PerLightEntry,
-	int32 VirtualShadowMapId)
+	const FVirtualShadowMapPerLightCacheEntry& PerLightEntry)
 {
-	// Swap previous frame data over.
-	int32 PrevVirtualShadowMapId = CurrentVirtualShadowMapId;
 	UpdatePrevHZBMetadata();
 	
-	bool bCacheValid = (PrevVirtualShadowMapId != INDEX_NONE);
-
 	// Not valid if it was never rendered
-	bCacheValid = bCacheValid && (PerLightEntry.Prev.RenderedFrameNumber >= 0);
+	bool bCacheValid = (PerLightEntry.RenderedFrameNumber >= 0);
 
 	if (bCacheValid)
 	{
-		// Invalidate on transition between single page and full
-		bool bPrevSinglePage = FVirtualShadowMapArray::IsSinglePage(PrevVirtualShadowMapId);
-		bool bCurrentSinglePage = FVirtualShadowMapArray::IsSinglePage(VirtualShadowMapId);
-		if (bPrevSinglePage != bCurrentSinglePage)
-		{
-			bCacheValid = false;
-		}
+		NextData.Flags = VSM_NEXT_FLAG_VALID;
+		NextData.PageAddressOffset = FInt32Point(0, 0);
 	}
-
-	if (bCacheValid)
+	else
 	{
-		// Update previous/next frame mapping if we have a valid cached shadow map
-		VirtualShadowMapArray.UpdateNextData(PrevVirtualShadowMapId, VirtualShadowMapId);
+		NextData.Flags = 0;
 	}
 
-	CurrentVirtualShadowMapId = VirtualShadowMapId;
 	// Current HZB metadata gets updated during rendering
 }
 
@@ -340,22 +302,19 @@ void FVirtualShadowMapPerLightCacheEntry::UpdateClipmap(
 	bool bForceInvalidate,
 	bool bInUseReceiverMask)
 {
-	Prev.RenderedFrameNumber = FMath::Max(Prev.RenderedFrameNumber, Current.RenderedFrameNumber);
-	Current.RenderedFrameNumber = -1;
-
 	if (bForceInvalidate ||
 		LightDirection != ClipmapCacheKey.LightDirection ||
 		FirstLevel != ClipmapCacheKey.FirstLevel)
 	{
-		Prev.RenderedFrameNumber = -1;
+		RenderedFrameNumber = -1;
+		ClipmapCacheKey.LightDirection = LightDirection;
+		ClipmapCacheKey.FirstLevel = FirstLevel;
 	}
-	ClipmapCacheKey.LightDirection = LightDirection;
-	ClipmapCacheKey.FirstLevel = FirstLevel;
 
 	// If we swap receiver mask on/off we also may need to invalidate due to incomplete pages
 	if (bInUseReceiverMask != bUseReceiverMask)
 	{
-		Prev.RenderedFrameNumber = -1;
+		RenderedFrameNumber = -1;
 		bUseReceiverMask = bInUseReceiverMask;
 	}
 
@@ -365,17 +324,25 @@ void FVirtualShadowMapPerLightCacheEntry::UpdateClipmap(
 	// frame though we will swap back so that we can begin establishing static cache data. Thus it is still
 	// useful to explicitly set ForceInvalidateDirectional in cases where the light is invalidating frequently
 	// but not every single frame to keep the performance consistent.
-	bool bNewIsUncached = Prev.RenderedFrameNumber < 0;
+	bool bNewIsUncached = RenderedFrameNumber < 0;
 
 	// On transition between uncached <-> cached we must invalidate since the static pages may not be initialized
 	if (bNewIsUncached != bIsUncached)
 	{
-		Prev.RenderedFrameNumber = -1;
+		RenderedFrameNumber = -1;
 		bIsUncached = bNewIsUncached;
 	}
 
+	// Not used by clipmaps
 	LightOrigin = FVector(0, 0, 0);
 	LightRadius = -1.0f;
+
+	// On invalidate, wipe out any existing VSM IDs
+	if (RenderedFrameNumber < 0)
+	{
+		PrevVirtualShadowMapId = INDEX_NONE;
+		VirtualShadowMapId = INDEX_NONE;
+	}
 }
 
 void FVirtualShadowMapPerLightCacheEntry::UpdateLocal(
@@ -388,10 +355,6 @@ void FVirtualShadowMapPerLightCacheEntry::UpdateLocal(
 	bool bInUseReceiverMask)
 {
 	// TODO: The logic in this function is needlessly convoluted... clean up
-
-	Prev.RenderedFrameNumber = FMath::Max(Prev.RenderedFrameNumber, Current.RenderedFrameNumber);
-	Prev.ScheduledFrameNumber = FMath::Max(Prev.ScheduledFrameNumber, Current.ScheduledFrameNumber);
-
 	const bool bLightMoved = 
 		LocalCacheKey.PreShadowTranslation != InCacheKey.PreShadowTranslation ||
 		LocalCacheKey.WorldToLight != InCacheKey.WorldToLight;
@@ -403,36 +366,46 @@ void FVirtualShadowMapPerLightCacheEntry::UpdateLocal(
 	// If it is a distant light, we want to let the time-share perform the invalidation.
 	if (bForceInvalidate || (bAllowInvalidation && bLightMoved))
 	{
-		Prev.RenderedFrameNumber = -1;
-		//UE_LOG(LogRenderer, Display, TEXT("Invalidated!"));
+		RenderedFrameNumber = -1;
 	}
 
 	// If we swap receiver mask on/off we also may need to invalidate due to incomplete pages
 	if (bInUseReceiverMask != bUseReceiverMask)
 	{
-		Prev.RenderedFrameNumber = -1;
+		RenderedFrameNumber = -1;
 		bUseReceiverMask = bInUseReceiverMask;
 	}
 
 	// On transition between uncached <-> cached we must invalidate since the static pages may not be initialized
-	bool bNewIsUncached = Prev.RenderedFrameNumber < 0;
+	bool bNewIsUncached = RenderedFrameNumber < 0;
 	if (bNewIsUncached != bIsUncached)
 	{
-		Prev.RenderedFrameNumber = -1;
+		RenderedFrameNumber = -1;
+		bIsUncached = bNewIsUncached;
 	}
 
 	// On transition between distant <-> regular we must invalidate
 	if (bNewIsDistantLight != bIsDistantLight)
 	{
-		Prev.RenderedFrameNumber = -1;
+		RenderedFrameNumber = -1;
+		bIsDistantLight = bNewIsDistantLight;
 	}
 
-	Current.RenderedFrameNumber = -1;
-	Current.ScheduledFrameNumber = -1;
-	bIsDistantLight = bNewIsDistantLight;
-	bIsUncached = bNewIsUncached;
 	LightOrigin = NewLightOrigin;
 	LightRadius = NewLightRadius;
+
+	// On invalidate, wipe out any existing VSM IDs
+	if (RenderedFrameNumber < 0)
+	{
+		PrevVirtualShadowMapId = INDEX_NONE;
+		VirtualShadowMapId = INDEX_NONE;
+	}
+}
+
+void FVirtualShadowMapPerLightCacheEntry::UpdateVirtualShadowMapId(int32 NextVirtualShadowMapId)
+{
+	PrevVirtualShadowMapId = VirtualShadowMapId;
+	VirtualShadowMapId = NextVirtualShadowMapId;
 }
 
 void FVirtualShadowMapArrayCacheManager::FShadowInvalidatingInstancesImplementation::AddPrimitive(const FPrimitiveSceneInfo* PrimitiveSceneInfo)
@@ -593,6 +566,12 @@ void FVirtualShadowMapArrayCacheManager::FInvalidatingPrimitiveCollector::AddInv
 	bool bLightRadiusCulling,
 	const FBoxSphereBounds& PrimitiveBounds)
 {
+	int32 VirtualShadowMapId = CacheEntry.GetVirtualShadowMapId();
+	if (VirtualShadowMapId < 0)
+	{
+		return;
+	}
+
 	// We don't need explicit invalidations for force invalidated/uncached lights
 	if (CacheEntry.IsUncached())
 	{
@@ -609,10 +588,12 @@ void FVirtualShadowMapArrayCacheManager::FInvalidatingPrimitiveCollector::AddInv
 		return;
 	}
 	
-	// Add item for each shadow map explicitly, inflates host data but improves load balancing,
-	for (const auto& SmCacheEntry : CacheEntry.ShadowMapEntries)
+	// Add item for each shadow map explicitly, inflates host data but improves load balancing
+	const int32 NumEntries = CacheEntry.ShadowMapEntries.Num();
+	for (int32 Index = 0; Index < NumEntries; ++Index)
 	{
-		Instances.Add(InstanceSceneDataOffset, NumInstanceSceneDataEntries, EncodeInstanceInvalidationPayload(SmCacheEntry.CurrentVirtualShadowMapId));
+		Instances.Add(InstanceSceneDataOffset, NumInstanceSceneDataEntries,
+					EncodeInstanceInvalidationPayload(VirtualShadowMapId + Index));
 	}
 }
 
@@ -1208,25 +1189,25 @@ void FVirtualShadowMapArrayCacheManager::UpdateUnreferencedCacheEntries(
 
 	for (FEntryMap::TIterator It = CacheEntries.CreateIterator(); It; ++It)
 	{
+		bool bKeepEntry = true;
+
 		TSharedPtr<FVirtualShadowMapPerLightCacheEntry> CacheEntry = (*It).Value;
 		// For this test we care if it is active *this render*, not just this scene frame number (which can include multiple renders)
 		if (CacheEntry->bReferencedThisRender)
 		{
 			// Active this render, leave it alone
-			check(CacheEntry->ShadowMapEntries.Last().CurrentVirtualShadowMapId < VirtualShadowMapArray.GetNumShadowMapSlots());
+			check((CacheEntry->GetVirtualShadowMapId() + CacheEntry->ShadowMapEntries.Num()) <= VirtualShadowMapArray.GetNumShadowMapSlots());
 		}
 		else if (int32(SceneFrameNumber - CacheEntry->LastReferencedFrameNumber) <= MaxLightAge)
 		{
 			// Not active this render, but still recent enough to keep it and its pages alive
-			int PrevBaseVirtualShadowMapId = CacheEntry->ShadowMapEntries[0].CurrentVirtualShadowMapId;
-			bool bIsSinglePage = FVirtualShadowMapArray::IsSinglePage(PrevBaseVirtualShadowMapId);
-
-			// Keep the entry, reallocate new VSM IDs
+			// Reallocate new VSM ID for this frame
 			int32 NumMaps = CacheEntry->ShadowMapEntries.Num();
-			int32 VirtualShadowMapId = VirtualShadowMapArray.Allocate(bIsSinglePage, NumMaps);
+			const int32 NextVirtualShadowMapId = VirtualShadowMapArray.Allocate(CacheEntry->bIsDistantLight, NumMaps);
+			CacheEntry->UpdateVirtualShadowMapId(NextVirtualShadowMapId);
 			for (int32 Map = 0; Map < NumMaps; ++Map)
 			{
-				CacheEntry->ShadowMapEntries[Map].Update(VirtualShadowMapArray, *CacheEntry, VirtualShadowMapId + Map);
+				CacheEntry->ShadowMapEntries[Map].Update(*CacheEntry);
 				// Mark it as inactive for this frame/render
 				// NOTE: We currently recompute/overwrite the whole ProjectionData structure for referenced lights, but if that changes we
 				// will need to clear this flag again when they become referenced.
@@ -1234,6 +1215,25 @@ void FVirtualShadowMapArrayCacheManager::UpdateUnreferencedCacheEntries(
 			}
 		}
 		else
+		{
+			bKeepEntry = false;
+		}
+
+		if (bKeepEntry)
+		{
+			// Update next data (for each entry)
+			const int32 PrevVirtualShadowMapId = CacheEntry->GetPrevVirtualShadowMapId();
+			if (PrevVirtualShadowMapId >= 0)
+			{
+				for (int32 Index = 0; Index < CacheEntry->ShadowMapEntries.Num(); ++Index)
+				{
+					auto& NextData = CacheEntry->ShadowMapEntries[Index].NextData;
+					NextData.NextVirtualShadowMapId = CacheEntry->GetVirtualShadowMapId() + Index;
+					VirtualShadowMapArray.UpdateNextData(PrevVirtualShadowMapId + Index, NextData);
+				}
+			}
+		}
+		else
 		{
 			It.RemoveCurrent();
 		}
@@ -1663,9 +1663,12 @@ void FVirtualShadowMapArrayCacheManager::UpdateCachePrimitiveAsDynamic(FInvalida
 				// Add an invalidation for every light
 				for (auto& CacheEntry : CacheEntries)
 				{
-					for (const auto& SmCacheEntry : CacheEntry.Value->ShadowMapEntries)
+					const int32 VirtualShadowMapId = CacheEntry.Value->GetVirtualShadowMapId();
+					check(VirtualShadowMapId >= 0);
+					const int32 NumEntries = CacheEntry.Value->ShadowMapEntries.Num();
+					for (int32 Index = 0; Index < NumEntries; ++Index)
 					{
-						const uint32 PayloadForceStatic = EncodeInstanceInvalidationPayload(SmCacheEntry.CurrentVirtualShadowMapId, VSM_INVALIDATION_PAYLOAD_FLAG_FORCE_STATIC);
+						const uint32 PayloadForceStatic = EncodeInstanceInvalidationPayload(VirtualShadowMapId + Index, VSM_INVALIDATION_PAYLOAD_FLAG_FORCE_STATIC);
 						InvalidatingPrimitiveCollector.Instances.Add(
 							PrimitiveSceneInfo->GetInstanceSceneDataOffset(),
 							PrimitiveSceneInfo->GetNumInstanceSceneDataEntries(),
diff --git a/Engine/Source/Runtime/Renderer/Private/VirtualShadowMaps/VirtualShadowMapCacheManager.h b/Engine/Source/Runtime/Renderer/Private/VirtualShadowMaps/VirtualShadowMapCacheManager.h
--- a/Engine/Source/Runtime/Renderer/Private/VirtualShadowMaps/VirtualShadowMapCacheManager.h
+++ b/Engine/Source/Runtime/Renderer/Private/VirtualShadowMaps/VirtualShadowMapCacheManager.h
@@ -43,17 +43,11 @@ class FVirtualShadowMapCacheEntry
 {
 public:
 	// Generic version used for local lights but also inactive lights
-	// Updates the VSM ID
-	void Update(
-		FVirtualShadowMapArray& VirtualShadowMapArray,
-		const FVirtualShadowMapPerLightCacheEntry &PerLightEntry,
-		int32 VirtualShadowMapId);
+	void Update(const FVirtualShadowMapPerLightCacheEntry &PerLightEntry);
 
 	// Specific version of the above for clipmap levels, which have additional constraints
 	void UpdateClipmapLevel(
-		FVirtualShadowMapArray& VirtualShadowMapArray,
 		const FVirtualShadowMapPerLightCacheEntry& PerLightEntry,
-		int32 VirtualShadowMapId,
 		FInt64Point PageSpaceLocation,
 		double LevelRadius,
 		double ViewCenterZ,
@@ -74,12 +68,11 @@ public:
 		}
 	}
 
-	// Previous frame data
-	FVirtualShadowMapHZBMetadata PrevHZBMetadata;
-
-	// Current frame data
-	int32 CurrentVirtualShadowMapId = INDEX_NONE;
+	FVirtualShadowMapHZBMetadata PrevHZBMetadata;	
 	FVirtualShadowMapHZBMetadata CurrentHZBMetadata;
+	
+	// Tracks mapping to any previously cached data
+	FNextVirtualShadowMapData NextData;
 
 	// Stores the projection shader data. This is needed for cached entries that may be inactive in the current frame/render
 	// and also avoids recomputing it every frame.
@@ -106,29 +99,34 @@ public:
 	}
 
 	void OnPrimitiveRendered(const FPrimitiveSceneInfo* PrimitiveSceneInfo, bool bPrimitiveRevealed);
+	
 	/**
 	 * The (local) VSM is fully cached if it is distant and has been rendered to previously
 	 * "Fully" implies that we know all pages are mapped as well as rendered to (ignoring potential CPU-side object culling).
 	 */
-	inline bool IsFullyCached() const { return bIsDistantLight && Prev.RenderedFrameNumber >= 0; }
+	inline bool IsFullyCached() const { return bIsDistantLight && RenderedFrameNumber >= 0; }
 
-	/**
-	 */
 	inline bool IsUncached() const { return bIsUncached; }
 
 	inline bool ShouldUseReceiverMask() const { return bUseReceiverMask; }
 
-	void MarkRendered(int32 FrameIndex) { Current.RenderedFrameNumber = FrameIndex; }
-	int32 GetLastScheduledFrameNumber() const { return Prev.ScheduledFrameNumber; }
+	void MarkRendered(int32 FrameIndex) { RenderedFrameNumber = FrameIndex; }
+	void Invalidate() { RenderedFrameNumber = -1; }
+	bool IsInvalidated() const { return RenderedFrameNumber < 0; }
+
+	void MarkScheduled(int32 FrameIndex) { ScheduledFrameNumber = FrameIndex; }
+	int32 GetLastScheduledFrameNumber() const { return ScheduledFrameNumber; }
+	
+	void UpdateVirtualShadowMapId(int32 NextVirtualShadowMapId);
+	int32 GetVirtualShadowMapId() const { return VirtualShadowMapId; }
+	int32 GetPrevVirtualShadowMapId() const { return PrevVirtualShadowMapId; }
+
 	void UpdateClipmap(
 		const FVector& LightDirection,
 		int FirstLevel,
 		bool bForceInvalidate,
 		bool bInUseReceiverMask);
 
-	/**
-	 * Returns true if the cache entry is valid (has previous state).
-	 */
 	void UpdateLocal(
 		const FProjectedShadowInitializer &InCacheKey,
 		const FVector& NewLightOrigin,
@@ -138,24 +136,18 @@ public:
 		bool bAllowInvalidation,
 		bool bInUseReceiverMask);
 
-	void Invalidate() { Prev.RenderedFrameNumber = -1; }
-
-	bool IsInvalidated() const { return Prev.RenderedFrameNumber < 0; }
-
 	bool AffectsBounds(const FBoxSphereBounds& Bounds) const
 	{
 		return (LightRadius <= 0.0f) ||			// Infinite extent light (directional, etc)
 			((Bounds.Origin - LightOrigin).SizeSquared() <= FMath::Square(LightRadius + Bounds.SphereRadius));
 	}
 
-	// TODO: We probably don't need the prev/next thing anymore
-	struct FFrameState
-	{
-		int32 RenderedFrameNumber = -1;
-		int32 ScheduledFrameNumber = -1;
-	};
-	FFrameState Prev;
-	FFrameState Current;
+	// Last frame numbers that we rendered/scheduled this light
+	int32 RenderedFrameNumber = -1;
+	int32 ScheduledFrameNumber = -1;
+
+	int32 PrevVirtualShadowMapId = INDEX_NONE;
+	int32 VirtualShadowMapId = INDEX_NONE;
 
 	bool bIsUncached = false;
 	bool bIsDistantLight = false;
@@ -178,10 +170,6 @@ public:
 
 	TArray<FVirtualShadowMapInstanceRange> PrimitiveInstancesToInvalidate;
 
-	// Rough bounds for invalidation culling
-	FVector LightOrigin = FVector(0, 0, 0);
-	float LightRadius = -1.0f;		// Negative means infinite
-
 private:
 	struct FLocalLightCacheKey
 	{
@@ -197,6 +185,10 @@ private:
 		int LevelCount;
 	};
 	FClipmapCacheKey ClipmapCacheKey;
+
+	// Rough bounds for invalidation culling
+	FVector LightOrigin = FVector(0, 0, 0);
+	float LightRadius = -1.0f;		// Negative means infinite
 };
 
 class FVirtualShadowMapFeedback
diff --git a/Engine/Source/Runtime/Renderer/Private/VirtualShadowMaps/VirtualShadowMapClipmap.cpp b/Engine/Source/Runtime/Renderer/Private/VirtualShadowMaps/VirtualShadowMapClipmap.cpp
--- a/Engine/Source/Runtime/Renderer/Private/VirtualShadowMaps/VirtualShadowMapClipmap.cpp
+++ b/Engine/Source/Runtime/Renderer/Private/VirtualShadowMaps/VirtualShadowMapClipmap.cpp
@@ -309,8 +309,6 @@ FVirtualShadowMapClipmap::FVirtualShadowMapClipmap(
 	LevelData.Empty();
 	LevelData.AddDefaulted(LevelCount);
 
-	VirtualShadowMapId = VirtualShadowMapArray.Allocate(false, LevelCount);
-
 	// Enable various fast paths if we are not storing cache data
 	bool bUncached = Config.bForceInvalidate || !VirtualShadowMapArrayCacheManager->IsCacheEnabled();
 
@@ -339,10 +337,10 @@ FVirtualShadowMapClipmap::FVirtualShadowMapClipmap(
 	const double ViewRadiusZScale = CVarVirtualShadowMapClipmapZRangeScale.GetValueOnRenderThread();
 	const FVector ViewCenter = WorldToLightViewRotationMatrix.TransformPosition(WorldOrigin);
 
-	for (int32 Index = 0; Index < LevelCount; ++Index)
+	for (int32 ClipmapIndex = 0; ClipmapIndex < LevelCount; ++ClipmapIndex)
 	{
-		FLevelData& Level = LevelData[Index];
-		const int32 AbsoluteLevel = Index + FirstLevel;		// Absolute (virtual) level index
+		FLevelData& Level = LevelData[ClipmapIndex];
+		const int32 AbsoluteLevel = ClipmapIndex + FirstLevel;		// Absolute (virtual) level index
 
 		const double RawLevelRadius = GetLevelRadius(static_cast<float>(AbsoluteLevel));
 
@@ -378,7 +376,7 @@ FVirtualShadowMapClipmap::FVirtualShadowMapClipmap(
 
 		// Check if we have a cache entry for this level
 		// If we do and it covers our required depth range, we can use cached pages. Otherwise we need to invalidate.
-		FVirtualShadowMapCacheEntry* ClipmapLevelEntry = &PerLightCacheEntry->ShadowMapEntries[Index];
+		FVirtualShadowMapCacheEntry* ClipmapLevelEntry = &PerLightCacheEntry->ShadowMapEntries[ClipmapIndex];
 
 		double ViewRadiusZ = RawLevelRadius * ViewRadiusZScale;
 		double ViewCenterDeltaZ = 0.0f;
@@ -411,9 +409,7 @@ FVirtualShadowMapClipmap::FVirtualShadowMapClipmap(
 		}
 
 		ClipmapLevelEntry->UpdateClipmapLevel(
-			VirtualShadowMapArray,
 			*PerLightCacheEntry,
-			GetVirtualShadowMapId(Index),
 			PageOffset,
 			RawLevelRadius,
 			ViewCenter.Z,
@@ -456,7 +452,7 @@ FVirtualShadowMapClipmap::FVirtualShadowMapClipmap(
 			}
 		}
 		// Update the entry with the new projection data
-		ClipmapLevelEntry->ProjectionData = ComputeProjectionShaderData(Index);
+		ClipmapLevelEntry->ProjectionData = ComputeProjectionShaderData(ClipmapIndex);
 	}
 
 	ComputeBoundingVolumes(WorldOrigin);
diff --git a/Engine/Source/Runtime/Renderer/Private/VirtualShadowMaps/VirtualShadowMapClipmap.h b/Engine/Source/Runtime/Renderer/Private/VirtualShadowMaps/VirtualShadowMapClipmap.h
--- a/Engine/Source/Runtime/Renderer/Private/VirtualShadowMaps/VirtualShadowMapClipmap.h
+++ b/Engine/Source/Runtime/Renderer/Private/VirtualShadowMaps/VirtualShadowMapClipmap.h
@@ -10,14 +10,13 @@ VirtualShadowMapClipmap.h
 #include "ConvexVolume.h"
 #include "Templates/RefCounting.h"
 #include "VirtualShadowMapProjection.h"
+#include "VirtualShadowMapCacheManager.h"
 
 struct FViewMatrices;
 struct FVirtualShadowMapProjectionShaderData;
 class FPrimitiveSceneInfo;
 class FVirtualShadowMap;
 class FVirtualShadowMapArray;
-class FVirtualShadowMapArrayCacheManager;
-class FVirtualShadowMapPerLightCacheEntry;
 
 bool IsVirtualShadowMapDirectionalReceiverMaskEnabled();
 
@@ -60,9 +59,9 @@ public:
 		return LevelData[ClipmapIndex].DynamicDepthCullRange;
 	}
 
-	int32 GetVirtualShadowMapId(int32 ClipmapIndex = 0) const
+	int32 GetVirtualShadowMapId() const
 	{
-		return VirtualShadowMapId + ClipmapIndex;
+		return PerLightCacheEntry->GetVirtualShadowMapId();
 	}
 
 	int32 GetLevelCount() const
@@ -175,7 +174,6 @@ private:
 	FSphere BoundingSphere;
 	FConvexVolume ViewFrustumBounds;
 
-	int32 VirtualShadowMapId = INDEX_NONE;		// Base ID; levels are contiguous
 	TSharedPtr<FVirtualShadowMapPerLightCacheEntry> PerLightCacheEntry;
 
 	// Rendered primitives are marked during culling (through OnPrimitiveRendered being called).
