Fix static analysis warning. Check View.Family pointer before accessing it.

diff --git a/Engine/Source/Runtime/Renderer/Private/PostProcess/PostProcessing.cpp b/Engine/Source/Runtime/Renderer/Private/PostProcess/PostProcessing.cpp
--- a/Engine/Source/Runtime/Renderer/Private/PostProcess/PostProcessing.cpp
+++ b/Engine/Source/Runtime/Renderer/Private/PostProcess/PostProcessing.cpp
@@ -555,7 +555,7 @@ void AddPostProcessingPasses(
 #endif
 	PassSequence.SetEnabled(EPass::VisualizeLumenScene, LumenVisualizeMode >= 0 && LumenVisualizeMode != VISUALIZE_MODE_OVERVIEW && LumenVisualizeMode != VISUALIZE_MODE_PERFORMANCE_OVERVIEW && bPostProcessingEnabled);
 	PassSequence.SetEnabled(EPass::VisualizeSubstrate, Substrate::ShouldRenderSubstrateDebugPasses(View));
-	PassSequence.SetEnabled(EPass::VisualizeLightGrid, ShouldVisualizeLightGrid(View.Family->GetShaderPlatform()));
+	PassSequence.SetEnabled(EPass::VisualizeLightGrid, View.Family && ShouldVisualizeLightGrid(View.Family->GetShaderPlatform()));
 
 #if WITH_EDITOR
 	PassSequence.SetEnabled(EPass::VisualizeSkyAtmosphere, Scene&& View.Family && View.Family->EngineShowFlags.VisualizeSkyAtmosphere && ShouldRenderSkyAtmosphereDebugPasses(Scene, View.Family->EngineShowFlags));
