Fix potential race condition with ISM LODUniformBuffers

diff --git a/Engine/Source/Runtime/Renderer/Private/RendererScene.cpp b/Engine/Source/Runtime/Renderer/Private/RendererScene.cpp
--- a/Engine/Source/Runtime/Renderer/Private/RendererScene.cpp
+++ b/Engine/Source/Runtime/Renderer/Private/RendererScene.cpp
@@ -177,6 +177,16 @@ static TAutoConsoleVariable<int32> CVarVisibilityLocalLightPrimitiveInteraction(
 	ECVF_Scalability | ECVF_RenderThreadSafe
 );
 
+bool GSafeCullDistanceUpdate = true;
+static FAutoConsoleVariableRef CVarSafeCullDistanceUpdate(
+	TEXT("r.Render.SafeCullDistanceUpdate"),
+	GSafeCullDistanceUpdate,
+	TEXT("Update primitive cull distances before static meshes are added to the scene.\n")
+	TEXT("0: Update while static meshes are potentially asynchronously being added\n")
+	TEXT("1: Force update before static meshes are added\n"),
+	ECVF_RenderThreadSafe
+);
+
 DECLARE_CYCLE_STAT(TEXT("DeferredShadingSceneRenderer MotionBlurStartFrame"), STAT_FDeferredShadingSceneRenderer_MotionBlurStartFrame, STATGROUP_SceneRendering);
 
 IMPLEMENT_GLOBAL_SHADER_PARAMETER_STRUCT(FDistanceCullFadeUniformShaderParameters, "PrimitiveFade");
@@ -6140,6 +6150,23 @@ void FScene::Update(FRDGBuilder& GraphBuilder, const FUpdateParameters& Paramete
 
 	const bool bAsyncCacheMeshDrawCommands = EnumHasAnyFlags(Parameters.AsyncOps, EUpdateAllPrimitiveSceneInfosAsyncOps::CacheMeshDrawCommands) && GRHISupportsMultithreadedShaderCreation;
 
+	auto UpdateCullDistances = [&UpdatedInstanceCullDistance]()
+	{
+		for (const auto& CullDistance : UpdatedInstanceCullDistance)
+		{
+			FPrimitiveSceneProxy* SceneProxy = CullDistance.SceneInfo->Proxy;
+			float StartCullDistance = CullDistance.Payload.Value.X;
+			float EndCullDistance = CullDistance.Payload.Value.Y;
+		
+			SceneProxy->SetInstanceCullDistance_RenderThread(StartCullDistance, EndCullDistance);
+		}
+	};
+
+	if (GSafeCullDistanceUpdate)
+	{
+		UpdateCullDistances();
+	}
+
 	UE::Tasks::FTask AddStaticMeshesTask = GraphBuilder.AddCommandListSetupTask(
 		[this, AddStaticMeshes = CopyTemp(SceneInfosWithStaticDrawListUpdate), SceneInfosWithFlushRuntimeVirtualTexture = MoveTemp(SceneInfosWithFlushRuntimeVirtualTexture), &SceneInfosWithAddToScene]
 			(FRHICommandListBase& RHICmdList) mutable
@@ -6366,13 +6393,9 @@ void FScene::Update(FRDGBuilder& GraphBuilder, const FUpdateParameters& Paramete
 		PrimitiveOcclusionBounds[SceneInfo->PackedIndex] = NewOccBounds.ExpandBy(OCCLUSION_SLOP + Delta);
 	}
 
-	for (const auto& CullDistance : UpdatedInstanceCullDistance)
+	if (!GSafeCullDistanceUpdate)
 	{
-		FPrimitiveSceneProxy* SceneProxy = CullDistance.SceneInfo->Proxy;
-		float StartCullDistance = CullDistance.Payload.Value.X;
-		float EndCullDistance = CullDistance.Payload.Value.Y;
-		
-		SceneProxy->SetInstanceCullDistance_RenderThread(StartCullDistance, EndCullDistance);
+		UpdateCullDistances();
 	}
 
 	for (const auto& DrawDistance : UpdatedDrawDistance)
