[Substrate] Change Dbuffer pass to run after the base pass when Substrate's DBufferPass is enabled. Also change Substrate DBuffer pass to not write out normal when not needed.

This allows to have correct normal input for the DBuffer base pass to read from, and avoid to rely on depth buffer reconstruction or on normal reprojection.

diff --git a/Engine/Shaders/Private/DBufferNormalReprojection.ush b/Engine/Shaders/Private/DBufferNormalReprojection.ush
--- a/Engine/Shaders/Private/DBufferNormalReprojection.ush
+++ b/Engine/Shaders/Private/DBufferNormalReprojection.ush
@@ -69,6 +69,10 @@ float4 GetDBufferReprojectedWorldNormal(float2 UV)
 {
 	int2 BaseCoord = (int2)trunc(UV * View.BufferSizeAndInvSize.xy);
 
+	// When using Substrate with DBuffer pass, we can rely on TopLayerTexture to have been produced
+	#if SUBSTRATE_ENABLED && SUBSTRATE_GBUFFER_FORMAT == 1 && SUBSTRATE_USE_DBUFFER_PASS
+	float3 N = SubstrateUnpackTopLayerData(SubstrateStruct.TopLayerTexture.Load(int3(BaseCoord, 0))).WorldNormal;
+	#else
 	float DeviceZ = SceneTexturesStruct.SceneDepthTexture.Load(int3(BaseCoord,0)).x;
 
 	// Velocity has to be the same size as the rest of the targets
@@ -111,6 +115,7 @@ float4 GetDBufferReprojectedWorldNormal(float2 UV)
 		}
 		N = normalize(lerp(FacetedNorm,ReprojNorm,Influence));
 	}
+	#endif
 
 	return float4(N,1.0f);
 }
diff --git a/Engine/Shaders/Private/Substrate/SubstrateDBuffer.usf b/Engine/Shaders/Private/Substrate/SubstrateDBuffer.usf
--- a/Engine/Shaders/Private/Substrate/SubstrateDBuffer.usf
+++ b/Engine/Shaders/Private/Substrate/SubstrateDBuffer.usf
@@ -168,10 +168,14 @@ void MainCS(
 		}
 
 		// 3. Top layer data
-		FSubstrateTopLayerData TopLayerData = SubstrateUnpackTopLayerData(SubstratePixelHeader.PackedTopLayerData);
-		TopLayerData.WorldNormal = TopLayerData.WorldNormal * DBufferData.OneMinusCoverage_WorldNormal + DBufferData.WorldNormal;
-		TopLayerData.Roughness = Roughness;
-		TopLayerTexture[PixelCoord] = SubstratePackTopLayerData(TopLayerData);
+		BRANCH
+		if (DBufferData.OneMinusCoverage_WorldNormal > 0 || DBufferData.OneMinusCoverage_Roughness > 0)
+		{
+			FSubstrateTopLayerData TopLayerData = SubstrateUnpackTopLayerData(SubstratePixelHeader.PackedTopLayerData);
+			TopLayerData.WorldNormal = TopLayerData.WorldNormal * DBufferData.OneMinusCoverage_WorldNormal + DBufferData.WorldNormal;
+			TopLayerData.Roughness = Roughness;
+			TopLayerTexture[PixelCoord] = SubstratePackTopLayerData(TopLayerData);
+		}
 
 		// 4. Rencode data
 		// SUBSTRATE_TODO: For now, doing manual reencoding by duplicating logic from Substrate export. Need to refactor Substrate export to have small reusable packing functions
diff --git a/Engine/Source/Runtime/Renderer/Private/CompositionLighting/CompositionLighting.cpp b/Engine/Source/Runtime/Renderer/Private/CompositionLighting/CompositionLighting.cpp
--- a/Engine/Source/Runtime/Renderer/Private/CompositionLighting/CompositionLighting.cpp
+++ b/Engine/Source/Runtime/Renderer/Private/CompositionLighting/CompositionLighting.cpp
@@ -557,8 +557,9 @@ void FCompositionLighting::ProcessBeforeBasePass(FRDGBuilder& GraphBuilder, FDBu
 		FViewInfo& View = Views[ViewIndex];
 		const FAOConfig& ViewConfig = ViewAOConfigs[ViewIndex];
 
+		const bool bRunDbufferBeforeBasePass = !Substrate::IsDBufferPassEnabled(View.GetShaderPlatform());
 		const bool bEnableSSAO = ViewConfig.SSAOLocation == ESSAOLocation::BeforeBasePass;
-		const bool bEnableDecals = DecalVisibility && DecalVisibility->HasStage(ViewIndex, EDecalRenderStage::BeforeBasePass);
+		const bool bEnableDecals = bRunDbufferBeforeBasePass && DecalVisibility && DecalVisibility->HasStage(ViewIndex, EDecalRenderStage::BeforeBasePass);
 
 		if (!bEnableDecals && !bEnableSSAO)
 		{
diff --git a/Engine/Source/Runtime/Renderer/Private/CompositionLighting/PostProcessDeferredDecals.cpp b/Engine/Source/Runtime/Renderer/Private/CompositionLighting/PostProcessDeferredDecals.cpp
--- a/Engine/Source/Runtime/Renderer/Private/CompositionLighting/PostProcessDeferredDecals.cpp
+++ b/Engine/Source/Runtime/Renderer/Private/CompositionLighting/PostProcessDeferredDecals.cpp
@@ -97,9 +97,11 @@ FDeferredDecalPassTextures GetDeferredDecalPassTextures(
 	SetupSceneTextureUniformParameters(GraphBuilder, &SceneTextures, View.FeatureLevel, TextureReadAccess, Parameters->SceneTextures);
 	SetupMobileSceneTextureUniformParameters(GraphBuilder, &SceneTextures, MobileTextureReadAccess, Parameters->MobileSceneTextures);
 	Parameters->EyeAdaptationBuffer = GraphBuilder.CreateSRV(GetEyeAdaptationBuffer(GraphBuilder, View));
-	if (DecalRenderStage == EDecalRenderStage::Emissive)
+
+	// When using SubstrateDbufferPass, the DBuffer are rendered after the BasePass. In this case, the toplayer data can be piped to the DBuffer base pass for normal reading.
+	if (DecalRenderStage == EDecalRenderStage::Emissive || (DecalRenderStage == EDecalRenderStage::BeforeBasePass && Substrate::IsDBufferPassEnabled(View.GetShaderPlatform())))
 	{
-		 Substrate::BindSubstratePublicGlobalUniformParameters(GraphBuilder, &SubstrateSceneData, Parameters->SubstratePublic);
+		Substrate::BindSubstratePublicGlobalUniformParameters(GraphBuilder, &SubstrateSceneData, Parameters->SubstratePublic);
 	}
 	else
 	{
diff --git a/Engine/Source/Runtime/Renderer/Private/DeferredShadingRenderer.cpp b/Engine/Source/Runtime/Renderer/Private/DeferredShadingRenderer.cpp
--- a/Engine/Source/Runtime/Renderer/Private/DeferredShadingRenderer.cpp
+++ b/Engine/Source/Runtime/Renderer/Private/DeferredShadingRenderer.cpp
@@ -2883,6 +2883,14 @@ void FDeferredShadingSceneRenderer::Render(FRDGBuilder& GraphBuilder, const FSce
 			RDG_EVENT_SCOPE_STAT(GraphBuilder, Substrate, "Substrate");
 			RDG_GPU_STAT_SCOPE(GraphBuilder, Substrate);
 
+			// Substrate DBufferPass (optional)
+			if (Substrate::IsDBufferPassEnabled(ShaderPlatform))
+			{
+				RDG_CSV_STAT_EXCLUSIVE_SCOPE(GraphBuilder, DeferredShadingSceneRenderer_DBuffer);
+				SCOPE_CYCLE_COUNTER(STAT_FDeferredShadingSceneRenderer_DBuffer);
+				Substrate::AddSubstrateDBufferBasePass(GraphBuilder, Views, SceneTextures, DBufferTextures, InitViewTaskDatas.Decals, InstanceCullingManager, Scene->SubstrateSceneData);
+			}
+
 			Substrate::AddSubstrateMaterialClassificationPass(GraphBuilder, SceneTextures, DBufferTextures, Views);
 			Substrate::AddSubstrateDBufferPass(GraphBuilder, SceneTextures, DBufferTextures, Views);
 			Substrate::AddSubstrateSampleMaterialPass(GraphBuilder, Scene, SceneTextures, Views);
diff --git a/Engine/Source/Runtime/Renderer/Private/Substrate/Substrate.cpp b/Engine/Source/Runtime/Renderer/Private/Substrate/Substrate.cpp
--- a/Engine/Source/Runtime/Renderer/Private/Substrate/Substrate.cpp
+++ b/Engine/Source/Runtime/Renderer/Private/Substrate/Substrate.cpp
@@ -17,6 +17,8 @@
 #include "EngineAnalytics.h"
 #include "SystemTextures.h"
 #include "DBufferTextures.h"
+#include "DecalRenderingShared.h"
+#include "CompositionLighting/PostProcessDeferredDecals.h"
 #include "MegaLights/MegaLights.h"
 
 
@@ -1785,6 +1787,40 @@ void AddSubstrateDBufferPass(FRDGBuilder& GraphBuilder, const FMinimalSceneTextu
 	}
 }
 
+void AddSubstrateDBufferBasePass(
+	FRDGBuilder& GraphBuilder, 
+	TArray<FViewInfo>& Views, 
+	const FSceneTextures& InSceneTextures,
+	FDBufferTextures& DBufferTextures, 
+	FDecalVisibilityTaskData* DecalVisibility, 
+	FInstanceCullingManager& InstanceCullingManager, 
+	const FSubstrateSceneData& SubstrateSceneData)
+{
+	if (!IsSubstrateEnabled() || DecalVisibility == nullptr)
+	{
+		return;
+	}
+
+	RDG_EVENT_SCOPE_CONDITIONAL(GraphBuilder, IsSubstrateEnabled() && Views.Num() > 0, "Substrate::DBufferBasePass");
+
+	for (int32 ViewIndex = 0; ViewIndex < Views.Num(); ++ViewIndex)
+	{
+		FViewInfo& View = Views[ViewIndex];
+
+		const bool bEnableDecals = IsDBufferPassEnabled(View.GetShaderPlatform()) && DecalVisibility->HasStage(ViewIndex, EDecalRenderStage::BeforeBasePass);
+		if (bEnableDecals)
+		{
+			RDG_GPU_MASK_SCOPE(GraphBuilder, View.GPUMask);
+			RDG_EVENT_SCOPE_CONDITIONAL(GraphBuilder, Views.Num() > 1, "View%d", ViewIndex);
+
+			View.BeginRenderView();
+			FDeferredDecalPassTextures DecalPassTextures = GetDeferredDecalPassTextures(GraphBuilder, View, SubstrateSceneData, InSceneTextures, &DBufferTextures, EDecalRenderStage::BeforeBasePass);
+			TConstArrayView<const FVisibleDecal*> SortedDecals = DecalVisibility->FinishRelevantDecals(ViewIndex, EDecalRenderStage::BeforeBasePass);
+			AddDeferredDecalPass(GraphBuilder, View, SortedDecals, DecalPassTextures, InstanceCullingManager, EDecalRenderStage::BeforeBasePass);
+		}
+	}
+}
+
 void AddSubstrateSampleMaterialPass(FRDGBuilder& GraphBuilder, const FScene* Scene, const FMinimalSceneTextures& SceneTextures, const TArray<FViewInfo>& Views)
 {
 	if (Substrate::IsSubstrateEnabled())
diff --git a/Engine/Source/Runtime/Renderer/Private/Substrate/Substrate.h b/Engine/Source/Runtime/Renderer/Private/Substrate/Substrate.h
--- a/Engine/Source/Runtime/Renderer/Private/Substrate/Substrate.h
+++ b/Engine/Source/Runtime/Renderer/Private/Substrate/Substrate.h
@@ -21,6 +21,8 @@ struct FDBufferTextures;
 struct FMinimalSceneTextures;
 struct FScreenPassTexture;
 struct FTextureRenderTargetBinding;
+class FDecalVisibilityTaskData;
+class FInstanceCullingManager;
 
 BEGIN_SHADER_PARAMETER_STRUCT(FSubstrateCommonParameters, )
 	SHADER_PARAMETER(uint32, MaxBytesPerPixel)
@@ -292,4 +294,13 @@ bool IsStochasticLightingActive(EShaderPlatform In);
 
 uint32 GetClosureTileIndirectArgsOffset(uint32 InDownsampleFactor);
 
+void AddSubstrateDBufferBasePass(
+	FRDGBuilder& GraphBuilder, 
+	TArray<FViewInfo>& Views, 
+	const FSceneTextures& InSceneTextures,
+	FDBufferTextures& DBufferTextures, 
+	FDecalVisibilityTaskData* DecalVisibility, 
+	FInstanceCullingManager& InstanceCullingManager, 
+	const FSubstrateSceneData& SubstrateSceneData);
+
 }; // namespace Substrate
