Temp workaround for avoiding editor crash at startup introduced by CL 43667472.

#test EngineTest editor

diff --git a/Engine/Source/Runtime/Renderer/Private/Shadows/ShadowSceneRenderer.cpp b/Engine/Source/Runtime/Renderer/Private/Shadows/ShadowSceneRenderer.cpp
--- a/Engine/Source/Runtime/Renderer/Private/Shadows/ShadowSceneRenderer.cpp
+++ b/Engine/Source/Runtime/Renderer/Private/Shadows/ShadowSceneRenderer.cpp
@@ -636,7 +636,10 @@ void FShadowSceneRenderer::AllocateVirtualShadowMapIds()
 	}
 
 	// Then add any unreferenced lights that we still may keep around their cached pages at the end, and allocate VSM IDs for everything
-	VirtualShadowMapArray.CacheManager->UpdateUnreferencedCacheEntries(VirtualShadowMapArray);
+	if (VirtualShadowMapArray.CacheManager)
+	{
+		VirtualShadowMapArray.CacheManager->UpdateUnreferencedCacheEntries(VirtualShadowMapArray);
+	}
 }
 
 void FShadowSceneRenderer::PostInitDynamicShadowsSetup()
