Fix OverrideOutput not being properly handled when "VisualizeLightGrid" is the last pass.
- also fixing static analysis.

diff --git a/Engine/Source/Runtime/Renderer/Private/PostProcess/PostProcessing.cpp b/Engine/Source/Runtime/Renderer/Private/PostProcess/PostProcessing.cpp
--- a/Engine/Source/Runtime/Renderer/Private/PostProcess/PostProcessing.cpp
+++ b/Engine/Source/Runtime/Renderer/Private/PostProcess/PostProcessing.cpp
@@ -371,6 +371,8 @@ void AddPostProcessingPasses(
 #endif
 	Inputs.Validate();
 
+	check(View.Family);
+
 	FScene* Scene = View.Family->Scene->GetRenderScene();
 
 	const FIntRect PrimaryViewRect = View.ViewRect;
@@ -1652,7 +1654,7 @@ void AddPostProcessingPasses(
 	{
 		FScreenPassRenderTarget OverrideOutput;
 		PassSequence.AcceptOverrideIfLastPass(EPass::VisualizeLightGrid, OverrideOutput);
-		SceneColor = AddVisualizeLightGridPass(GraphBuilder, View, SceneColor, SceneDepth);
+		SceneColor = AddVisualizeLightGridPass(GraphBuilder, View, OverrideOutput.IsValid() ? OverrideOutput : SceneColor, SceneDepth);
 	}
 
 #if WITH_EDITOR
@@ -2370,6 +2372,8 @@ void AddMobilePostProcessingPasses(FRDGBuilder& GraphBuilder, FScene* Scene, con
 	QUICK_SCOPE_CYCLE_COUNTER(STAT_PostProcessing_Process);
 	using namespace UE::Renderer::PostProcess;
 
+	check(View.Family);
+
 	check(IsInRenderingThread());
 	Inputs.Validate();
 
@@ -3384,7 +3388,7 @@ void AddMobilePostProcessingPasses(FRDGBuilder& GraphBuilder, FScene* Scene, con
 	{
 		FScreenPassRenderTarget OverrideOutput;
 		PassSequence.AcceptOverrideIfLastPass(EPass::VisualizeLightGrid, OverrideOutput);
-		SceneColor = AddVisualizeLightGridPass(GraphBuilder, View, SceneColor, SceneDepth);
+		SceneColor = AddVisualizeLightGridPass(GraphBuilder, View, OverrideOutput.IsValid() ? OverrideOutput : SceneColor, SceneDepth);
 	}
 
 	if (ShaderPrint::IsEnabled(View.ShaderPrintData))
