Fix crash on some platforms without VSMs

diff --git a/Engine/Source/Runtime/Renderer/Private/Shadows/ShadowSceneRenderer.cpp b/Engine/Source/Runtime/Renderer/Private/Shadows/ShadowSceneRenderer.cpp
--- a/Engine/Source/Runtime/Renderer/Private/Shadows/ShadowSceneRenderer.cpp
+++ b/Engine/Source/Runtime/Renderer/Private/Shadows/ShadowSceneRenderer.cpp
@@ -609,7 +609,10 @@ void FShadowSceneRenderer::AddDirectionalLightShadow(FLightSceneInfo& LightScene
 
 void FShadowSceneRenderer::AllocateVirtualShadowMapIds()
 {
-	// TODO: Move the bReferencedThisRender setting stuff to here as well just for organization?
+	if (!VirtualShadowMapArray.IsEnabled())
+	{
+		return;
+	}
 
 	// Directional lights allocated first
 	for (const FDirectionalLightShadowFrameSetup& FrameSetup : DirectionalLights)
diff --git a/Engine/Source/Runtime/Renderer/Private/VirtualShadowMaps/VirtualShadowMapCacheManager.cpp b/Engine/Source/Runtime/Renderer/Private/VirtualShadowMaps/VirtualShadowMapCacheManager.cpp
--- a/Engine/Source/Runtime/Renderer/Private/VirtualShadowMaps/VirtualShadowMapCacheManager.cpp
+++ b/Engine/Source/Runtime/Renderer/Private/VirtualShadowMaps/VirtualShadowMapCacheManager.cpp
@@ -1184,6 +1184,11 @@ void FVirtualShadowMapPerLightCacheEntry::OnPrimitiveRendered(const FPrimitiveSc
 void FVirtualShadowMapArrayCacheManager::UpdateUnreferencedCacheEntries(
 	FVirtualShadowMapArray& VirtualShadowMapArray)
 {
+	if (!VirtualShadowMapArray.IsEnabled())
+	{
+		return;
+	}
+
 	const uint32 SceneFrameNumber = Scene.GetFrameNumberRenderThread();
 	const int32 MaxLightAge = CVarMaxLightAgeSinceLastRequest.GetValueOnRenderThread();
 
