Backout of 43539808 and 43586875 but keeping subsequent unrelated changes from other CLs

diff --git a/Engine/Source/Runtime/Renderer/Private/InstanceCulling/InstanceCullingOcclusionQuery.cpp b/Engine/Source/Runtime/Renderer/Private/InstanceCulling/InstanceCullingOcclusionQuery.cpp
--- a/Engine/Source/Runtime/Renderer/Private/InstanceCulling/InstanceCullingOcclusionQuery.cpp
+++ b/Engine/Source/Runtime/Renderer/Private/InstanceCulling/InstanceCullingOcclusionQuery.cpp
@@ -348,7 +348,7 @@ struct FInstanceCullingOcclusionQueryDeferredContext
 			return;
 		}
 
-		FParallelMeshDrawCommandPass& MeshDrawCommandPass = *View->ParallelMeshDrawCommandPasses[MeshPass];
+		const FParallelMeshDrawCommandPass& MeshDrawCommandPass = *View->ParallelMeshDrawCommandPasses[MeshPass];
 
 		// Execute() is expected to run late enough to not stall here.
 		// If it does happen, then we may have to move the render pass to later point in the frame.
diff --git a/Engine/Source/Runtime/Renderer/Private/MeshDrawCommands.cpp b/Engine/Source/Runtime/Renderer/Private/MeshDrawCommands.cpp
--- a/Engine/Source/Runtime/Renderer/Private/MeshDrawCommands.cpp
+++ b/Engine/Source/Runtime/Renderer/Private/MeshDrawCommands.cpp
@@ -1484,7 +1484,7 @@ void FParallelMeshDrawCommandPass::DispatchPassSetup(
 		// This work needs to be deferred until at least BuildRenderingCommands (to ensure the DynamicPrimitiveCollector is uploaded), so we use the async mechanism either way 
 		auto FinalizeInstanceCullingSetup = [this, Scene](FInstanceCullingContext& InstanceCullingContext)
 		{
-			WaitForSetupTask();
+			WaitForMeshPassSetupTask();
 
 #if DO_CHECK
 			for (const FVisibleMeshDrawCommand& VisibleMeshDrawCommand : TaskContext.MeshDrawCommands)
@@ -1511,20 +1511,22 @@ bool FParallelMeshDrawCommandPass::IsOnDemandShaderCreationEnabled()
 		(GRHISupportsMultithreadedShaderCreation || (bIsMobileRenderer && (!GSupportsParallelRenderingTasksWithSeparateRHIThread && IsRunningRHIInSeparateThread())));
 }
 
-void FParallelMeshDrawCommandPass::WaitForSetupTask()
+void FParallelMeshDrawCommandPass::WaitForMeshPassSetupTask() const
 {
 	if (TaskEventRef.IsValid())
 	{
 		// Need to wait on GetRenderThread_Local, as mesh pass setup task can wait on rendering thread inside InitResourceFromPossiblyParallelRendering().
 		QUICK_SCOPE_CYCLE_COUNTER(STAT_WaitForMeshPassSetupTask);
 		TaskEventRef->Wait(IsInActualRenderingThread() ? ENamedThreads::GetRenderThread_Local() : ENamedThreads::AnyThread);
-		TaskEventRef = nullptr;
 	}
 }
 
-void FParallelMeshDrawCommandPass::Cleanup()
+void FParallelMeshDrawCommandPass::WaitForTasksAndEmpty()
 {
-	check(TaskEventRef == nullptr);
+	// Need to wait in case if someone dispatched sort and draw merge task, but didn't draw it.
+	WaitForMeshPassSetupTask();
+	TaskEventRef = nullptr;
+
 	DumpInstancingStats();
 
 	if (TaskContext.MeshPassProcessor)
@@ -1664,6 +1666,13 @@ void FParallelMeshDrawCommandPass::BuildRenderingCommands(
 	OutInstanceCullingDrawParams.IndirectArgsByteOffset = 0U;
 }
 
+
+
+void FParallelMeshDrawCommandPass::WaitForSetupTask() const
+{
+	WaitForMeshPassSetupTask();
+}
+
 PRAGMA_DISABLE_DEPRECATION_WARNINGS
 
 void FParallelMeshDrawCommandPass::DispatchDraw(FParallelCommandListSet* ParallelCommandListSet, FRHICommandList& RHICmdList, const FInstanceCullingDrawParams* InstanceCullingDrawParams) const
@@ -1741,7 +1750,7 @@ void FParallelMeshDrawCommandPass::Draw(FRHICommandList& RHICmdList, const FInst
 
 	QUICK_SCOPE_CYCLE_COUNTER(STAT_MeshPassDrawImmediate);
 
-	check(TaskEventRef == nullptr);
+	WaitForMeshPassSetupTask();
 
 	if (TaskContext.bUseGPUScene)
 	{
diff --git a/Engine/Source/Runtime/Renderer/Private/MeshDrawCommands.h b/Engine/Source/Runtime/Renderer/Private/MeshDrawCommands.h
--- a/Engine/Source/Runtime/Renderer/Private/MeshDrawCommands.h
+++ b/Engine/Source/Runtime/Renderer/Private/MeshDrawCommands.h
@@ -157,7 +157,7 @@ public:
 	/**
 	 * Sync with setup task.
 	 */
-	void WaitForSetupTask();
+	void WaitForSetupTask() const;
 
 	/**
 	 * Dispatch visible mesh draw command draw task.
@@ -171,7 +171,7 @@ PRAGMA_DISABLE_DEPRECATION_WARNINGS
 	void DispatchDraw(FParallelCommandListSet* ParallelCommandListSet, FRHICommandList& RHICmdList, const FInstanceCullingDrawParams* InstanceCullingDrawParams = nullptr) const;
 PRAGMA_ENABLE_DEPRECATION_WARNINGS
 
-	void Cleanup();
+	void WaitForTasksAndEmpty();
 	void SetDumpInstancingStats(const FString& InPassName);
 	bool HasAnyDraw() const { return MaxNumDraws > 0; }
 
@@ -180,8 +180,6 @@ PRAGMA_ENABLE_DEPRECATION_WARNINGS
 	FInstanceCullingContext* GetInstanceCullingContext() { return &TaskContext.InstanceCullingContext; }
 	const FGraphEventRef& GetTaskEvent() const { return TaskEventRef; }
 
-	FGraphEventRef&& AcquireTaskEvent() { return MoveTemp(TaskEventRef); }
-
 	// NOTE: It is only safe to access mesh draw commands after the setup task is complete (use WaitForSetupTask). 
 	// Only access the data late in the frame to allow as much time as possible for async tasks to complete.
 	const FMeshCommandOneFrameArray& GetMeshDrawCommands() const
@@ -201,6 +199,7 @@ private:
 	int32 MaxNumDraws;
 
 	void DumpInstancingStats() const;
+	void WaitForMeshPassSetupTask() const;
 };
 
 inline bool HasAnyDraw(const FParallelMeshDrawCommandPass* Pass)
diff --git a/Engine/Source/Runtime/Renderer/Private/MobileShadingRenderer.cpp b/Engine/Source/Runtime/Renderer/Private/MobileShadingRenderer.cpp
--- a/Engine/Source/Runtime/Renderer/Private/MobileShadingRenderer.cpp
+++ b/Engine/Source/Runtime/Renderer/Private/MobileShadingRenderer.cpp
@@ -406,11 +406,6 @@ void FMobileSceneRenderer::SetupMobileBasePassAfterShadowInit(FExclusiveDepthSte
 		FViewInfo& View = *AllViews[ViewIndex];
 		FViewCommands& ViewCommands = ViewCommandsPerView[ViewIndex];
 
-		if (ViewCommands.MeshCommands[EMeshPass::BasePass].IsEmpty() && ViewCommands.MeshCommands[EMeshPass::MobileBasePassCSM].IsEmpty() && View.NumVisibleDynamicMeshElements[EMeshPass::BasePass] == 0 && ViewCommands.NumDynamicMeshCommandBuildRequestElements[EMeshPass::BasePass] == 0)
-		{
-			continue;
-		}
-
 		FMeshPassProcessor* MeshPassProcessor = FPassProcessorManager::CreateMeshPassProcessor(EShadingPath::Mobile, EMeshPass::BasePass, Scene->GetFeatureLevel(), Scene, &View, nullptr);
 
 		FMeshPassProcessor* BasePassCSMMeshPassProcessor = FPassProcessorManager::CreateMeshPassProcessor(EShadingPath::Mobile, EMeshPass::MobileBasePassCSM, Scene->GetFeatureLevel(), Scene, &View, nullptr);
diff --git a/Engine/Source/Runtime/Renderer/Private/Renderer/ViewSnapshotCache.cpp b/Engine/Source/Runtime/Renderer/Private/Renderer/ViewSnapshotCache.cpp
--- a/Engine/Source/Runtime/Renderer/Private/Renderer/ViewSnapshotCache.cpp
+++ b/Engine/Source/Runtime/Renderer/Private/Renderer/ViewSnapshotCache.cpp
@@ -94,7 +94,7 @@ public:
 			{
 				if (Pass)
 				{
-					Pass->Cleanup();
+					Pass->WaitForTasksAndEmpty();
 				}
 			}
 
diff --git a/Engine/Source/Runtime/Renderer/Private/SceneRenderBuilder.cpp b/Engine/Source/Runtime/Renderer/Private/SceneRenderBuilder.cpp
--- a/Engine/Source/Runtime/Renderer/Private/SceneRenderBuilder.cpp
+++ b/Engine/Source/Runtime/Renderer/Private/SceneRenderBuilder.cpp
@@ -295,7 +295,7 @@ static void DestroySplitScreenDebugViewFamilies(TConstArrayView<const FSceneView
 
 //////////////////////////////////////////////////////////////////////////
 
-static int32 GSceneRenderCleanUpMode = 1;
+static int32 GSceneRenderCleanUpMode = 0;
 static FAutoConsoleVariableRef CVarSceneRenderCleanUpMode(
 	TEXT("r.SceneRender.CleanUpMode"),
 	GSceneRenderCleanUpMode,
@@ -705,12 +705,12 @@ static void CleanupSceneRenderer(FSceneRenderer* Renderer)
 
 	for (auto* Pass : Renderer->DispatchedShadowDepthPasses)
 	{
-		Pass->Cleanup();
+		Pass->WaitForTasksAndEmpty();
 	}
 
 	for (FViewInfo* View : Renderer->AllViews)
 	{
-		View->Cleanup();
+		View->WaitForTasks();
 	}
 
 	ViewSnapshotCache::Destroy();
@@ -725,14 +725,48 @@ void FSceneRenderProcessor::Cleanup(FRHICommandListImmediate& RHICmdList, FScene
 
 	const ESceneRenderCleanUpMode SceneRenderCleanUpMode = GetSceneRenderCleanUpMode();
 
-	if (SceneRenderCleanUpMode == ESceneRenderCleanUpMode::Async)
+	if (SceneRenderCleanUpMode == ESceneRenderCleanUpMode::Immediate)
 	{
+		WaitForAsyncDeleteTask(); // This is to handle cases where as switch from async to immediate.
+		RHICmdList.ImmediateFlush(EImmediateFlushType::WaitForOutstandingTasksOnly);
+		ViewSnapshotCache::Deallocate();
+		CleanupSceneRenderer(Renderer);
+	}
+	else
+	{
+		FGraphEventArray Prereqs;
 		Renderer->GPUSceneDynamicContext.Release();
 
+		// Wait on all setup tasks now to ensure that no additional render commands are enqueued which
+		// might mess with render state, since setup tasks are working with high-level render objects.
+		for (FParallelMeshDrawCommandPass* DispatchedShadowDepthPass : Renderer->DispatchedShadowDepthPasses)
+		{
+			if (DispatchedShadowDepthPass->GetTaskEvent())
+			{
+				Prereqs.Add(DispatchedShadowDepthPass->GetTaskEvent());
+			}
+		}
+
+		for (const FViewInfo& View : Renderer->Views)
+		{
+			for (const FParallelMeshDrawCommandPass* Pass : View.ParallelMeshDrawCommandPasses)
+			{
+				if (Pass && Pass->GetTaskEvent())
+				{
+					Prereqs.Add(Pass->GetTaskEvent());
+				}
+			}
+		}
+
 		// Wait for the last renderer's cleanup tasks so that snapshot deallocation and destruction don't overlap.
 		if (AsyncTasks.Cleanup)
 		{
-			AsyncTasks.Cleanup->Wait(ENamedThreads::GetRenderThread_Local());
+			Prereqs.Add(AsyncTasks.Cleanup);
+		}
+
+		if (!Prereqs.IsEmpty())
+		{
+			FTaskGraphInterface::Get().WaitUntilTasksComplete(Prereqs, ENamedThreads::GetRenderThread_Local());
 		}
 
 		ViewSnapshotCache::Deallocate();
@@ -743,13 +777,6 @@ void FSceneRenderProcessor::Cleanup(FRHICommandListImmediate& RHICmdList, FScene
 
 		}, TStatId(), &GRHICommandList.WaitOutstandingTasks);
 	}
-	else // Immediate
-	{
-		WaitForAsyncDeleteTask(); // This is to handle cases where as switch from async to immediate.
-		RHICmdList.ImmediateFlush(EImmediateFlushType::WaitForOutstandingTasksOnly);
-		ViewSnapshotCache::Deallocate();
-		CleanupSceneRenderer(Renderer);
-	}
 
 	GlobalDynamicBuffer::GarbageCollect();
 	GPrimitiveIdVertexBufferPool.DiscardAll();
diff --git a/Engine/Source/Runtime/Renderer/Private/SceneRendering.cpp b/Engine/Source/Runtime/Renderer/Private/SceneRendering.cpp
--- a/Engine/Source/Runtime/Renderer/Private/SceneRendering.cpp
+++ b/Engine/Source/Runtime/Renderer/Private/SceneRendering.cpp
@@ -1038,13 +1038,13 @@ FParallelMeshDrawCommandPass* FViewInfo::CreateMeshPass(EMeshPass::Type MeshPass
 	return ParallelMeshDrawCommandPasses[MeshPass] = Allocator.Create<FParallelMeshDrawCommandPass>();
 }
 
-void FViewInfo::Cleanup()
+void FViewInfo::WaitForTasks()
 {
 	for (int32 MeshDrawIndex = 0; MeshDrawIndex < EMeshPass::Num; MeshDrawIndex++)
 	{
 		if (auto* Pass = ParallelMeshDrawCommandPasses[MeshDrawIndex])
 		{
-			Pass->Cleanup();
+			Pass->WaitForTasksAndEmpty();
 		}
 	}
 }
@@ -4654,33 +4654,6 @@ void FSceneRenderer::OnRenderFinish(FRDGBuilder& GraphBuilder, FRDGTextureRef Vi
 			Scene->DumpMeshDrawCommandMemoryStats();
 		}
 	});
-
-	// Sync all mesh pass setup tasks prior to starting RDG execution.
-	FGraphEventArray TasksToWait;
-
-	for (FParallelMeshDrawCommandPass* DispatchedShadowDepthPass : DispatchedShadowDepthPasses)
-	{
-		if (FGraphEventRef Event = DispatchedShadowDepthPass->AcquireTaskEvent())
-		{
-			TasksToWait.Emplace(MoveTemp(Event));
-		}
-	}
-
-	for (const FViewInfo* View : AllViews)
-	{
-		for (FParallelMeshDrawCommandPass* Pass : View->ParallelMeshDrawCommandPasses)
-		{
-			if (Pass)
-			{
-				if (FGraphEventRef Event = Pass->AcquireTaskEvent())
-				{
-					TasksToWait.Emplace(MoveTemp(Event));
-				}
-			}
-		}
-	}
-
-	FTaskGraphInterface::Get().WaitUntilTasksComplete(TasksToWait, ENamedThreads::GetRenderThread_Local());
 }
 
 void FSceneRenderer::SetupMeshPass(FViewInfo& View, FExclusiveDepthStencil::Type BasePassDepthStencilAccess, FViewCommands& ViewCommands, FInstanceCullingManager& InstanceCullingManager)
diff --git a/Engine/Source/Runtime/Renderer/Private/SceneRendering.h b/Engine/Source/Runtime/Renderer/Private/SceneRendering.h
--- a/Engine/Source/Runtime/Renderer/Private/SceneRendering.h
+++ b/Engine/Source/Runtime/Renderer/Private/SceneRendering.h
@@ -1769,7 +1769,7 @@ public:
 	/** Create a snapshot of this view info on the scene allocator. */
 	FViewInfo* CreateSnapshot() const;
 
-	void Cleanup();
+	void WaitForTasks();
 
 	// Get the range in DynamicMeshElements[] for a given PrimitiveIndex
 	// @return range (start is inclusive, end is exclusive)
diff --git a/Engine/Source/Runtime/Renderer/Private/ShadowSetup.cpp b/Engine/Source/Runtime/Renderer/Private/ShadowSetup.cpp
--- a/Engine/Source/Runtime/Renderer/Private/ShadowSetup.cpp
+++ b/Engine/Source/Runtime/Renderer/Private/ShadowSetup.cpp
@@ -2908,7 +2908,7 @@ void FProjectedShadowInfo::ClearTransientArrays()
 	DynamicSubjectHeterogeneousVolumeMeshElements.Empty();
 
 	ShadowDepthPassVisibleCommands.Empty();
-	ShadowDepthPass.Cleanup();
+	ShadowDepthPass.WaitForTasksAndEmpty();
 
 	SubjectMeshCommandBuildRequests.Empty();
 	SubjectMeshCommandBuildFlags.Empty();
