Add Substrate DBuffer pass support for BlendableGBuffer.

This allows DBuffer usage with proper GBuffer normal usage during DBuffer base pass.

diff --git a/Engine/Shaders/Private/DBufferNormalReprojection.ush b/Engine/Shaders/Private/DBufferNormalReprojection.ush
--- a/Engine/Shaders/Private/DBufferNormalReprojection.ush
+++ b/Engine/Shaders/Private/DBufferNormalReprojection.ush
@@ -72,6 +72,8 @@ float4 GetDBufferReprojectedWorldNormal(float2 UV)
 	// When using Substrate with DBuffer pass, we can rely on TopLayerTexture to have been produced
 	#if SUBSTRATE_ENABLED && SUBSTRATE_GBUFFER_FORMAT == 1 && SUBSTRATE_USE_DBUFFER_PASS
 	float3 N = SubstrateUnpackTopLayerData(SubstrateStruct.TopLayerTexture.Load(int3(BaseCoord, 0))).WorldNormal;
+	#elif SUBSTRATE_ENABLED && SUBSTRATE_GBUFFER_FORMAT == 0 && SUBSTRATE_USE_DBUFFER_PASS
+	float3 N = GetGBufferDataUint(BaseCoord).WorldNormal;
 	#else
 	float DeviceZ = SceneTexturesStruct.SceneDepthTexture.Load(int3(BaseCoord,0)).x;
 
diff --git a/Engine/Shaders/Private/Substrate/SubstrateDBuffer.usf b/Engine/Shaders/Private/Substrate/SubstrateDBuffer.usf
--- a/Engine/Shaders/Private/Substrate/SubstrateDBuffer.usf
+++ b/Engine/Shaders/Private/Substrate/SubstrateDBuffer.usf
@@ -24,7 +24,7 @@
 
 ////////////////////////////////////////////////////////////////////////////////////////////////////////////
 // Post Dbuffer application
-#if SHADER_DBUFFER
+#if defined(DBufferPassCS) || defined(DBufferPassPS)
 Texture2D<float4> DBufferATexture;
 Texture2D<float4> DBufferBTexture;
 Texture2D<float4> DBufferCTexture;
@@ -73,6 +73,31 @@ SamplerState DBufferCTextureSampler;
 
 #include "../DBufferDecalShared.ush"
 
+uint GetDBufferResponseMask(uint2 PixelCoord)
+{
+	// Early out of no decal has been written or if the pixel has no 'decal responseness'
+	const uint ResponseMask = SceneStencilTexture.Load(uint3(PixelCoord, 0)) STENCIL_COMPONENT_SWIZZLE;
+	uint ValidDBufferTargetMask = GetDBufferTargetMask(PixelCoord);
+	if (!ValidDBufferTargetMask || (ResponseMask & SUBSTRATE_STENCIL_DBUFFER_MASK) == 0)
+	{
+		return 0;
+	}
+
+	// Load Dbuffer data only attributes for which the material has some 'responseness'
+	ValidDBufferTargetMask = ValidDBufferTargetMask & 
+	(
+		(((ResponseMask >> STENCIL_SUBSTRATE_RECEIVE_DBUFFER_DIFFUSE_BIT_ID)   & 0x1) << SUBSTRATE_DBUFFER_RESPONSE_BASECOLOR_BIT_OFFSET) |
+		(((ResponseMask >> STENCIL_SUBSTRATE_RECEIVE_DBUFFER_NORMAL_BIT_ID)    & 0x1) << SUBSTRATE_DBUFFER_RESPONSE_NORMAL_BIT_OFFSET)    |
+		(((ResponseMask >> STENCIL_SUBSTRATE_RECEIVE_DBUFFER_ROUGHNESS_BIT_ID) & 0x1) << SUBSTRATE_DBUFFER_RESPONSE_ROUGHNESS_BIT_OFFSET)
+	);
+
+	return ValidDBufferTargetMask;
+}
+
+#endif // defined(DBufferPassCS) || defined(DBufferPassPS)
+
+#ifdef DBufferPassCS
+
 int2 ViewResolution;
 uint MaxBytesPerPixel;
 uint FirstSliceStoringSubstrateSSSData;
@@ -89,7 +114,7 @@ Buffer<uint> TileIndirectBuffer;
 // * Refactor encoding to share code between Substrate export
 
 [numthreads(SUBSTRATE_TILE_SIZE, SUBSTRATE_TILE_SIZE, 1)]
-void MainCS(
+void DBufferPassCS(
 	uint2 DispatchThreadId : SV_DispatchThreadID, 
 	uint LinearIndex : SV_GroupIndex, 
 	uint3 GroupId : SV_GroupID,
@@ -105,22 +130,13 @@ void MainCS(
 	}
 
 	// Early out of no decal has been written or if the pixel has no 'decal responseness'
-	const uint ResponseMask = SceneStencilTexture.Load(uint3(PixelCoord, 0)) STENCIL_COMPONENT_SWIZZLE;
-	uint ValidDBufferTargetMask = GetDBufferTargetMask(PixelCoord);
+	const uint ValidDBufferTargetMask = GetDBufferResponseMask(PixelCoord);
 	BRANCH
-	if (!ValidDBufferTargetMask || (ResponseMask & SUBSTRATE_STENCIL_DBUFFER_MASK) == 0)
+	if (ValidDBufferTargetMask == 0)
 	{
 		return;
 	}
 
-	// Load Dbuffer data only attributes for which the material has some 'responseness'
-	ValidDBufferTargetMask = ValidDBufferTargetMask & 
-	(
-		(((ResponseMask >> STENCIL_SUBSTRATE_RECEIVE_DBUFFER_DIFFUSE_BIT_ID)   & 0x1) << SUBSTRATE_DBUFFER_RESPONSE_BASECOLOR_BIT_OFFSET) |
-		(((ResponseMask >> STENCIL_SUBSTRATE_RECEIVE_DBUFFER_NORMAL_BIT_ID)    & 0x1) << SUBSTRATE_DBUFFER_RESPONSE_NORMAL_BIT_OFFSET)    |
-		(((ResponseMask >> STENCIL_SUBSTRATE_RECEIVE_DBUFFER_ROUGHNESS_BIT_ID) & 0x1) << SUBSTRATE_DBUFFER_RESPONSE_ROUGHNESS_BIT_OFFSET)
-	);
-
 	const float2 BufferUV = SvPositionToBufferUV(float4(PixelCoord + 0.5f, 0, 0));
 	const FSubstrateDBuffer DBufferData = SubstrateGetDBufferData(BufferUV, ValidDBufferTargetMask);
 
@@ -342,4 +358,56 @@ void MainCS(
 		#endif // !SUBSTRATE_FASTPATH
 	}
 }
-#endif // SHADER_DBUFFER
+#endif // DBufferPassCS
+
+////////////////////////////////////////////////////////////////////////////////////////////////////////////
+
+#ifdef DBufferPassPS
+
+// @param BufferUV - UV space in the DBuffer textures
+void DBufferPassPS(
+	float4 SVPos : SV_POSITION,
+	out float4 OutGBufferA : SV_Target0,
+	out float4 OutGBufferB : SV_Target1,
+	out float4 OutGBufferC : SV_Target2)
+{
+	const uint2 PixelCoord = SVPos.xy;
+	const bool bIsValid = all(PixelCoord < uint2(View.ViewSizeAndInvSize.xy));
+
+	if (!bIsValid)
+	{
+		return;
+	}
+
+	#if 1
+	uint ValidDBufferTargetMask = 0x07;
+	#else
+	// Early out of no decal has been written or if the pixel has no 'decal responseness'
+	const uint ValidDBufferTargetMask = GetDBufferResponseMask(PixelCoord);
+	BRANCH
+	if (ValidDBufferTargetMask == 0)
+	{
+		return;
+	}
+	#endif
+
+	const float2 BufferUV = SvPositionToBufferUV(float4(PixelCoord + 0.5f, 0, 0));
+
+	// Setup default values, which mean that no decals are present in DBuffer
+	float4 DBufferA = float4(0.0, 0.0, 0.0, 1.0);
+	float4 DBufferB = float4(128.0f / 255.0f, 128.f / 255.5f, 128.f / 255.5f, 1.0);
+	float4 DBufferC = float4(0.0, 0.0, 0.0, 1.0);
+	{
+		BRANCH
+		if (ValidDBufferTargetMask & 0x1) { DBufferA = Texture2DSampleLevel(OpaqueBasePassDBufferATexture, OpaqueBasePassDBufferATextureSampler, BufferUV, 0);	}
+		BRANCH
+		if (ValidDBufferTargetMask & 0x2) { DBufferB = Texture2DSampleLevel(OpaqueBasePassDBufferBTexture, OpaqueBasePassDBufferBTextureSampler, BufferUV, 0);	}
+		BRANCH
+		if (ValidDBufferTargetMask & 0x4) { DBufferC = Texture2DSampleLevel(OpaqueBasePassDBufferCTexture, OpaqueBasePassDBufferCTextureSampler, BufferUV, 0);	}
+	}
+
+	OutGBufferA = DBufferB;
+	OutGBufferB = DBufferC;
+	OutGBufferC = DBufferA;
+}
+#endif // DBufferPassPS
diff --git a/Engine/Source/Runtime/Renderer/Private/Substrate/Substrate.cpp b/Engine/Source/Runtime/Renderer/Private/Substrate/Substrate.cpp
--- a/Engine/Source/Runtime/Renderer/Private/Substrate/Substrate.cpp
+++ b/Engine/Source/Runtime/Renderer/Private/Substrate/Substrate.cpp
@@ -1029,25 +1029,61 @@ class FSubstrateDBufferPassCS : public FGlobalShader
 		return GetMaxSupportedFeatureLevel(Parameters.Platform) >= ERHIFeatureLevel::SM5 && Substrate::IsSubstrateEnabled() && IsUsingDBuffers(Parameters.Platform) && !Substrate::IsSubstrateBlendableGBufferEnabled(Parameters.Platform);
 	}
 
-	static void ModifyCompilationEnvironment(const FGlobalShaderPermutationParameters& Parameters, FShaderCompilerEnvironment& OutEnvironment)
+	static void SetDBufferCompilationEnvironment(FShaderCompilerEnvironment& OutEnvironment)
 	{
 		const uint32 SubstrateStencilDbufferMask 
 			= GET_STENCIL_BIT_MASK(SUBSTRATE_RECEIVE_DBUFFER_NORMAL, 1) 
 			| GET_STENCIL_BIT_MASK(SUBSTRATE_RECEIVE_DBUFFER_DIFFUSE, 1) 
 			| GET_STENCIL_BIT_MASK(SUBSTRATE_RECEIVE_DBUFFER_ROUGHNESS, 1);
 
-		FGlobalShader::ModifyCompilationEnvironment(Parameters, OutEnvironment);
-		OutEnvironment.SetDefine(TEXT("SHADER_DBUFFER"), 1);
 		OutEnvironment.SetDefine(TEXT("SUBSTRATE_STENCIL_DBUFFER_MASK"), SubstrateStencilDbufferMask);
 		OutEnvironment.SetDefine(TEXT("STENCIL_SUBSTRATE_RECEIVE_DBUFFER_NORMAL_BIT_ID"), STENCIL_SUBSTRATE_RECEIVE_DBUFFER_NORMAL_BIT_ID);
 		OutEnvironment.SetDefine(TEXT("STENCIL_SUBSTRATE_RECEIVE_DBUFFER_DIFFUSE_BIT_ID"), STENCIL_SUBSTRATE_RECEIVE_DBUFFER_DIFFUSE_BIT_ID);
-		OutEnvironment.SetDefine(TEXT("STENCIL_SUBSTRATE_RECEIVE_DBUFFER_ROUGHNESS_BIT_ID"), STENCIL_SUBSTRATE_RECEIVE_DBUFFER_ROUGHNESS_BIT_ID);
+		OutEnvironment.SetDefine(TEXT("STENCIL_SUBSTRATE_RECEIVE_DBUFFER_ROUGHNESS_BIT_ID"), STENCIL_SUBSTRATE_RECEIVE_DBUFFER_ROUGHNESS_BIT_ID);	
+	}
+
+	static void ModifyCompilationEnvironment(const FGlobalShaderPermutationParameters& Parameters, FShaderCompilerEnvironment& OutEnvironment)
+	{
+		FGlobalShader::ModifyCompilationEnvironment(Parameters, OutEnvironment);
+		SetDBufferCompilationEnvironment(OutEnvironment);
 
 		// Needed as top layer texture can be a uint2
 		OutEnvironment.CompilerFlags.Add(CFLAG_AllowTypedUAVLoads);
 	}
 };
-IMPLEMENT_GLOBAL_SHADER(FSubstrateDBufferPassCS, "/Engine/Private/Substrate/SubstrateDBuffer.usf", "MainCS", SF_Compute);
+IMPLEMENT_GLOBAL_SHADER(FSubstrateDBufferPassCS, "/Engine/Private/Substrate/SubstrateDBuffer.usf", "DBufferPassCS", SF_Compute);
+
+////////////////////////////////////////////////////////////////////////////////////////////////////////////////////////////////////////////////
+
+class FSubstrateDBufferPassPS : public FGlobalShader
+{
+	DECLARE_GLOBAL_SHADER(FSubstrateDBufferPassPS);
+	SHADER_USE_PARAMETER_STRUCT(FSubstrateDBufferPassPS, FGlobalShader);
+
+	using FPermutationDomain = TShaderPermutationDomain<>;
+
+	BEGIN_SHADER_PARAMETER_STRUCT(FParameters, )
+		SHADER_PARAMETER(FIntPoint, ViewResolution)
+		SHADER_PARAMETER(uint32, MaxBytesPerPixel)
+		SHADER_PARAMETER(uint32, FirstSliceStoringSubstrateSSSData)
+		SHADER_PARAMETER_STRUCT_REF(FViewUniformShaderParameters, ViewUniformBuffer)
+		SHADER_PARAMETER_STRUCT_INCLUDE(FDBufferParameters, DBuffer)
+		SHADER_PARAMETER_RDG_TEXTURE_SRV(Texture2D<float>, SceneStencilTexture)
+		RENDER_TARGET_BINDING_SLOTS()
+	END_SHADER_PARAMETER_STRUCT()
+
+	static bool ShouldCompilePermutation(const FGlobalShaderPermutationParameters& Parameters)
+	{
+		return GetMaxSupportedFeatureLevel(Parameters.Platform) >= ERHIFeatureLevel::SM5 && Substrate::IsSubstrateEnabled() && IsUsingDBuffers(Parameters.Platform) && Substrate::IsSubstrateBlendableGBufferEnabled(Parameters.Platform);
+	}
+
+	static void ModifyCompilationEnvironment(const FGlobalShaderPermutationParameters& Parameters, FShaderCompilerEnvironment& OutEnvironment)
+	{
+		FGlobalShader::ModifyCompilationEnvironment(Parameters, OutEnvironment);
+		FSubstrateDBufferPassCS::SetDBufferCompilationEnvironment(OutEnvironment);
+	}
+};
+IMPLEMENT_GLOBAL_SHADER(FSubstrateDBufferPassPS, "/Engine/Private/Substrate/SubstrateDBuffer.usf", "DBufferPassPS", SF_Pixel);
 
 ////////////////////////////////////////////////////////////////////////////////////////////////////////////////////////////////////////////////
 
@@ -1699,8 +1735,7 @@ void AddSubstrateMaterialClassificationPass(FRDGBuilder& GraphBuilder, const FMi
 	}
 }
 
-
-void AddSubstrateDBufferPass(FRDGBuilder& GraphBuilder, const FMinimalSceneTextures& SceneTextures, const FDBufferTextures& DBufferTextures, const TArray<FViewInfo>& Views)
+void AddSubstrateDBufferPass(FRDGBuilder& GraphBuilder, const FSceneTextures& SceneTextures, const FDBufferTextures& DBufferTextures, const TArray<FViewInfo>& Views)
 {
 	RDG_EVENT_SCOPE_CONDITIONAL(GraphBuilder, IsSubstrateEnabled() && Views.Num() > 0, "Substrate::DBuffer");
 	if (!IsSubstrateEnabled() || !DBufferTextures.IsValid())
@@ -1723,67 +1758,110 @@ void AddSubstrateDBufferPass(FRDGBuilder& GraphBuilder, const FMinimalSceneTextu
 		const FSubstrateViewData* SubstrateViewData = &View.SubstrateViewData;
 		const FSubstrateSceneData* SubstrateSceneData = View.SubstrateViewData.SceneData;
 
-		FRDGTextureUAVRef RWMaterialTexture = GraphBuilder.CreateUAV(SubstrateSceneData->MaterialTextureArray, ERDGUnorderedAccessViewFlags::SkipBarrier);
-		FRDGTextureUAVRef RWTopLayerTexture = GraphBuilder.CreateUAV(SubstrateSceneData->TopLayerTexture, ERDGUnorderedAccessViewFlags::SkipBarrier);
+		if (IsSubstrateBlendableGBufferEnabled(View.GetShaderPlatform()))
+		{
+			const FIntRect ViewRect = View.ViewRect;
+
+			FSubstrateDBufferPassPS::FParameters* ParametersPS = GraphBuilder.AllocParameters<FSubstrateDBufferPassPS::FParameters>();
+			TShaderMapRef<FSubstrateDBufferPassPS> PixelShader(View.ShaderMap);
+			ParametersPS->ViewUniformBuffer = View.ViewUniformBuffer;
+			ParametersPS->DBuffer = GetDBufferParameters(GraphBuilder, DBufferTextures, View.GetShaderPlatform());
+			ParametersPS->SceneStencilTexture = SceneTextures.Stencil;
+			ParametersPS->RenderTargets[0] = FRenderTargetBinding(SceneTextures.GBufferA, ERenderTargetLoadAction::ELoad);
+			ParametersPS->RenderTargets[1] = FRenderTargetBinding(SceneTextures.GBufferB, ERenderTargetLoadAction::ELoad);
+			ParametersPS->RenderTargets[2] = FRenderTargetBinding(SceneTextures.GBufferC, ERenderTargetLoadAction::ELoad);
+			ParametersPS->RenderTargets.DepthStencil = FDepthStencilBinding(
+				SceneTextures.Depth.Target,
+				ERenderTargetLoadAction::ELoad,
+				ERenderTargetLoadAction::ELoad,
+				FExclusiveDepthStencil::DepthRead_StencilNop);	
+
+			GraphBuilder.AddPass(
+				RDG_EVENT_NAME("Substrate::DBufferPass(Blendable)"),
+				ParametersPS,
+				ERDGPassFlags::Raster,
+				[ParametersPS, &View, PixelShader](FRHICommandList& InRHICmdList)
+			{
+				InRHICmdList.SetViewport(View.ViewRect.Min.X, View.ViewRect.Min.Y, 0.0f, View.ViewRect.Max.X, View.ViewRect.Max.Y, 1.0f);
+
+				FGraphicsPipelineStateInitializer GraphicsPSOInit;
+				FPixelShaderUtils::InitFullscreenPipelineState(InRHICmdList, View.ShaderMap, PixelShader, GraphicsPSOInit);
+				GraphicsPSOInit.BlendState = TStaticBlendState<
+					CW_RGB, BO_Add, BF_InverseSourceAlpha, BF_SourceAlpha, BO_Add, BF_Zero, BF_One,// Normal
+					CW_RGB, BO_Add, BF_One, BF_SourceAlpha, BO_Add, BF_Zero, BF_One,// BaseColor
+					CW_RGB, BO_Add, BF_One, BF_SourceAlpha, BO_Add, BF_Zero, BF_One	// Metallic, Specular, Roughness
+					>::GetRHI();
+				GraphicsPSOInit.DepthStencilState = TStaticDepthStencilState<false, CF_Always>::GetRHI();
 
-		auto DBufferPass = [&](ESubstrateTileType TileType)
+				SetGraphicsPipelineState(InRHICmdList, GraphicsPSOInit, 0);
+				SetShaderParameters(InRHICmdList, PixelShader, PixelShader.GetPixelShader(), *ParametersPS);
+				FPixelShaderUtils::DrawFullscreenTriangle(InRHICmdList);
+			});
+		}
+		else
 		{
-			// Only simple & single material are support but also dispatch complex tiles, 
-			// as they can contain simple/single material pixels
+			FRDGTextureUAVRef RWMaterialTexture = GraphBuilder.CreateUAV(SubstrateSceneData->MaterialTextureArray, ERDGUnorderedAccessViewFlags::SkipBarrier);
+			FRDGTextureUAVRef RWTopLayerTexture = GraphBuilder.CreateUAV(SubstrateSceneData->TopLayerTexture, ERDGUnorderedAccessViewFlags::SkipBarrier);
 
-			uint32 TilePermutation = 0;
-			switch(TileType)
+			auto DBufferPass = [&](ESubstrateTileType TileType)
 			{
-			case ESubstrateTileType::EComplex:
-			case ESubstrateTileType::EDecalComplex:
-				TilePermutation = 2;
-				break;
-			case ESubstrateTileType::ESingle:
-			case ESubstrateTileType::EDecalSingle:
-				TilePermutation = 1;
-				break;
-			case ESubstrateTileType::ESimple:
-			case ESubstrateTileType::EDecalSimple:
-				TilePermutation = 0;
-				break;
-			}
+				// Only simple & single material are support but also dispatch complex tiles, 
+				// as they can contain simple/single material pixels
 
-			FSubstrateDBufferPassCS::FPermutationDomain PermutationVector;
-			PermutationVector.Set<FSubstrateDBufferPassCS::FTileType>(TilePermutation);
+				uint32 TilePermutation = 0;
+				switch(TileType)
+				{
+				case ESubstrateTileType::EComplex:
+				case ESubstrateTileType::EDecalComplex:
+					TilePermutation = 2;
+					break;
+				case ESubstrateTileType::ESingle:
+				case ESubstrateTileType::EDecalSingle:
+					TilePermutation = 1;
+					break;
+				case ESubstrateTileType::ESimple:
+				case ESubstrateTileType::EDecalSimple:
+					TilePermutation = 0;
+					break;
+				}
 
-			TShaderMapRef<FSubstrateDBufferPassCS> ComputeShader(View.ShaderMap, PermutationVector);
-			FSubstrateDBufferPassCS::FParameters* PassParameters = GraphBuilder.AllocParameters<FSubstrateDBufferPassCS::FParameters>();
+				FSubstrateDBufferPassCS::FPermutationDomain PermutationVector;
+				PermutationVector.Set<FSubstrateDBufferPassCS::FTileType>(TilePermutation);
 
-			PassParameters->DBuffer = GetDBufferParameters(GraphBuilder, DBufferTextures, View.GetShaderPlatform());
-			PassParameters->ViewUniformBuffer = View.ViewUniformBuffer;
-			PassParameters->ViewResolution = View.ViewRect.Size();
-			PassParameters->MaxBytesPerPixel = SubstrateSceneData->EffectiveMaxBytesPerPixel;
-			PassParameters->TopLayerTexture = RWTopLayerTexture;
-			PassParameters->MaterialTextureArrayUAV = RWMaterialTexture;
-			PassParameters->FirstSliceStoringSubstrateSSSData = SubstrateSceneData->FirstSliceStoringSubstrateSSSData;
-			PassParameters->SceneStencilTexture = SceneTextures.Stencil;
+				TShaderMapRef<FSubstrateDBufferPassCS> ComputeShader(View.ShaderMap, PermutationVector);
+				FSubstrateDBufferPassCS::FParameters* PassParameters = GraphBuilder.AllocParameters<FSubstrateDBufferPassCS::FParameters>();
 
-			PassParameters->TileListBuffer = SubstrateViewData->ClassificationTileListBufferSRV;
-			PassParameters->TileListBufferOffset = SubstrateViewData->ClassificationTileListBufferOffset[TileType];
-			PassParameters->TileEncoding = SubstrateViewData->TileEncoding;
-			PassParameters->TileIndirectBuffer = SubstrateViewData->ClassificationTileDispatchIndirectBuffer;
+				PassParameters->DBuffer = GetDBufferParameters(GraphBuilder, DBufferTextures, View.GetShaderPlatform());
+				PassParameters->ViewUniformBuffer = View.ViewUniformBuffer;
+				PassParameters->ViewResolution = View.ViewRect.Size();
+				PassParameters->MaxBytesPerPixel = SubstrateSceneData->EffectiveMaxBytesPerPixel;
+				PassParameters->TopLayerTexture = RWTopLayerTexture;
+				PassParameters->MaterialTextureArrayUAV = RWMaterialTexture;
+				PassParameters->FirstSliceStoringSubstrateSSSData = SubstrateSceneData->FirstSliceStoringSubstrateSSSData;
+				PassParameters->SceneStencilTexture = SceneTextures.Stencil;
 
-			// Dispatch with tile data
-			const uint32 GroupSize = 8;
-			FComputeShaderUtils::AddPass(
-				GraphBuilder,
-				RDG_EVENT_NAME("Substrate::Dbuffer(%s)", ToString(TileType)),
-				ERDGPassFlags::Compute,
-				ComputeShader,
-				PassParameters,
-				PassParameters->TileIndirectBuffer,
-				TileTypeDispatchIndirectArgOffset(TileType));
-		};
+				PassParameters->TileListBuffer = SubstrateViewData->ClassificationTileListBufferSRV;
+				PassParameters->TileListBufferOffset = SubstrateViewData->ClassificationTileListBufferOffset[TileType];
+				PassParameters->TileEncoding = SubstrateViewData->TileEncoding;
+				PassParameters->TileIndirectBuffer = SubstrateViewData->ClassificationTileDispatchIndirectBuffer;
 
-		const bool bDbufferTiles = CVarSubstrateDBufferPassDedicatedTiles.GetValueOnRenderThread() > 0;
-		DBufferPass(bDbufferTiles ? ESubstrateTileType::EDecalComplex : ESubstrateTileType::EComplex);
-		DBufferPass(bDbufferTiles ? ESubstrateTileType::EDecalSingle : ESubstrateTileType::ESingle);
-		DBufferPass(bDbufferTiles ? ESubstrateTileType::EDecalSimple : ESubstrateTileType::ESimple);
+				// Dispatch with tile data
+				const uint32 GroupSize = 8;
+				FComputeShaderUtils::AddPass(
+					GraphBuilder,
+					RDG_EVENT_NAME("Substrate::DbufferPass(%s)", ToString(TileType)),
+					ERDGPassFlags::Compute,
+					ComputeShader,
+					PassParameters,
+					PassParameters->TileIndirectBuffer,
+					TileTypeDispatchIndirectArgOffset(TileType));
+			};
+
+			const bool bDbufferTiles = CVarSubstrateDBufferPassDedicatedTiles.GetValueOnRenderThread() > 0;
+			DBufferPass(bDbufferTiles ? ESubstrateTileType::EDecalComplex : ESubstrateTileType::EComplex);
+			DBufferPass(bDbufferTiles ? ESubstrateTileType::EDecalSingle : ESubstrateTileType::ESingle);
+			DBufferPass(bDbufferTiles ? ESubstrateTileType::EDecalSimple : ESubstrateTileType::ESimple);
+		}
 	}
 }
 
diff --git a/Engine/Source/Runtime/Renderer/Private/Substrate/Substrate.h b/Engine/Source/Runtime/Renderer/Private/Substrate/Substrate.h
--- a/Engine/Source/Runtime/Renderer/Private/Substrate/Substrate.h
+++ b/Engine/Source/Runtime/Renderer/Private/Substrate/Substrate.h
@@ -248,7 +248,7 @@ void BindSubstratePublicGlobalUniformParameters(FRDGBuilder& GraphBuilder, const
 TRDGUniformBufferRef<FSubstratePublicGlobalUniformParameters> CreatePublicGlobalUniformBuffer(FRDGBuilder& GraphBuilder, FSubstrateSceneData* SubstrateScene);
 
 void AddSubstrateMaterialClassificationPass(FRDGBuilder& GraphBuilder, const FMinimalSceneTextures& SceneTextures, const FDBufferTextures& DBufferTextures, const TArray<FViewInfo>& Views);
-void AddSubstrateDBufferPass(FRDGBuilder& GraphBuilder, const FMinimalSceneTextures& SceneTextures, const FDBufferTextures& DBufferTextures, const TArray<FViewInfo>& Views);
+void AddSubstrateDBufferPass(FRDGBuilder& GraphBuilder, const FSceneTextures& SceneTextures, const FDBufferTextures& DBufferTextures, const TArray<FViewInfo>& Views);
 
 void AddSubstrateStencilPass(FRDGBuilder& GraphBuilder, const TArray<FViewInfo>& Views, const FMinimalSceneTextures& SceneTextures);
 void AddSubstrateSampleMaterialPass(FRDGBuilder& GraphBuilder, const FScene* Scene, const FMinimalSceneTextures& SceneTextures, const TArray<FViewInfo>& Views);
