[VSM] Added per-light and per-vsm nanite cluster counts under ShowStats=2. Use DumpLightNaniteStats to print to the output log.

diff --git a/Engine/Shaders/Private/VirtualShadowMaps/VirtualShadowMapPrintStats.usf b/Engine/Shaders/Private/VirtualShadowMaps/VirtualShadowMapPrintStats.usf
--- a/Engine/Shaders/Private/VirtualShadowMaps/VirtualShadowMapPrintStats.usf
+++ b/Engine/Shaders/Private/VirtualShadowMaps/VirtualShadowMapPrintStats.usf
@@ -4,6 +4,7 @@
 #include "../ShaderPrint.ush"
 #include "VirtualShadowMapPageAccessCommon.ush"
 #include "VirtualShadowMapStats.ush"
+#include "VirtualShadowMapProjectionStructs.ush"
 #include "../Nanite/NaniteDataDecode.ush"
 #include "../GPUMessaging.ush"
 #include "../Visualization.ush"
@@ -92,6 +93,8 @@ uint bThrottlingEnabled;
 
 StructuredBuffer<FNaniteStats> NaniteStats;
 
+ByteAddressBuffer PrevProjectionData;
+
 void SendStatusMessage()
 {
 	FGPUMessageWriter Mw = GPUMessageBegin(StatusMessageId, 2U);
@@ -101,7 +104,8 @@ void SendStatusMessage()
 
 void SendStatsMessage()
 {
-	FGPUMessageWriter Mw = GPUMessageBegin(StatsMessageId, 2U + VSM_STAT_NUM + (MAX_PAGE_AREA_DIAGNOSTIC_SLOTS * 2U));
+	uint MessageSize = 2U + VSM_STAT_NUM + (MAX_PAGE_AREA_DIAGNOSTIC_SLOTS * 2U) + (MAX_NPF_DIAGNOSTIC_SLOTS_SINGLEPAGE + MAX_NPF_DIAGNOSTIC_SLOTS_MULTIPAGE) * 2U;
+	FGPUMessageWriter Mw = GPUMessageBegin(StatsMessageId, MessageSize);
 
 	uint ClustersHW = NanitePerformanceFeedback[VSM_NPF_HEADER_TOTAL_HW_CLUSTERS];
 	uint ClustersSW = NanitePerformanceFeedback[VSM_NPF_HEADER_TOTAL_SW_CLUSTERS];
@@ -127,6 +131,34 @@ void SendStatsMessage()
 	{
 		GPUMessageWriteItem(Mw, InStatsBuffer[VSM_STAT_NUM + Index]);
 	}
+
+
+	for (uint Index = 0; Index < MAX_NPF_DIAGNOSTIC_SLOTS_SINGLEPAGE; ++Index)
+	{
+		uint BufferIndex = Index;
+
+		uint CurClustersHW = NanitePerformanceFeedback[VSM_NPF_SIZEOF_HEADER + VSM_NPF_SIZEOF_ENTRY * BufferIndex + VSM_NPF_ENTRY_HW_CLUSTERS];
+		uint CurClustersSW = NanitePerformanceFeedback[VSM_NPF_SIZEOF_HEADER + VSM_NPF_SIZEOF_ENTRY * BufferIndex + VSM_NPF_ENTRY_SW_CLUSTERS];
+		uint CurClustersTotal = CurClustersHW + CurClustersSW;
+
+		FVirtualShadowMapProjectionShaderData ProjectionData = DecodeVirtualShadowMapProjectionData(PrevProjectionData, FVirtualShadowMapHandle::MakeFromId(BufferIndex));
+
+		GPUMessageWriteItem(Mw, ProjectionData.LightId);
+		GPUMessageWriteItem(Mw, CurClustersTotal);
+	}
+	for (uint Index = 0; Index < MAX_NPF_DIAGNOSTIC_SLOTS_MULTIPAGE; ++Index)
+	{
+		uint BufferIndex = VSM_MAX_SINGLE_PAGE_SHADOW_MAPS + Index;
+
+		uint CurClustersHW = NanitePerformanceFeedback[VSM_NPF_SIZEOF_HEADER + VSM_NPF_SIZEOF_ENTRY * BufferIndex + VSM_NPF_ENTRY_HW_CLUSTERS];
+		uint CurClustersSW = NanitePerformanceFeedback[VSM_NPF_SIZEOF_HEADER + VSM_NPF_SIZEOF_ENTRY * BufferIndex + VSM_NPF_ENTRY_SW_CLUSTERS];
+		uint CurClustersTotal = CurClustersHW + CurClustersSW;
+
+		FVirtualShadowMapProjectionShaderData ProjectionData = DecodeVirtualShadowMapProjectionData(PrevProjectionData, FVirtualShadowMapHandle::MakeFromId(BufferIndex));
+
+		GPUMessageWriteItem(Mw, ProjectionData.LightId);
+		GPUMessageWriteItem(Mw, CurClustersTotal);
+	}
 }
 
 #define NANITE_GEO_HISTOGRAM_NUM_BINS (VSM_GENERATE_STATS_GROUPSIZE*2)
diff --git a/Engine/Shaders/Private/VirtualShadowMaps/VirtualShadowMapProjectionStructs.ush b/Engine/Shaders/Private/VirtualShadowMaps/VirtualShadowMapProjectionStructs.ush
--- a/Engine/Shaders/Private/VirtualShadowMaps/VirtualShadowMapProjectionStructs.ush
+++ b/Engine/Shaders/Private/VirtualShadowMaps/VirtualShadowMapProjectionStructs.ush
@@ -39,6 +39,8 @@ struct FVirtualShadowMapProjectionShaderData
 	int SceneRendererPrimaryViewId;
 	int PersistentViewId;
 
+	int LightId;
+
 	// Derived data for convenience when passing the structure around
 	//int VirtualShadowMapId;
 	FVirtualShadowMapHandle VirtualShadowMapHandle;
@@ -103,6 +105,8 @@ FVirtualShadowMapProjectionShaderData DecodeVirtualShadowMapProjectionData(ByteA
 	
 	LoadAndIncrementOffset(Result.TexelDitherScale, ProjectionData, CurrentLoadOffset);
 	LoadAndIncrementOffset(Result.MinMipLevel, ProjectionData, CurrentLoadOffset);
+
+	LoadAndIncrementOffset(Result.LightId, ProjectionData, CurrentLoadOffset);
 	
 	// Computed data
 	Result.PreViewTranslation				= MakeDFVector3(PreViewTranslationHigh, PreViewTranslationLow);
diff --git a/Engine/Shaders/Shared/VirtualShadowMapDefinitions.h b/Engine/Shaders/Shared/VirtualShadowMapDefinitions.h
--- a/Engine/Shaders/Shared/VirtualShadowMapDefinitions.h
+++ b/Engine/Shaders/Shared/VirtualShadowMapDefinitions.h
@@ -98,7 +98,7 @@ VirtualShadowMapDefinitions.h: used in virtual shadow map shaders and C++ code t
 #define VSM_STATUS_MSG_OVERFLOW							1
 
 // Projection Shader Data
-#define VSM_PSD_STRIDE (16 * 18)
+#define VSM_PSD_STRIDE (16 * 19)
 
 // Nanite Performance Feedback
 #define VSM_NPF_HEADER_TOTAL_HW_CLUSTERS				0
diff --git a/Engine/Source/Runtime/Renderer/Private/Shadows/ShadowSceneRenderer.cpp b/Engine/Source/Runtime/Renderer/Private/Shadows/ShadowSceneRenderer.cpp
--- a/Engine/Source/Runtime/Renderer/Private/Shadows/ShadowSceneRenderer.cpp
+++ b/Engine/Source/Runtime/Renderer/Private/Shadows/ShadowSceneRenderer.cpp
@@ -484,6 +484,7 @@ TSharedPtr<FVirtualShadowMapPerLightCacheEntry> FShadowSceneRenderer::AddLocalLi
 			ProjectionData.Flags				= Flags;
 			ProjectionData.MinMipLevel			= MinMipLevel;
 			ProjectionData.PackedCullingViewId	= PackedCullingViewId;
+			ProjectionData.LightId				= Proxy->GetLightSceneInfo()->GetPersistentIndex();
 		}
 	}
 
diff --git a/Engine/Source/Runtime/Renderer/Private/VirtualShadowMaps/VirtualShadowMapArray.cpp b/Engine/Source/Runtime/Renderer/Private/VirtualShadowMaps/VirtualShadowMapArray.cpp
--- a/Engine/Source/Runtime/Renderer/Private/VirtualShadowMaps/VirtualShadowMapArray.cpp
+++ b/Engine/Source/Runtime/Renderer/Private/VirtualShadowMaps/VirtualShadowMapArray.cpp
@@ -109,7 +109,7 @@ TAutoConsoleVariable<int32> CVarBuildDynamicHZB(
 	ECVF_RenderThreadSafe
 );
 
-static TAutoConsoleVariable<int32> CVarShowStats(
+TAutoConsoleVariable<int32> CVarVSMShowStats(
 	TEXT("r.Shadow.Virtual.ShowStats"),
 	0,
 	TEXT("Show VSM statistics."),
@@ -916,6 +916,8 @@ void FVirtualShadowMapArray::SetShaderDefines(FShaderCompilerEnvironment& OutEnv
 	static_assert(FVirtualShadowMap::Log2Level0DimPagesXY * 2U + NANITE_MAX_VIEWS_PER_CULL_RASTERIZE_PASS_BITS <= 32U, "Page indirection plus view index must fit into 32-bits for page-routing storage!");
 	OutEnvironment.SetDefine(TEXT("ENABLE_NON_NANITE_VSM"), GEnableNonNaniteVSM);
 	OutEnvironment.SetDefine(TEXT("MAX_PAGE_AREA_DIAGNOSTIC_SLOTS"), MaxPageAreaDiagnosticSlots);
+	OutEnvironment.SetDefine(TEXT("MAX_NPF_DIAGNOSTIC_SLOTS_SINGLEPAGE"), MaxNPFDiagnosticSlotsSinglePage);
+	OutEnvironment.SetDefine(TEXT("MAX_NPF_DIAGNOSTIC_SLOTS_MULTIPAGE"), MaxNPFDiagnosticSlotsMultiPage);
 	OutEnvironment.SetDefine(TEXT("INDEX_NONE"), INDEX_NONE);
 }
 
@@ -1869,7 +1871,7 @@ bool FVirtualShadowMapArray::ShouldGenerateStats() const
 	constexpr bool bRunPageAreaDiagnostics = false;
 	constexpr bool bInsightsVSMChannelEnabled = false;
 #endif
-	return CVarShowStats.GetValueOnRenderThread() || CacheManager->IsAccumulatingStats() || bRunPageAreaDiagnostics || IsCsvLogEnabled() || bInsightsVSMChannelEnabled;
+	return CVarVSMShowStats.GetValueOnRenderThread() || CacheManager->IsAccumulatingStats() || bRunPageAreaDiagnostics || IsCsvLogEnabled() || bInsightsVSMChannelEnabled;
 }
 
 bool FVirtualShadowMapArray::IsCsvLogEnabled() const
@@ -2179,6 +2181,18 @@ void FVirtualShadowMapArray::BuildPageAllocations(
 		ThrottlingParameters.ThrottleEMAHistoryWeight = CVarThrottleEMAHistoryWeight.GetValueOnRenderThread();
 	}
 
+	// Is used during logging
+	{
+		if (CacheManager->GetPrevBuffers().ProjectionData)
+		{
+			PrevProjectionData = GraphBuilder.RegisterExternalBuffer(CacheManager->GetPrevBuffers().ProjectionData);
+		}
+		else
+		{
+			PrevProjectionData = GSystemTextures.GetDefaultByteAddressBuffer(GraphBuilder, sizeof(FVirtualShadowMapProjectionShaderData));
+		}
+	}
+
 	if (bThrottlingEnabled && CacheManager->GetPrevBuffers().NanitePerformanceFeedback)
 	{
 		FProcessPrevFramePerfDataCS::FParameters* PassParameters = GraphBuilder.AllocParameters<FProcessPrevFramePerfDataCS::FParameters>();
@@ -2854,6 +2868,7 @@ class FVirtualSmLogStatsCS : public FVirtualShadowMapPageManagementShader
 		SHADER_PARAMETER_RDG_BUFFER_SRV(StructuredBuffer<uint>, InStatsBuffer)
 		SHADER_PARAMETER_RDG_BUFFER_SRV(StructuredBuffer<FNaniteStats>, NaniteStats)
 		SHADER_PARAMETER_RDG_BUFFER_SRV(StructuredBuffer<uint>, NanitePerformanceFeedback)
+		SHADER_PARAMETER_RDG_BUFFER_SRV(ByteAddressBuffer, PrevProjectionData)
 		SHADER_PARAMETER(uint32, NanitePerformanceFeedbackNumEntries)
 		SHADER_PARAMETER_RDG_BUFFER_SRV(StructuredBuffer<uint>, ThrottleBuffer)
 		SHADER_PARAMETER(uint32, bThrottlingEnabled)
@@ -2902,7 +2917,7 @@ void FVirtualShadowMapArray::LogStats(FRDGBuilder& GraphBuilder, const FViewInfo
 		ShaderPrint::SetParameters(GraphBuilder, View.ShaderPrintData, PassParameters->ShaderPrintStruct);
 		ShaderPrint::RequestSpaceForTriangles(8192);
 
-		int ShowStatsValue = CVarShowStats.GetValueOnRenderThread();
+		int ShowStatsValue = CVarVSMShowStats.GetValueOnRenderThread();
 		PassParameters->ShowStatsValue = ShowStatsValue;
 
 #if !UE_BUILD_SHIPPING
@@ -2921,6 +2936,8 @@ void FVirtualShadowMapArray::LogStats(FRDGBuilder& GraphBuilder, const FViewInfo
 		{
 			PassParameters->NaniteStats = GraphBuilder.CreateSRV(GSystemTextures.GetDefaultStructuredBuffer(GraphBuilder, sizeof(FNaniteStats)));
 		}
+
+		PassParameters->PrevProjectionData = GraphBuilder.CreateSRV(PrevProjectionData);
 		
 		PassParameters->NanitePerformanceFeedback = GraphBuilder.CreateSRV(NanitePerformanceFeedbackRDG);
 		PassParameters->NanitePerformanceFeedbackNumEntries = GetNumShadowMapSlots();
diff --git a/Engine/Source/Runtime/Renderer/Private/VirtualShadowMaps/VirtualShadowMapArray.h b/Engine/Source/Runtime/Renderer/Private/VirtualShadowMaps/VirtualShadowMapArray.h
--- a/Engine/Source/Runtime/Renderer/Private/VirtualShadowMaps/VirtualShadowMapArray.h
+++ b/Engine/Source/Runtime/Renderer/Private/VirtualShadowMaps/VirtualShadowMapArray.h
@@ -128,7 +128,11 @@ struct FVirtualShadowMapProjectionShaderData
 	float TexelDitherScale;
 	
 	uint32 MinMipLevel = 0u;
+
+	int32 LightId = -1;
+
 	// Note: Seems the FMatrix forces 16-byte alignment so pad as needed.
+	uint32 Padding[3];
 
 	static uint32 PackCullingViewId(int32 SceneRendererPrimaryViewId, const FPersistentViewId& PersistentViewId);
 };
@@ -482,6 +486,8 @@ public:
 	bool bHZBBuiltThisFrame = false;
 
 	static constexpr uint32 MaxPageAreaDiagnosticSlots = 32;
+	static constexpr uint32 MaxNPFDiagnosticSlotsSinglePage = 32;
+	static constexpr uint32 MaxNPFDiagnosticSlotsMultiPage = 96;
 
 	FRDGBufferRef StatsBufferRDG = nullptr;
 	FRDGBufferUAVRef StatsBufferUAV = nullptr;
@@ -543,4 +549,6 @@ private:
 	bool bCullBackfacingPixels = false;
 
 	FRDGTextureRef PageTableDummyRDG = nullptr;
+
+	FRDGBufferRef PrevProjectionData = nullptr;
 };
diff --git a/Engine/Source/Runtime/Renderer/Private/VirtualShadowMaps/VirtualShadowMapCacheManager.cpp b/Engine/Source/Runtime/Renderer/Private/VirtualShadowMaps/VirtualShadowMapCacheManager.cpp
--- a/Engine/Source/Runtime/Renderer/Private/VirtualShadowMaps/VirtualShadowMapCacheManager.cpp
+++ b/Engine/Source/Runtime/Renderer/Private/VirtualShadowMaps/VirtualShadowMapCacheManager.cpp
@@ -157,6 +157,17 @@ FAutoConsoleVariableRef CVarVSMAllowScreenOverflowMessages(
 	ECVF_RenderThreadSafe
 );
 
+extern TAutoConsoleVariable<int32> CVarVSMShowStats;
+
+static bool bDumpLightNaniteStats = false;
+static FAutoConsoleCommandWithArgsAndOutputDevice DumpLightNaniteStatsConsoleCommand(
+	TEXT("r.Shadow.Virtual.DumpLightNaniteStats"),
+	TEXT("Print per-vsm Nanite cluster counts as a performance profiling tool."),
+	FConsoleCommandWithArgsAndOutputDeviceDelegate::CreateLambda([](const TArray<FString>& Args, FOutputDevice& Output)
+		{
+			bDumpLightNaniteStats = true;
+		}));
+
 static const char *VirtualShadowMap_StatNames[] = 
 {
 	"REQUESTED_THIS_FRAME_PAGES",
@@ -885,6 +896,62 @@ void FVirtualShadowMapArrayCacheManager::InitExtension(FScene& InScene)
 				LargePageAreaItems.Add(PersistentPrimitiveId, FLargePageAreaItem{ Overlap, RealTimeSeconds });
 			}
 		}
+
+		auto ParseNPFDiagnostic = [&](uint32 VSMId, uint32 LightId, uint32 ClusterCount)
+		{
+			if (ClusterCount > 0)
+			{
+				FString LightNameWithLevel;
+				if (LightId <= uint64(Scene.Lights.GetMaxIndex()))
+				{
+					FSceneRenderer::GetLightNameForDrawEvent(Scene.Lights[LightId].LightSceneInfo->Proxy, LightNameWithLevel);
+				}
+				NPFDiagnosticMessages.Add(FString::Printf(TEXT("%8d | %48s | %12d"), VSMId, *LightNameWithLevel, ClusterCount));
+				NPFDiagnosticTimer.FindOrAdd(VSMId) = 180;
+			}
+			else
+			{
+				uint8* FramesRemaining = NPFDiagnosticTimer.Find(VSMId);
+				if (FramesRemaining != nullptr)
+				{
+					NPFDiagnosticMessages.Add(FString(TEXT("")));
+					if ((*FramesRemaining)-- == 0)
+					{
+						NPFDiagnosticTimer.Remove(VSMId);
+					}
+				}
+			}
+		};
+
+		NPFDiagnosticMessages.SetNum(0, EAllowShrinking::No);
+		NPFDiagnosticMessages.Add(FString::Printf(TEXT("%8s | %48s | %12s"), TEXT("VSM Id"), TEXT("Light"), TEXT("# Clusters")));
+		TConstArrayView<uint32> NPFDiagsSinglePage = Message.ReadCount(FVirtualShadowMapArray::MaxNPFDiagnosticSlotsSinglePage * 2U);
+		for (int32 Index = 0; Index < NPFDiagsSinglePage.Num(); Index += 2)
+		{
+			uint32 VSMId = Index / 2;
+			uint32 LightId = NPFDiagsSinglePage[Index];
+			uint32 ClusterCount = NPFDiagsSinglePage[Index + 1];
+			ParseNPFDiagnostic(VSMId, LightId, ClusterCount);
+		}
+		TConstArrayView<uint32> NPFDiagsMultiPage = Message.ReadCount(FVirtualShadowMapArray::MaxNPFDiagnosticSlotsMultiPage * 2U);
+		for (int32 Index = 0; Index < NPFDiagsMultiPage.Num(); Index += 2)
+		{
+			uint32 VSMId = VSM_MAX_SINGLE_PAGE_SHADOW_MAPS + (Index / 2);
+			uint32 LightId = NPFDiagsMultiPage[Index];
+			uint32 ClusterCount = NPFDiagsMultiPage[Index + 1];
+			ParseNPFDiagnostic(VSMId, LightId, ClusterCount);
+		}
+		if (bDumpLightNaniteStats)
+		{
+			bDumpLightNaniteStats = false;
+			for (const FString& Str : NPFDiagnosticMessages)
+			{
+				if (!Str.IsEmpty())
+				{
+					UE_LOG(LogRenderer, Display, TEXT("%s"), *Str);
+				}
+			}
+		}
 	});
 #endif
 
@@ -926,6 +993,15 @@ void FVirtualShadowMapArrayCacheManager::InitExtension(FScene& InScene)
 		{
 			OutMessages.Add(FCoreDelegates::EOnScreenMessageSeverity::Warning, FText::FromString(FString::Printf(TEXT("Virtual Shadow Map Stats Accumulation Active. Set r.Shadow.Virtual.AccumulateStats to 0 to stop."))));
 		}
+
+		if (CVarVSMShowStats.GetValueOnRenderThread() > 1)
+		{
+			for (const FString& Message : NPFDiagnosticMessages)
+			{
+				OutMessages.Add(FCoreDelegates::EOnScreenMessageSeverity::Info, FText::FromString(Message));
+			}
+			NPFDiagnosticMessages.SetNum(0, EAllowShrinking::No);
+		}
 	});
 #endif
 }
diff --git a/Engine/Source/Runtime/Renderer/Private/VirtualShadowMaps/VirtualShadowMapCacheManager.h b/Engine/Source/Runtime/Renderer/Private/VirtualShadowMaps/VirtualShadowMapCacheManager.h
--- a/Engine/Source/Runtime/Renderer/Private/VirtualShadowMaps/VirtualShadowMapCacheManager.h
+++ b/Engine/Source/Runtime/Renderer/Private/VirtualShadowMaps/VirtualShadowMapCacheManager.h
@@ -537,6 +537,9 @@ private:
 		float LastTimeSeen;
 	};
 	TMap<uint32, FLargePageAreaItem> LargePageAreaItems;
+
+	TArray<FString> NPFDiagnosticMessages {};
+	TMap<uint32, uint8> NPFDiagnosticTimer {};
 #endif // UE_BUILD_SHIPPING
 
 	FShadowInvalidatingInstancesImplementation ShadowInvalidatingInstancesImplementation;
diff --git a/Engine/Source/Runtime/Renderer/Private/VirtualShadowMaps/VirtualShadowMapClipmap.cpp b/Engine/Source/Runtime/Renderer/Private/VirtualShadowMaps/VirtualShadowMapClipmap.cpp
--- a/Engine/Source/Runtime/Renderer/Private/VirtualShadowMaps/VirtualShadowMapClipmap.cpp
+++ b/Engine/Source/Runtime/Renderer/Private/VirtualShadowMaps/VirtualShadowMapClipmap.cpp
@@ -540,6 +540,7 @@ FVirtualShadowMapProjectionShaderData FVirtualShadowMapClipmap::ComputeProjectio
 	Data.LightSourceRadius = GetLightSceneInfo().Proxy->GetSourceRadius();
 	Data.Flags = PerLightCacheEntry->IsUncached() ? VSM_PROJ_FLAG_UNCACHED : 0U;
 	Data.PackedCullingViewId = FVirtualShadowMapProjectionShaderData::PackCullingViewId(DependentView->SceneRendererPrimaryViewId, DependentView->PersistentViewId);
+	Data.LightId = GetLightSceneInfo().GetPersistentIndex();
 
 	if (Config.FirstCoarseLevel >= 0)
 	{
