Turn off compute and draw control flow for LensFlareBlur on Android devices.

This codepath is no longer necessary with the addition of GRHINeedsSRVGraphicsNonPixelWorkaround

diff --git a/Engine/Config/Android/AndroidEngine.ini b/Engine/Config/Android/AndroidEngine.ini
--- a/Engine/Config/Android/AndroidEngine.ini
+++ b/Engine/Config/Android/AndroidEngine.ini
@@ -93,9 +93,6 @@ r.Vulkan.RHIThread=1
 r.Vulkan.AllowAsyncCompute=0
 ; Disable split barriers until its effects are measured on all devices
 r.Vulkan.AllowSplitBarriers=0
-; Split LensFlareBlur into two passes to remove vertex shader texture sample
-; Splitting into compute shader is necessary on Android. See UE-243673
-r.LensFlareBlurComputeShader=1
 ; Crash reporting diagnotic buffer
 r.Vulkan.DiagnosticBuffer=1
 ; Transient memory allocator
diff --git a/Engine/Source/Runtime/Renderer/Private/PostProcess/PostProcessLensFlares.cpp b/Engine/Source/Runtime/Renderer/Private/PostProcess/PostProcessLensFlares.cpp
--- a/Engine/Source/Runtime/Renderer/Private/PostProcess/PostProcessLensFlares.cpp
+++ b/Engine/Source/Runtime/Renderer/Private/PostProcess/PostProcessLensFlares.cpp
@@ -244,6 +244,7 @@ FScreenPassTexture AddLensFlaresPass(
 
 	// Lens flare blur pass. Rasterizes a bokeh quad for each pixel on the screen based on the intensity threshold.
 	// Decide between one draw call versus one compute and one draw
+	// TODO: The split into compute and draw is no longer necessary and can be removed. See UE-297792
 	if (CVarLensFlareBlurComputeShader.GetValueOnRenderThread() == 0)
 	{
 		const uint32 TileSizeInPixels = 1;
