Move initialization of SceneViewExtensions owned by the renderer to module startup to fix crash caused by static init order, and add additional asserts to MeshEdgesRendering.

#tests A CI build of Mac TwinMotion was used to validate that this fixes the crash.

diff --git a/Engine/Source/Runtime/Renderer/Private/MeshEdgesRendering.cpp b/Engine/Source/Runtime/Renderer/Private/MeshEdgesRendering.cpp
--- a/Engine/Source/Runtime/Renderer/Private/MeshEdgesRendering.cpp
+++ b/Engine/Source/Runtime/Renderer/Private/MeshEdgesRendering.cpp
@@ -181,8 +181,7 @@ const FMeshEdgesViewSettings& GetMeshEdgesViewSettings(const FSceneViewFamily& V
 
 	if (FamilyData->ViewSettings.IsEmpty())
 	{
-		// Allocate 1 extra as a fallback for weird edgecase where the view can't be found in the viewfamily.
-		FamilyData->ViewSettings.Init(FMeshEdgesViewSettings{}, ViewFamily.Views.Num() + 1);
+		FamilyData->ViewSettings.Init(FMeshEdgesViewSettings{}, ViewFamily.Views.Num());
 	}
 
 	return FamilyData->ViewSettings[ViewIndex];
@@ -224,33 +223,26 @@ class FMeshEdgesExtension : public FSceneViewExtensionBase
 public:
 	FMeshEdgesExtension( const FAutoRegister& AutoRegister)
 	: FSceneViewExtensionBase( AutoRegister )
-	{}
+	{
+		FCoreDelegates::OnEnginePreExit.AddLambda([]()
+			{ 
+				Instance = nullptr;
+			});
+	}
 
 	virtual void PostCreateSceneRenderer(const FSceneViewFamily& InViewFamily, ISceneRenderer* Renderer) override
 	{
 		RenderMeshEdges(static_cast<FSceneRenderer*>(Renderer)->ViewFamily);
 	}
 
-	inline static TSharedPtr<FMeshEdgesExtension,ESPMode::ThreadSafe> Instance;
-
-private:
-	struct FStaticConstructor
-	{
-		FStaticConstructor()
-		{
-			FCoreDelegates::OnPostEngineInit.AddLambda([]()
-			{ 
-				Instance = FSceneViewExtensions::NewExtension<FMeshEdgesExtension>(); 
-			});
-			FCoreDelegates::OnEnginePreExit.AddLambda([]()
-			{ 
-				Instance = nullptr;
-			});
-		}
-	};
-	static inline FStaticConstructor StaticConstructor;
+	inline static TSharedPtr<FMeshEdgesExtension,ESPMode::ThreadSafe> Instance = nullptr;
 };
 
+void InitMeshEdgesViewExtension()
+{
+	FMeshEdgesExtension::Instance = FSceneViewExtensions::NewExtension<FMeshEdgesExtension>(); 
+}
+
 void CopyViewFamily(const FSceneViewFamily& SrcViewFamily, FSceneViewFamily& ViewFamily)
 {
 	ViewFamily.FrameNumber = SrcViewFamily.FrameNumber;
@@ -393,6 +385,7 @@ void ComposeMeshEdges(FRDGBuilder& GraphBuilder,
                         FScreenPassRenderTarget& EditorPrimitivesColor,
                         FScreenPassRenderTarget& EditorPrimitivesDepth)
 {
+	check(View.Family);
 	const FSceneViewFamily& ViewFamily = *View.Family;
 	if (!ViewFamily.EngineShowFlags.MeshEdges)
 	{
@@ -405,8 +398,17 @@ void ComposeMeshEdges(FRDGBuilder& GraphBuilder,
 		if (ViewFamily.Views[ViewIndex] == &View) break;
 	}
 
+	if (ViewIndex >= ViewFamily.Views.Num())
+	{
+		// View was not found in its own ViewFamily (?), let's abort.
+		return;
+	}
+
 	const FMeshEdgesViewFamilyData* ViewFamilyData = ViewFamily.GetExtentionData<FMeshEdgesViewFamilyData>();
 	check(ViewFamilyData); // should have been created in RenderMeshEdges
+	check(ViewFamilyData->ViewRects.Num() > ViewIndex);
+	check(ViewFamilyData->WireframeColor);
+	check(ViewFamilyData->WireframeDepth);
 	const FMeshEdgesViewSettings& ViewSettings = GetMeshEdgesViewSettings(ViewFamily, ViewIndex);
 
 	FRenderTargetTexture& WireframeTextureColor = *ViewFamilyData->WireframeColor;
@@ -443,8 +445,8 @@ void ComposeMeshEdges(FRDGBuilder& GraphBuilder,
 	FComposeMeshEdgesPS::FPermutationDomain PermutationVector;
 	PermutationVector.Set<FComposeMeshEdgesPS::FSampleCountDimension>(MSAASampleCountDim);
 
-	const FGlobalShaderMap* GlobalShaderMap = View.ShaderMap;
-	const TShaderRef<FComposeMeshEdgesPS>& PixelShader = TShaderMapRef<FComposeMeshEdgesPS>(View.ShaderMap, PermutationVector);
+	check(View.ShaderMap);
+	const TShaderRef<FComposeMeshEdgesPS> PixelShader = View.ShaderMap->GetShader<FComposeMeshEdgesPS>(PermutationVector);
 
 	FIntRect Viewport = EditorPrimitivesColor.ViewRect;
 
diff --git a/Engine/Source/Runtime/Renderer/Private/MeshEdgesRendering.h b/Engine/Source/Runtime/Renderer/Private/MeshEdgesRendering.h
--- a/Engine/Source/Runtime/Renderer/Private/MeshEdgesRendering.h
+++ b/Engine/Source/Runtime/Renderer/Private/MeshEdgesRendering.h
@@ -9,6 +9,8 @@ class FViewInfo;
 struct FCompositePrimitiveInputs;
 struct FScreenPassRenderTarget;
 
+void InitMeshEdgesViewExtension();
+
 void ComposeMeshEdges(FRDGBuilder& GraphBuilder,
 						const FViewInfo& View,
 						FScreenPassRenderTarget& EditorPrimitivesColor,
diff --git a/Engine/Source/Runtime/Renderer/Private/Renderer.cpp b/Engine/Source/Runtime/Renderer/Private/Renderer.cpp
--- a/Engine/Source/Runtime/Renderer/Private/Renderer.cpp
+++ b/Engine/Source/Runtime/Renderer/Private/Renderer.cpp
@@ -51,6 +51,8 @@
 #include "PathTracing.h"
 #include "LightFunctionAtlas.h"
 #include "SceneRenderBuilder.h"
+#include "MeshEdgesRendering.h"
+#include "Skinning/SkinnedMeshDebugView.h"
 
 DEFINE_LOG_CATEGORY(LogRenderer);
 
@@ -117,6 +119,15 @@ void FRendererModule::StartupModule()
 		});
 	}
 #endif
+
+	if (FApp::CanEverRender() && !FApp::IsUnattended())
+	{
+		FCoreDelegates::OnPostEngineInit.AddLambda([]()
+		{
+			InitMeshEdgesViewExtension();
+			InitSkinnedMeshDebugViewExtension();
+		});
+	}
 }
 
 void FRendererModule::ShutdownModule()
diff --git a/Engine/Source/Runtime/Renderer/Private/Skinning/SkinnedMeshDebugView.cpp b/Engine/Source/Runtime/Renderer/Private/Skinning/SkinnedMeshDebugView.cpp
--- a/Engine/Source/Runtime/Renderer/Private/Skinning/SkinnedMeshDebugView.cpp
+++ b/Engine/Source/Runtime/Renderer/Private/Skinning/SkinnedMeshDebugView.cpp
@@ -87,6 +87,10 @@ IMPLEMENT_GLOBAL_SHADER(FSkinnedMeshDrawLinePS, "/Engine/Private/Skinning/Skinne
 FSkinnedMeshDebugViewExtension::FSkinnedMeshDebugViewExtension(const FAutoRegister& AutoRegister)
 : FSceneViewExtensionBase(AutoRegister)
 {
+	FCoreDelegates::OnEnginePreExit.AddLambda([]()
+		{
+			Instance = nullptr;
+		});
 }
 
 void FSkinnedMeshDebugViewExtension::SubscribeToPostProcessingPass(EPostProcessingPass PassId, const FSceneView& View, FAfterPassCallbackDelegateArray& InOutPassCallbacks, bool bIsPassEnabled)
diff --git a/Engine/Source/Runtime/Renderer/Private/Skinning/SkinnedMeshDebugView.h b/Engine/Source/Runtime/Renderer/Private/Skinning/SkinnedMeshDebugView.h
--- a/Engine/Source/Runtime/Renderer/Private/Skinning/SkinnedMeshDebugView.h
+++ b/Engine/Source/Runtime/Renderer/Private/Skinning/SkinnedMeshDebugView.h
@@ -27,25 +27,20 @@ public:
 
 	void RenderSkeletons(class FRDGBuilder& GraphBuilder, const class FSceneView& View, const struct FScreenPassRenderTarget& Output);
 
+	static void Init()
+	{
+		Instance = FSceneViewExtensions::NewExtension<FSkinnedMeshDebugViewExtension>();
+	}
+
 private:
 	inline static TSharedPtr<FSkinnedMeshDebugViewExtension, ESPMode::ThreadSafe> Instance;
-
-	struct FStaticConstructor
-	{
-		FStaticConstructor()
-		{
-			FCoreDelegates::OnPostEngineInit.AddLambda([]()
-			{ 
-				Instance = FSceneViewExtensions::NewExtension<FSkinnedMeshDebugViewExtension>();
-			});
-			FCoreDelegates::OnEnginePreExit.AddLambda([]()
-			{ 
-				Instance = nullptr;
-			});
-		}
-	};
-
-	static inline FStaticConstructor StaticConstructor;
 };
 
 #endif
+
+static void InitSkinnedMeshDebugViewExtension()
+{
+#if UE_ENABLE_DEBUG_DRAWING
+	FSkinnedMeshDebugViewExtension::Init();
+#endif
+}
\ No newline at end of file
