Fixed issue in Lit Wireframe where turning on lens distortion would cause the wireframe to disappear. Since the distortion is applied after the wireframe is composited, this issue is fixed by simply disabling distortion on the wireframe pass.

#tests Fix confirmed in editor with r.LensDistortion.Panini.D

diff --git a/Engine/Source/Runtime/Renderer/Private/MeshEdgesRendering.cpp b/Engine/Source/Runtime/Renderer/Private/MeshEdgesRendering.cpp
--- a/Engine/Source/Runtime/Renderer/Private/MeshEdgesRendering.cpp
+++ b/Engine/Source/Runtime/Renderer/Private/MeshEdgesRendering.cpp
@@ -310,6 +310,7 @@ void RenderMeshEdges(FSceneViewFamily& ViewFamily)
 		WireframeShowFlags.SetCompositeEditorPrimitives(false);
 		WireframeShowFlags.SetGrid(false);
 		WireframeShowFlags.SetShaderPrint(false);
+		WireframeShowFlags.SetLensDistortion(false);
 		//WireframeShowFlags.SetScreenPercentage(false);
 		//WireframeShowFlags.SetTranslucency(false);
 	}
