Add Dither LOD using Stencil on mobile for shader platforms that support it

diff --git a/Engine/Source/Runtime/Core/Private/HAL/ConsoleManager.cpp b/Engine/Source/Runtime/Core/Private/HAL/ConsoleManager.cpp
--- a/Engine/Source/Runtime/Core/Private/HAL/ConsoleManager.cpp
+++ b/Engine/Source/Runtime/Core/Private/HAL/ConsoleManager.cpp
@@ -3740,7 +3740,7 @@ static TAutoConsoleVariable<int32> CVarMobileFloatPrecisionMode(
 		 "2: Force use of high precision in pixel shaders.\n"),
 	ECVF_ReadOnly);
 
-static TAutoConsoleVariable<int32> CVarMobileAllowDitheredLODTransition(
+static TAutoConsoleVariable<bool> CVarMobileAllowDitheredLODTransition(
 	TEXT("r.Mobile.AllowDitheredLODTransition"),
 	0,
 	TEXT("Whether to support 'Dithered LOD Transition' material option on mobile platforms"),
diff --git a/Engine/Source/Runtime/Engine/Classes/Engine/RendererSettings.h b/Engine/Source/Runtime/Engine/Classes/Engine/RendererSettings.h
--- a/Engine/Source/Runtime/Engine/Classes/Engine/RendererSettings.h
+++ b/Engine/Source/Runtime/Engine/Classes/Engine/RendererSettings.h
@@ -342,7 +342,7 @@ class URendererSettings : public UDeveloperSettings
 		ConsoleVariable="r.Mobile.AllowDitheredLODTransition", DisplayName="Allow Dithered LOD Transition",
 		ToolTip="Whether to support 'Dithered LOD Transition' material option on mobile platforms. Enabling this may degrade performance as rendering will not benefit from Early-Z optimization.",
 		ConfigRestartRequired=true))
-	uint32 bMobileAllowDitheredLODTransition:1;
+	bool bMobileAllowDitheredLODTransition;
 
 	UPROPERTY(config, EditAnywhere, Category = Mobile, meta = (
 		EditCondition = "bVirtualTextures",
diff --git a/Engine/Source/Runtime/Engine/Private/Materials/MaterialShared.cpp b/Engine/Source/Runtime/Engine/Private/Materials/MaterialShared.cpp
--- a/Engine/Source/Runtime/Engine/Private/Materials/MaterialShared.cpp
+++ b/Engine/Source/Runtime/Engine/Private/Materials/MaterialShared.cpp
@@ -210,10 +210,9 @@ bool AllowDitheredLODTransition(ERHIFeatureLevel::Type FeatureLevel)
 bool AllowDitheredLODTransition(EShaderPlatform ShaderPlatform)
 {
 	// On mobile support for 'Dithered LOD Transition' has to be explicitly enabled in projects settings
-	if (IsMobilePlatform(ShaderPlatform))
+	if (IsMobilePlatform(ShaderPlatform) && !FReadOnlyCVARCache::MobileAllowDitheredLODTransition(ShaderPlatform))
 	{
-		static FShaderPlatformCachedIniValue<int32> CVarAllowDitheredLODTransition(TEXT("r.Mobile.AllowDitheredLODTransition"));
-		return (CVarAllowDitheredLODTransition.Get(ShaderPlatform) != 0) ? true : false;
+		return false;
 	}
 	return true;
 }
@@ -2721,11 +2720,7 @@ void FMaterial::SetupMaterialEnvironment(
 		SET_SHADER_DEFINE(OutEnvironment, EDITOR_PRIMITIVE_MATERIAL, 1);
 	}
 
-	if (IsFeatureLevelSupported(GetShaderPlatform(), ERHIFeatureLevel::SM5))
-	{	
-		static const auto CVar = IConsoleManager::Get().FindTConsoleVariableDataInt(TEXT("r.StencilForLODDither"));
-		SET_SHADER_DEFINE(OutEnvironment, USE_STENCIL_LOD_DITHER_DEFAULT, CVar->GetValueOnAnyThread() != 0 ? 1 : 0);
-	}
+	SET_SHADER_DEFINE(OutEnvironment, USE_STENCIL_LOD_DITHER_DEFAULT, IsStencilForLODDitherEnabled(GetShaderPlatform()) != 0 ? 1 : 0);
 
 	{
 		switch (GetMaterialDomain())
@@ -4639,7 +4634,7 @@ bool FMaterial::WritesEveryPixel(bool bShadowPass) const
 
 bool FMaterial::WritesEveryPixel(bool bShadowPass, bool bVFTypeSupportsNullPixelShader) const
 {
-	bool bLocalStencilDitheredLOD = FeatureLevel >= ERHIFeatureLevel::SM5 && bStencilDitheredLOD;
+	bool bLocalStencilDitheredLOD = IsStencilForLODDitherEnabled(GetShaderPlatform());
 	return !IsMasked()
 		// Render dithered material as masked if a stencil prepass is not used (UE-50064, UE-49537)
 		&& !((bShadowPass || !bLocalStencilDitheredLOD) && IsDitheredLODTransition())
diff --git a/Engine/Source/Runtime/Engine/Public/MaterialShared.h b/Engine/Source/Runtime/Engine/Public/MaterialShared.h
--- a/Engine/Source/Runtime/Engine/Public/MaterialShared.h
+++ b/Engine/Source/Runtime/Engine/Public/MaterialShared.h
@@ -2064,13 +2064,6 @@ public:
 		bLoadedCookedShaderMapId(false),
 		bAllowPixelDepthOffset(true)
 	{
-		// this option affects only deferred renderer
-		static TConsoleVariableData<int32>* CVarStencilDitheredLOD;
-		if (CVarStencilDitheredLOD == nullptr)
-		{
-			CVarStencilDitheredLOD = IConsoleManager::Get().FindTConsoleVariableDataInt(TEXT("r.StencilForLODDither"));
-		}
-		bStencilDitheredLOD = (CVarStencilDitheredLOD->GetValueOnAnyThread() != 0);
 #if UE_CHECK_FMATERIAL_LIFETIME
 		bOwnerBeginDestroyed = false;
 #endif
@@ -2916,9 +2909,6 @@ private:
 	EShaderPlatform ShaderPlatform;
 #endif
 
-	/** Whether tthis project uses stencil dither lod. */
-	uint32 bStencilDitheredLOD : 1;
-
 	/** 
 	 * Whether this material was loaded with shaders inlined. 
 	 * If true, GameThreadShaderMap will contain a reference to the inlined shader map between Serialize and PostLoad.
diff --git a/Engine/Source/Runtime/RenderCore/Private/ReadOnlyCVARCache.cpp b/Engine/Source/Runtime/RenderCore/Private/ReadOnlyCVARCache.cpp
--- a/Engine/Source/Runtime/RenderCore/Private/ReadOnlyCVARCache.cpp
+++ b/Engine/Source/Runtime/RenderCore/Private/ReadOnlyCVARCache.cpp
@@ -29,6 +29,7 @@ bool FReadOnlyCVARCache::bMobileForwardParticleLights = false;
 bool FReadOnlyCVARCache::bMobileDeferredShadingValue = false;
 bool FReadOnlyCVARCache::bMobileEnableMovableSpotlightsShadowValue = false;
 int32 FReadOnlyCVARCache::MobileForwardDecalLightingValue = 1;
+bool FReadOnlyCVARCache::bMobileAllowDitheredLODTransitionValue = false;
 
 int32 FReadOnlyCVARCache::MobileEarlyZPassIniValue(EShaderPlatform Platform)
 {
@@ -72,6 +73,12 @@ bool FReadOnlyCVARCache::MobileEnableMovableSpotlightsShadowIniValue(EShaderPlat
 	return CVar.Get(Platform);
 }
 
+bool FReadOnlyCVARCache::MobileAllowDitheredLODTransitionValueIniValue(EShaderPlatform Platform)
+{
+	static FShaderPlatformCachedIniValue<bool> CVar(TEXT("r.Mobile.AllowDitheredLODTransition"));
+	return CVar.Get(Platform);
+}
+
 void FReadOnlyCVARCache::Initialize()
 {
 	UE_LOG(LogInit, Log, TEXT("Initializing FReadOnlyCVARCache"));
@@ -111,6 +118,7 @@ void FReadOnlyCVARCache::Initialize()
 	bMobileSupportsGPUScene = CVarMobileSupportGPUScene->GetValueOnAnyThread() != 0;
 	bMobileForwardParticleLights = MobileForwardParticleLightsIniValue(GMaxRHIShaderPlatform);
 	MobileForwardDecalLightingValue = CVarMobileForwardDecalLightingValue->GetValueOnAnyThread();
+	bMobileAllowDitheredLODTransitionValue = MobileAllowDitheredLODTransitionValueIniValue(GMaxRHIShaderPlatform);
 
 #ifdef PROJECT_CVAR_ALLOW_STATIC_LIGHTING
 	check(!!PROJECT_CVAR_ALLOW_STATIC_LIGHTING == bAllowStaticLighting);
diff --git a/Engine/Source/Runtime/RenderCore/Private/RenderUtils.cpp b/Engine/Source/Runtime/RenderCore/Private/RenderUtils.cpp
--- a/Engine/Source/Runtime/RenderCore/Private/RenderUtils.cpp
+++ b/Engine/Source/Runtime/RenderCore/Private/RenderUtils.cpp
@@ -630,7 +630,7 @@ RENDERCORE_API bool MobileBasePassAlwaysUsesCSM(const FStaticShaderPlatform Plat
 
 RENDERCORE_API bool MobileUsesFullDepthPrepass(const FStaticShaderPlatform Platform)
 {
-	return MobileUsesShadowMaskTexture(Platform) || IsMobileAmbientOcclusionEnabled(Platform) || IsUsingDBuffers(Platform) || FReadOnlyCVARCache::MobileEarlyZPass(Platform) == 1;
+	return IsStencilForLODDitherEnabled(Platform) || MobileUsesShadowMaskTexture(Platform) || IsMobileAmbientOcclusionEnabled(Platform) || IsUsingDBuffers(Platform) || FReadOnlyCVARCache::MobileEarlyZPass(Platform) == 1;
 }
 
 RENDERCORE_API bool MobileSupportsSM5MaterialNodes(const FStaticShaderPlatform Platform)
@@ -650,9 +650,7 @@ RENDERCORE_API bool ShouldForceFullDepthPass(const FStaticShaderPlatform Platfor
 		const bool bNaniteEnabled = UseNanite(Platform);
 		const bool bDBufferAllowed = IsUsingDBuffers(Platform);
 		const bool bVirtualTextureEnabled = UseVirtualTexturing(Platform);
-
-		static const auto StencilLODDitherCVar = IConsoleManager::Get().FindTConsoleVariableDataInt(TEXT("r.StencilForLODDither"));
-		const bool bStencilLODDither = StencilLODDitherCVar->GetValueOnAnyThread() != 0;
+		const bool bStencilLODDither = IsStencilForLODDitherEnabled(Platform);
 
 		static const auto AOComputeCVar = IConsoleManager::Get().FindTConsoleVariableDataInt(TEXT("r.AmbientOcclusion.Compute"));
 		const bool bAOCompute = AOComputeCVar->GetValueOnAnyThread() > 0;
@@ -1589,6 +1587,16 @@ bool IsMobileDistanceFieldAOEnabled(const FStaticShaderPlatform Platform)
 	return MobileDistanceFieldAO.Get(Platform) > 0;
 }
 
+bool IsStencilForLODDitherEnabled(const FStaticShaderPlatform Platform)
+{
+	if (IsMobilePlatform(Platform) && !FReadOnlyCVARCache::MobileAllowDitheredLODTransition(Platform))
+	{
+		return false;
+	}
+	static FShaderPlatformCachedIniValue<bool> StencilForLODDitherIniValue(TEXT("r.StencilForLODDither"));
+	return StencilForLODDitherIniValue.Get(Platform);
+}
+
 int32 GetMaterialShadingPathNodeOverride(const FStaticShaderPlatform Platform)
 {
 	static FShaderPlatformCachedIniValue<int32> MaterialShadingPathNodeOverride(TEXT("r.Material.ShadingPathNodeOverride"));
diff --git a/Engine/Source/Runtime/RenderCore/Private/Shader.cpp b/Engine/Source/Runtime/RenderCore/Private/Shader.cpp
--- a/Engine/Source/Runtime/RenderCore/Private/Shader.cpp
+++ b/Engine/Source/Runtime/RenderCore/Private/Shader.cpp
@@ -1974,8 +1974,7 @@ void ShaderMapAppendKey(EShaderPlatform Platform, FShaderKeyGenerator& KeyGen)
 		}
 
 		{
-			static FShaderPlatformCachedIniValue<int32> CVarAllowDitheredLODTransition(TEXT("r.Mobile.AllowDitheredLODTransition"));
-			if (CVarAllowDitheredLODTransition.Get(Platform) != 0)
+			if (FReadOnlyCVARCache::MobileAllowDitheredLODTransition(Platform))
 			{
 				KeyGen.AppendSeparator();
 				KeyGen.Append(TEXT("DLODT"));
@@ -2248,15 +2247,11 @@ void ShaderMapAppendKey(EShaderPlatform Platform, FShaderKeyGenerator& KeyGen)
 	// Is DXC shader compiler enabled for this platform?
 	KeyGen.AppendSeparator();
 	KeyGen.Append(IsDxcEnabledForPlatform(Platform) ? TEXT("DXC1") : TEXT("DXC0"));
-
-	if (IsFeatureLevelSupported(Platform, ERHIFeatureLevel::SM5))
+	
+	if (IsStencilForLODDitherEnabled(Platform))
 	{
-		static const auto CVar = IConsoleManager::Get().FindTConsoleVariableDataInt(TEXT("r.StencilForLODDither"));
-		if (CVar && CVar->GetValueOnAnyThread() > 0)
-		{
-			KeyGen.AppendSeparator();
-			KeyGen.Append(TEXT("SD"));
-		}
+		KeyGen.AppendSeparator();
+		KeyGen.Append(TEXT("SD"));
 	}
 
 	{
diff --git a/Engine/Source/Runtime/RenderCore/Public/ReadOnlyCVARCache.h b/Engine/Source/Runtime/RenderCore/Public/ReadOnlyCVARCache.h
--- a/Engine/Source/Runtime/RenderCore/Public/ReadOnlyCVARCache.h
+++ b/Engine/Source/Runtime/RenderCore/Public/ReadOnlyCVARCache.h
@@ -144,6 +144,15 @@ struct FReadOnlyCVARCache
 		#endif
 	}
 
+	static inline bool MobileAllowDitheredLODTransition(EShaderPlatform Platform)
+	{
+		#if WITH_EDITOR
+			return MobileAllowDitheredLODTransitionValueIniValue(Platform);
+		#else
+			return bMobileAllowDitheredLODTransitionValue;
+		#endif
+	}
+
 private:
 	RENDERCORE_API static bool bInitialized;
 
@@ -165,6 +174,7 @@ private:
 	RENDERCORE_API static int32 MobileForwardDecalLightingValue;
 	RENDERCORE_API static bool bMobileDeferredShadingValue;
 	RENDERCORE_API static bool bMobileEnableMovableSpotlightsShadowValue;
+	RENDERCORE_API static bool bMobileAllowDitheredLODTransitionValue;
 
 private:
 	RENDERCORE_API static int32 MobileEarlyZPassIniValue(EShaderPlatform Platform);
@@ -172,5 +182,6 @@ private:
 	RENDERCORE_API static bool MobileForwardParticleLightsIniValue(EShaderPlatform Platform);
 	RENDERCORE_API static bool MobileDeferredShadingIniValue(EShaderPlatform Platform);
 	RENDERCORE_API static bool MobileEnableMovableSpotlightsShadowIniValue(EShaderPlatform Platform);
+	RENDERCORE_API static bool MobileAllowDitheredLODTransitionValueIniValue(EShaderPlatform Platform);
 };
 
diff --git a/Engine/Source/Runtime/RenderCore/Public/RenderUtils.h b/Engine/Source/Runtime/RenderCore/Public/RenderUtils.h
--- a/Engine/Source/Runtime/RenderCore/Public/RenderUtils.h
+++ b/Engine/Source/Runtime/RenderCore/Public/RenderUtils.h
@@ -366,6 +366,9 @@ RENDERCORE_API bool IsWaterSeparateMainDirLightEnabled(const FStaticShaderPlatfo
 /** Returns if Mobile Distance Field AO Enabled. */
 RENDERCORE_API bool IsMobileDistanceFieldAOEnabled(const FStaticShaderPlatform Platform);
 
+/** Returns if Stencil For LOD Dither is Enabled for this ShaderPlatform. */
+RENDERCORE_API bool IsStencilForLODDitherEnabled(const FStaticShaderPlatform Platform);
+
 /** Returns shading path node material override. Returns -1 if no override. */
 RENDERCORE_API int32 GetMaterialShadingPathNodeOverride(const FStaticShaderPlatform Platform);
 
diff --git a/Engine/Source/Runtime/Renderer/Private/DeferredShadingRenderer.cpp b/Engine/Source/Runtime/Renderer/Private/DeferredShadingRenderer.cpp
--- a/Engine/Source/Runtime/Renderer/Private/DeferredShadingRenderer.cpp
+++ b/Engine/Source/Runtime/Renderer/Private/DeferredShadingRenderer.cpp
@@ -435,7 +435,6 @@ const TCHAR* GetDepthPassReason(bool bDitheredLODTransitionsUseStencil, EShaderP
 
 FDeferredShadingSceneRenderer::FDeferredShadingSceneRenderer(const FSceneViewFamily* InViewFamily, FHitProxyConsumer* HitProxyConsumer)
 	: FSceneRenderer(InViewFamily, HitProxyConsumer)
-	, DepthPass(GetDepthPassInfo(Scene))
 	, bAreLightsInLightGrid(false)
 {
 	ViewPipelineStates.SetNum(AllViews.Num());
diff --git a/Engine/Source/Runtime/Renderer/Private/DeferredShadingRenderer.h b/Engine/Source/Runtime/Renderer/Private/DeferredShadingRenderer.h
--- a/Engine/Source/Runtime/Renderer/Private/DeferredShadingRenderer.h
+++ b/Engine/Source/Runtime/Renderer/Private/DeferredShadingRenderer.h
@@ -315,8 +315,6 @@ enum class EReflectionsMethod
 class FDeferredShadingSceneRenderer : public FSceneRenderer
 {
 public:
-	/** Defines which objects we want to render in the EarlyZPass. */
-	FDepthPassInfo DepthPass;
 
 	FLumenCardRenderer LumenCardRenderer;
 
diff --git a/Engine/Source/Runtime/Renderer/Private/DepthRendering.cpp b/Engine/Source/Runtime/Renderer/Private/DepthRendering.cpp
--- a/Engine/Source/Runtime/Renderer/Private/DepthRendering.cpp
+++ b/Engine/Source/Runtime/Renderer/Private/DepthRendering.cpp
@@ -94,7 +94,7 @@ FDepthPassInfo GetDepthPassInfo(const FScene* Scene)
 	FDepthPassInfo Info;
 	Info.EarlyZPassMode = Scene ? Scene->EarlyZPassMode : DDM_None;
 	Info.bEarlyZPassMovable = Scene ? Scene->bEarlyZPassMovable : false;
-	Info.bDitheredLODTransitionsUseStencil = CVarStencilForLODDither.GetValueOnAnyThread() > 0;
+	Info.bDitheredLODTransitionsUseStencil = IsStencilForLODDitherEnabled(Scene ? Scene->GetShaderPlatform() : GMaxRHIShaderPlatform);
 	Info.StencilDitherPassFlags = ERDGPassFlags::Raster;
 
 	if (GRHISupportsDepthUAV && !IsHMDHiddenAreaMaskActive())
@@ -231,7 +231,10 @@ FDepthStencilStateRHIRef GetDitheredLODTransitionDepthStencilState()
 		>::GetRHI();
 }
 
-void SetDepthPassDitheredLODTransitionState(const FSceneView* SceneView, const FMeshBatch& RESTRICT Mesh, int32 StaticMeshId, FMeshPassProcessorRenderState& DrawRenderState)
+TOptional<uint32> GetOptionalDitheringFlag(
+	const FSceneView* SceneView,
+	const FMeshBatch& RESTRICT Mesh,
+	int32 StaticMeshId)
 {
 	if (SceneView && StaticMeshId >= 0 && Mesh.bDitheredLODTransition)
 	{
@@ -242,15 +245,16 @@ void SetDepthPassDitheredLODTransitionState(const FSceneView* SceneView, const F
 		{
 			if (ViewInfo->StaticMeshFadeOutDitheredLODMap[StaticMeshId])
 			{
-				DrawRenderState.SetDepthStencilState(GetDitheredLODTransitionDepthStencilState());
-				DrawRenderState.SetStencilRef(STENCIL_SANDBOX_MASK);
+				return STENCIL_SANDBOX_MASK;
 			}
 			else if (ViewInfo->StaticMeshFadeInDitheredLODMap[StaticMeshId])
 			{
-				DrawRenderState.SetDepthStencilState(GetDitheredLODTransitionDepthStencilState());
+				return 0;
 			}
 		}
 	}
+
+	return TOptional<uint32>();
 }
 
 /** A pixel shader used to fill the stencil buffer with the current dithered transition mask. */
@@ -266,10 +270,6 @@ public:
 		RENDER_TARGET_BINDING_SLOTS()
 	END_SHADER_PARAMETER_STRUCT()
 
-	static bool ShouldCompilePermutation(const FGlobalShaderPermutationParameters& Parameters)
-	{ 
-		return IsFeatureLevelSupported(Parameters.Platform, ERHIFeatureLevel::SM5);
-	}
 };
 
 IMPLEMENT_GLOBAL_SHADER(FDitheredTransitionStencilPS, "/Engine/Private/DitheredTransitionStencil.usf", "Main", SF_Pixel);
@@ -322,7 +322,9 @@ void AddDitheredStencilFillPass(FRDGBuilder& GraphBuilder, TConstArrayView<FView
 			auto* PassParameters = GraphBuilder.AllocParameters<FDitheredTransitionStencilPS::FParameters>();
 			PassParameters->View = View.ViewUniformBuffer;
 			PassParameters->DitheredTransitionFactor = View.GetTemporalLODTransition();
-			PassParameters->RenderTargets.DepthStencil = FDepthStencilBinding(DepthTexture, ERenderTargetLoadAction::ELoad, ERenderTargetLoadAction::ELoad, FExclusiveDepthStencil::DepthWrite_StencilWrite);
+			PassParameters->RenderTargets.DepthStencil = IsMobilePlatform(View.GetShaderPlatform()) ?
+				FDepthStencilBinding(DepthTexture, ERenderTargetLoadAction::EClear, ERenderTargetLoadAction::EClear, FExclusiveDepthStencil::DepthWrite_StencilWrite):
+				FDepthStencilBinding(DepthTexture, ERenderTargetLoadAction::ELoad, ERenderTargetLoadAction::ELoad, FExclusiveDepthStencil::DepthWrite_StencilWrite);
 
 			FPixelShaderUtils::AddFullscreenPass(GraphBuilder, View.ShaderMap, {}, PixelShader, PassParameters, View.ViewRect, nullptr, nullptr, DepthStencilState, StencilRef);
 		}
@@ -362,6 +364,49 @@ void AddDitheredStencilFillPass(FRDGBuilder& GraphBuilder, TConstArrayView<FView
 	}
 }
 
+FDepthStencilStateRHIRef GetDitheredStencilClearPassDepthStencilState(bool bClearAllStencilBits)
+{
+	if (bClearAllStencilBits)
+	{
+		return TStaticDepthStencilState<false, CF_Always,
+			true, CF_Always, SO_Replace, SO_Replace, SO_Replace,
+			false, CF_Always, SO_Replace, SO_Replace, SO_Replace,
+			0xff, 0xff
+		>::GetRHI();
+	}
+	else
+	{
+		return TStaticDepthStencilState<false, CF_Always,
+			true, CF_Always, SO_Replace, SO_Replace, SO_Replace,
+			false, CF_Always, SO_Replace, SO_Replace, SO_Replace,
+			STENCIL_SANDBOX_MASK, STENCIL_SANDBOX_MASK
+		>::GetRHI();
+	}
+}
+
+void AddDitheredStencilClearPass(FRDGBuilder& GraphBuilder, FRDGTextureRef DepthTexture, EShaderPlatform ShaderPlatform, bool bClearAllStencilBits)
+{
+	auto* PassParameters = GraphBuilder.AllocParameters<TOneColorPixelShaderMRT::FParameters>();
+	PassParameters->RenderTargets.DepthStencil = FDepthStencilBinding(DepthTexture, ERenderTargetLoadAction::ELoad, ERenderTargetLoadAction::ELoad, FExclusiveDepthStencil::DepthWrite_StencilWrite);
+	FIntPoint DepthSize = DepthTexture->Desc.Extent;
+	const auto ShaderMap = GetGlobalShaderMap(ShaderPlatform);
+	TOneColorPixelShaderMRT::FPermutationDomain PermutationVector;
+	TShaderMapRef<TOneColorPixelShaderMRT> PixelShader(ShaderMap, PermutationVector);
+	FRHIDepthStencilState* DepthStencilState = GetDitheredStencilClearPassDepthStencilState(bClearAllStencilBits);
+
+	FPixelShaderUtils::AddFullscreenPass(
+		GraphBuilder, 
+		ShaderMap,
+		RDG_EVENT_NAME("AddDitheredStencilClearPass"),
+		PixelShader, 
+		PassParameters, 
+		FIntRect(0, 0, DepthSize.X, DepthSize.Y), 
+		nullptr, 
+		nullptr, 
+		DepthStencilState, 
+		0);
+}
+
 // GPUCULL_TODO: Move to Utils file and make templated on params and mesh pass processor
 static void AddViewMeshElementsPass(const TIndirectArray<FMeshBatch>& MeshElements, FRDGBuilder& GraphBuilder, FDepthPassParameters* PassParameters, const FScene* Scene, const FViewInfo& View, const FMeshPassProcessorRenderState& DrawRenderState, bool bRespectUseAsOccluderFlag, EDepthDrawingMode DepthDrawingMode, FInstanceCullingManager& InstanceCullingManager)
 {
@@ -620,26 +665,7 @@ void FDeferredShadingSceneRenderer::RenderPrePass(FRDGBuilder& GraphBuilder, TAr
 	// Dithered transition stencil mask clear, accounting for all active viewports
 	if (DepthPass.bDitheredLODTransitionsUseStencil)
 	{
-		auto* PassParameters = GraphBuilder.AllocParameters<FRenderTargetParameters>();
-		PassParameters->RenderTargets.DepthStencil = FDepthStencilBinding(SceneDepthTexture, ERenderTargetLoadAction::ELoad, ERenderTargetLoadAction::ELoad, FExclusiveDepthStencil::DepthWrite_StencilWrite);
-
-		GraphBuilder.AddPass(
-			RDG_EVENT_NAME("DitherStencilClear"),
-			PassParameters,
-			ERDGPassFlags::Raster,
-			[this, InViews](FRDGAsyncTask, FRHICommandList& RHICmdList)
-		{
-			if (InViews.Num() > 1)
-			{
-				FIntRect FullViewRect = InViews[0].ViewRect;
-				for (int32 ViewIndex = 1; ViewIndex < InViews.Num(); ++ViewIndex)
-				{
-					FullViewRect.Union(InViews[ViewIndex].ViewRect);
-				}
-				RHICmdList.SetViewport(FullViewRect.Min.X, FullViewRect.Min.Y, 0, FullViewRect.Max.X, FullViewRect.Max.Y, 1);
-			}
-			DrawClearQuad(RHICmdList, false, FLinearColor::Transparent, false, 0, true, 0);
-		});
+		AddDitheredStencilClearPass(GraphBuilder, SceneDepthTexture, ShaderPlatform, true);
 	}
 
 #if !(UE_BUILD_SHIPPING)
@@ -779,21 +805,36 @@ bool FDepthPassMeshProcessor::Process(
 	}
 
 	FMeshPassProcessorRenderState DrawRenderState(PassDrawRenderState);
-
 	if (!bDitheredLODFadingOutMaskPass && !bShadowProjection)
 	{
-		SetDepthPassDitheredLODTransitionState(ViewIfDynamicMeshCommand, MeshBatch, StaticMeshId, DrawRenderState);
-	}
-
-	// Use StencilMask for DecalOutput on mobile
-	if (FeatureLevel == ERHIFeatureLevel::ES3_1 && !bShadowProjection)
-	{
-		extern void SetMobileBasePassDepthState(FMeshPassProcessorRenderState& DrawRenderState, const FPrimitiveSceneProxy* PrimitiveSceneProxy, const FMaterial& Material, FMaterialShadingModelField ShadingModels, bool bUsesDeferredShading);
-		
-		// *Don't* get shading models from MaterialResource since it's for a default material
-		FMaterialShadingModelField ShadingModels = MeshBatch.MaterialRenderProxy->GetIncompleteMaterialWithFallback(ERHIFeatureLevel::ES3_1).GetShadingModels();
-		bool bUsesDeferredShading = IsMobileDeferredShadingEnabled(GetFeatureLevelShaderPlatform(FeatureLevel));
-		SetMobileBasePassDepthState(DrawRenderState, PrimitiveSceneProxy, MaterialResource, ShadingModels, bUsesDeferredShading);
+		TOptional<uint32> DitheringFlag = GetOptionalDitheringFlag(ViewIfDynamicMeshCommand, MeshBatch, StaticMeshId);
+		// Use StencilMask for DecalOutput on mobile
+		if (FeatureLevel == ERHIFeatureLevel::ES3_1)
+		{
+			extern void SetMobileBasePassDepthState(
+				FMeshPassProcessorRenderState & DrawRenderState,
+				const FPrimitiveSceneProxy * PrimitiveSceneProxy,
+				const FMaterial & Material,
+				FMaterialShadingModelField ShadingModels,
+				bool bUsesDeferredShading,
+				TOptional<uint32> DitheringFlag);
+
+			// *Don't* get shading models from MaterialResource since it's for a default material
+			FMaterialShadingModelField ShadingModels = MeshBatch.MaterialRenderProxy->GetIncompleteMaterialWithFallback(ERHIFeatureLevel::ES3_1).GetShadingModels();
+			bool bUsesDeferredShading = IsMobileDeferredShadingEnabled(GetFeatureLevelShaderPlatform(FeatureLevel));
+			SetMobileBasePassDepthState(
+				DrawRenderState,
+				PrimitiveSceneProxy,
+				MaterialResource,
+				ShadingModels,
+				bUsesDeferredShading,
+				DitheringFlag);
+		}
+		else if (DitheringFlag.IsSet())
+		{
+			DrawRenderState.SetDepthStencilState(GetDitheredLODTransitionDepthStencilState());
+			DrawRenderState.SetStencilRef(DitheringFlag.GetValue());
+		}
 	}
 
 	FMeshMaterialShaderElementData ShaderElementData;
diff --git a/Engine/Source/Runtime/Renderer/Private/DepthRendering.h b/Engine/Source/Runtime/Renderer/Private/DepthRendering.h
--- a/Engine/Source/Runtime/Renderer/Private/DepthRendering.h
+++ b/Engine/Source/Runtime/Renderer/Private/DepthRendering.h
@@ -56,6 +56,7 @@ struct FDepthPassInfo
 extern FDepthPassInfo GetDepthPassInfo(const FScene* Scene);
 
 void AddDitheredStencilFillPass(FRDGBuilder& GraphBuilder, TConstArrayView<FViewInfo> Views, FRDGTextureRef DepthTexture, const FDepthPassInfo& DepthPass);
+void AddDitheredStencilClearPass(FRDGBuilder& GraphBuilder, FRDGTextureRef DepthTexture, EShaderPlatform ShaderPlatform, bool bClearAllStencilBits);
 
 FMeshDrawCommandSortKey CalculateDepthPassMeshStaticSortKey(const bool bIsMasked, const FMeshMaterialShader* VertexShader, const FMeshMaterialShader* PixelShader);
 
diff --git a/Engine/Source/Runtime/Renderer/Private/MobileBasePass.cpp b/Engine/Source/Runtime/Renderer/Private/MobileBasePass.cpp
--- a/Engine/Source/Runtime/Renderer/Private/MobileBasePass.cpp
+++ b/Engine/Source/Runtime/Renderer/Private/MobileBasePass.cpp
@@ -85,16 +85,35 @@ uint8 GetMobileShadingModelStencilValue(FMaterialShadingModelField ShadingModel,
 	return MobileUsesGBufferCustomData(GMaxRHIShaderPlatform) ? MobileShadingModelSupportStencilValue : DefaultLitMask;
 }
 
-void SetMobileBasePassDepthState(FMeshPassProcessorRenderState& DrawRenderState, const FPrimitiveSceneProxy* PrimitiveSceneProxy, const FMaterial& Material, FMaterialShadingModelField ShadingModels, bool bUsesDeferredShading)
+void SetMobileBasePassDepthState(
+	FMeshPassProcessorRenderState& DrawRenderState, 
+	const FPrimitiveSceneProxy* PrimitiveSceneProxy, 
+	const FMaterial& Material, 
+	FMaterialShadingModelField ShadingModels, 
+	bool bUsesDeferredShading,
+	TOptional<uint32> DitheringFlag)
 {
-	DrawRenderState.SetDepthStencilState(TStaticDepthStencilState<
-		true, CF_DepthNearOrEqual,
-		true, CF_Always, SO_Keep, SO_Keep, SO_Replace,
-		false, CF_Always, SO_Keep, SO_Keep, SO_Keep,
-		// don't use masking as it has significant performance hit on Mali GPUs (T860MP2)
-		0x00, 0xff >::GetRHI());
-
 	uint8 StencilValue = 0u;
+	if (DitheringFlag.IsSet())
+	{
+		StencilValue = DitheringFlag.GetValue();
+
+		DrawRenderState.SetDepthStencilState(TStaticDepthStencilState<
+			true, CF_DepthNearOrEqual,
+			true, CF_Equal, SO_Keep, SO_Keep, SO_Replace,
+			false, CF_Always, SO_Keep, SO_Keep, SO_Keep,
+			STENCIL_SANDBOX_MASK, 0xFF
+		>::GetRHI());
+	}
+	else
+	{
+		DrawRenderState.SetDepthStencilState(TStaticDepthStencilState<
+			true, CF_DepthNearOrEqual,
+			true, CF_Always, SO_Keep, SO_Keep, SO_Replace,
+			false, CF_Always, SO_Keep, SO_Keep, SO_Keep,
+			// don't use masking as it has significant performance hit on Mali GPUs (T860MP2)
+			0x00, 0xff >::GetRHI());
+	}
 
 	uint8 ReceiveDecals = (PrimitiveSceneProxy && !PrimitiveSceneProxy->ReceivesDecals() ? 0x01 : 0x00);
 	StencilValue |= GET_STENCIL_BIT_MASK(RECEIVE_DECAL, ReceiveDecals);
@@ -462,7 +481,7 @@ void MobileBasePass::SetOpaqueRenderState(FMeshPassProcessorRenderState& DrawRen
 {
 	if (bCanUseDepthStencil)
 	{
-		SetMobileBasePassDepthState(DrawRenderState, PrimitiveSceneProxy, Material, ShadingModels, bUsesDeferredShading);
+		SetMobileBasePassDepthState(DrawRenderState, PrimitiveSceneProxy, Material, ShadingModels, bUsesDeferredShading, TOptional<uint32>());
 	}
 	else
 	{
diff --git a/Engine/Source/Runtime/Renderer/Private/MobileShadingRenderer.cpp b/Engine/Source/Runtime/Renderer/Private/MobileShadingRenderer.cpp
--- a/Engine/Source/Runtime/Renderer/Private/MobileShadingRenderer.cpp
+++ b/Engine/Source/Runtime/Renderer/Private/MobileShadingRenderer.cpp
@@ -818,6 +818,10 @@ void FMobileSceneRenderer::RenderFullDepthPrepass(FRDGBuilder& GraphBuilder, TAr
 	GetRenderViews(InViews, RenderViews);
 
 	TRDGUniformBufferBinding<FMobileBasePassUniformParameters> LastViewMobileBasePassUB;
+	if (DepthPass.IsRasterStencilDitherEnabled())
+	{
+		AddDitheredStencilFillPass(GraphBuilder, InViews, SceneTextures.Depth.Target, DepthPass);
+	}
 
 	for (FRenderViewContext& ViewContext : RenderViews)
 	{
@@ -839,7 +843,9 @@ void FMobileSceneRenderer::RenderFullDepthPrepass(FRDGBuilder& GraphBuilder, TAr
 		{
 			PassParameters->RenderTargets[0] = FRenderTargetBinding(SceneTextures.Velocity, ERenderTargetLoadAction::EClear);
 		}
-		PassParameters->RenderTargets.DepthStencil = FDepthStencilBinding(SceneTextures.Depth.Target, ERenderTargetLoadAction::EClear, ERenderTargetLoadAction::EClear, FExclusiveDepthStencil::DepthWrite_StencilWrite);
+		PassParameters->RenderTargets.DepthStencil = DepthPass.IsRasterStencilDitherEnabled() ?
+			FDepthStencilBinding(SceneTextures.Depth.Target, ERenderTargetLoadAction::ELoad, ERenderTargetLoadAction::ELoad, FExclusiveDepthStencil::DepthWrite_StencilWrite):
+			FDepthStencilBinding(SceneTextures.Depth.Target, ERenderTargetLoadAction::EClear, ERenderTargetLoadAction::EClear, FExclusiveDepthStencil::DepthWrite_StencilWrite);
 		PassParameters->View = View.GetShaderParameters();
 		PassParameters->MobileBasePass = CreateMobileBasePassUniformBuffer(GraphBuilder, View, EMobileBasePass::DepthPrePass, EMobileSceneTextureSetupMode::None);
 		//if the scenecolor isn't multiview but the app is, need to render as a single-view multiview due to shaders
@@ -882,6 +888,12 @@ void FMobileSceneRenderer::RenderFullDepthPrepass(FRDGBuilder& GraphBuilder, TAr
 			});
 	}
 
+	// Dithered transition stencil mask clear, accounting for all active viewports
+	if (DepthPass.bDitheredLODTransitionsUseStencil)
+	{
+		AddDitheredStencilClearPass(GraphBuilder, SceneTextures.Depth.Target, ShaderPlatform, false);
+	}
+
 	for (FRenderViewContext& ViewContext : RenderViews)
 	{
 		FViewInfo& View = *ViewContext.ViewInfo;
diff --git a/Engine/Source/Runtime/Renderer/Private/SceneRendering.cpp b/Engine/Source/Runtime/Renderer/Private/SceneRendering.cpp
--- a/Engine/Source/Runtime/Renderer/Private/SceneRendering.cpp
+++ b/Engine/Source/Runtime/Renderer/Private/SceneRendering.cpp
@@ -2640,6 +2640,7 @@ FSceneRendererBase* FSceneRendererBase::GetActiveInstance(FRDGBuilder& GraphBuil
 
 FSceneRenderer::FSceneRenderer(const FSceneViewFamily* InViewFamily, FHitProxyConsumer* HitProxyConsumer)
 :	FSceneRendererBase(*CheckPointer(CheckPointer(CheckPointer(InViewFamily)->Scene)->GetRenderScene()))
+,	DepthPass(GetDepthPassInfo(Scene))
 ,	ViewFamily(*InViewFamily)
 ,	VirtualShadowMapArray(*Scene)
 ,	bHasRequestedToggleFreeze(false)
diff --git a/Engine/Source/Runtime/Renderer/Private/SceneRendering.h b/Engine/Source/Runtime/Renderer/Private/SceneRendering.h
--- a/Engine/Source/Runtime/Renderer/Private/SceneRendering.h
+++ b/Engine/Source/Runtime/Renderer/Private/SceneRendering.h
@@ -53,6 +53,7 @@
 #include "LightFunctionAtlas.h"
 #include "SceneExtensions.h"
 #include "PostProcess/LensDistortion.h"
+#include "DepthRendering.h"
 
 #if RHI_RAYTRACING
 #include "RayTracingInstanceBufferUtil.h"
@@ -2059,6 +2060,9 @@ private:
 class FSceneRenderer : public FSceneRendererBase
 {
 public:
+	/** Defines which objects we want to render in the EarlyZPass. */
+	FDepthPassInfo DepthPass;
+
 	/** Linear bulk allocator with a lifetime tied to the scene renderer. */
 	FSceneRenderingBulkObjectAllocator Allocator;
 
diff --git a/Engine/Source/Runtime/Renderer/Private/VelocityRendering.cpp b/Engine/Source/Runtime/Renderer/Private/VelocityRendering.cpp
--- a/Engine/Source/Runtime/Renderer/Private/VelocityRendering.cpp
+++ b/Engine/Source/Runtime/Renderer/Private/VelocityRendering.cpp
@@ -806,12 +806,29 @@ bool FVelocityMeshProcessor::Process(
 	// When velocity is used as a depth pass we need to set a correct stencil state on mobile
 	if (FeatureLevel == ERHIFeatureLevel::ES3_1 && EarlyZPassMode == DDM_AllOpaqueNoVelocity)
 	{
-		extern void SetMobileBasePassDepthState(FMeshPassProcessorRenderState& DrawRenderState, const FPrimitiveSceneProxy* PrimitiveSceneProxy, const FMaterial& Material, FMaterialShadingModelField ShadingModels, bool bUsesDeferredShading);
+		extern TOptional<uint32> GetOptionalDitheringFlag(
+			const FSceneView * SceneView,
+			const FMeshBatch & RESTRICT Mesh,
+			int32 StaticMeshId);
+
+		extern void SetMobileBasePassDepthState(
+			FMeshPassProcessorRenderState& DrawRenderState, 
+			const FPrimitiveSceneProxy* PrimitiveSceneProxy,
+			const FMaterial& Material, 
+			FMaterialShadingModelField ShadingModels, 
+			bool bUsesDeferredShading,
+			TOptional<uint32> DitheringFlag);
 
 		// *Don't* get shading models from MaterialResource since it's for a default material
 		FMaterialShadingModelField ShadingModels = MeshBatch.MaterialRenderProxy->GetIncompleteMaterialWithFallback(ERHIFeatureLevel::ES3_1).GetShadingModels();
 		bool bUsesDeferredShading = IsMobileDeferredShadingEnabled(GetFeatureLevelShaderPlatform(FeatureLevel));
-		SetMobileBasePassDepthState(PassDrawRenderState, PrimitiveSceneProxy, MaterialResource, ShadingModels, bUsesDeferredShading);
+		SetMobileBasePassDepthState(
+			PassDrawRenderState, 
+			PrimitiveSceneProxy, 
+			MaterialResource, 
+			ShadingModels, 
+			bUsesDeferredShading, 
+			GetOptionalDitheringFlag(ViewIfDynamicMeshCommand, MeshBatch, StaticMeshId));
 	}
 
 	FMeshMaterialShaderElementData ShaderElementData;
