Path Tracer: Run compilation of path tracing RTPSOs in the background

Avoid blocking on startup while the path tracing precompilation is still in progress.

#rnx

diff --git a/Engine/Source/Runtime/Renderer/Private/PathTracing.cpp b/Engine/Source/Runtime/Renderer/Private/PathTracing.cpp
--- a/Engine/Source/Runtime/Renderer/Private/PathTracing.cpp
+++ b/Engine/Source/Runtime/Renderer/Private/PathTracing.cpp
@@ -3362,6 +3362,7 @@ void PreparePathTracingRTPSO()
 				};
 				FRayTracingPipelineStateInitializer Initializer;
 				Initializer.bPartial = true; // TODO: got a crash in some older nvidia drivers when false - need to find out which driver version has the fix
+				Initializer.bBackgroundCompilation = true;
 				Initializer.SetRayGenShaderTable(RayGenShaderTable);
 				Initializer.SetMissShaderTable(MissShaderTable);
 				Initializer.SetHitGroupTable(HitGroupTable);
@@ -3373,7 +3374,7 @@ void PreparePathTracingRTPSO()
 					Initializer.ShaderBindingLayout = &ShaderBindingLayout->RHILayout;
 				}
 
-				FRayTracingPipelineState* PipelineState = PipelineStateCache::GetAndOrCreateRayTracingPipelineState(RHICmdList, Initializer);
+				FRayTracingPipelineState* PipelineState = PipelineStateCache::GetAndOrCreateRayTracingPipelineState(RHICmdList, Initializer, ERayTracingPipelineCacheFlags::NonBlocking);
 				NumValidPermutations++;
 			}
 			UE_LOG(LogRenderer, Log, TEXT("Requested compilation of Path Tracing RTPSOs (%d permutations)."), NumValidPermutations);
