Stochastic Lighting - fixed check when using ray tracing overlay (RT debug view modes) with Substrate

diff --git a/Engine/Source/Runtime/Renderer/Private/DeferredShadingRenderer.cpp b/Engine/Source/Runtime/Renderer/Private/DeferredShadingRenderer.cpp
--- a/Engine/Source/Runtime/Renderer/Private/DeferredShadingRenderer.cpp
+++ b/Engine/Source/Runtime/Renderer/Private/DeferredShadingRenderer.cpp
@@ -2897,7 +2897,7 @@ void FDeferredShadingSceneRenderer::Render(FRDGBuilder& GraphBuilder, const FSce
 
 		FAsyncLumenIndirectLightingOutputs AsyncLumenIndirectLightingOutputs;
 
-		if (RequiresStochasticLightingPass())
+		if (!bHasRayTracedOverlay && RequiresStochasticLightingPass())
 		{
 			// Decals may modify GBuffers so they need to be done first. Can decals read velocities and/or custom depth? If so, they need to be rendered earlier too.
 			CompositionLighting.ProcessAfterBasePass(GraphBuilder, InstanceCullingManager, FCompositionLighting::EProcessAfterBasePassMode::OnlyBeforeLightingDecals, Scene->SubstrateSceneData);
