[TSR] Reduce ghosting on thin geometry region when r.TSR.ThinGeometryDetection is enabled.

1. has pixel animation flag now reduces history relaxation weight for TSR.
2. Add `Translucent pixel animation` support for BeforeDOF translucency covered regions when thin geometry detection and `OutputTranslucentVelocity` is enabled. The history weight relaxation will be turned off for regions where the translucency opacity is below the opacity mask clip value. This will double the cost of translucency velocity pass at most as we will run an additional velocity clipped depth pass to mark the last bit of velocity. This feature requires atomics on uint64 and is supported on nanite compatible devices only.
3. Use `r.Velocity.OutputTranslucentClippedDepth.Supported` in ini to support it (default off), and `r.Velocity.OutputTranslucentClippedDepth.Enabled` to turn it on/off in real time.

#tests Test on local scene with beforedof translucency

diff --git a/Engine/Shaders/Private/TemporalSuperResolution/TSRDetectThinGeometry.usf b/Engine/Shaders/Private/TemporalSuperResolution/TSRDetectThinGeometry.usf
--- a/Engine/Shaders/Private/TemporalSuperResolution/TSRDetectThinGeometry.usf
+++ b/Engine/Shaders/Private/TemporalSuperResolution/TSRDetectThinGeometry.usf
@@ -39,6 +39,7 @@ float MaxRelaxationWeight;
 Texture2D<float> SceneDepthTexture;
 Texture2D<tsr_half> CurrentCoverageTexture;
 Texture2D<tsr_half> ReprojectedHistoryCoverageTexture;
+Texture2D<tsr_half2> DecimateMaskTexture;
 
 RWTexture2DArray<tsr_half> HistoryCoverageOutput;
 RWTexture2DArray<uint> R8Output;
@@ -97,6 +98,7 @@ void FetchInput(
 	,out tsr_tensor_half DeviceZThinGeometryDetectionError
 	,out tsr_tensor_half CurrentCoverage
 	,out tsr_tensor_half ReprojectedHistoryCoverage
+	,out tsr_tensor_bool bHasPixelAnimation
 	)
 {
 	tsr_tensor_float DepthCache;
@@ -118,6 +120,8 @@ void FetchInput(
 #if CONFIG_THIN_GEOMETRY_EDGE_REPROJECTION
 		ReprojectionOffset.SetElement(ElementIndex, DecodeReprojectionOffset(R8Output[int3(ClampedFetchPixelPos, ThinGeometryTextureIndex)]));
 #endif
+		tsr_ushort BitMask = tsr_ushort(DecimateMaskTexture[ClampedFetchPixelPos].r * tsr_half(255.0));
+		bHasPixelAnimation.SetElement(ElementIndex,(BitMask & tsr_ushort(0x4)) != tsr_ushort(0));
 	}
 
 #if CONFIG_THIN_GEOMETRY_EDGE_REPROJECTION
@@ -286,11 +290,13 @@ void MainCS(
 	tsr_tensor_half DeviceZThinGeometryDetectionError;
 	tsr_tensor_half CurrentCoverage;
 	tsr_tensor_half ReprojectedHistoryCoverage;
+	tsr_tensor_bool bHasPixelAnimation;
 	FetchInput(PixelPos,
 		DeviceZ, 
 		DeviceZThinGeometryDetectionError
 		,CurrentCoverage
 		,ReprojectedHistoryCoverage
+		,bHasPixelAnimation
 	);
 	
 	//Debug[0].SetComponent(0, DeviceZ);
@@ -344,11 +350,11 @@ void MainCS(
 
 		// Sample weight for coverage update.
 		{
-			// Reject history coverage if history is not valid or absolute match.
+			// Reject history coverage if history is not valid or absolute match or has pixel animation.
 			// TODO: Absolute coverage match might has similar color, then no need to reset history coverage.
 			// e.g., 5x5 flat surface disappear and 5x5 flat surface appear. This is uncommon in thin geometry
 			// yet need to address to be less aggressive.
-			tsr_tensor_bool AcceptHistory = and(ValidCurrentThinGeometryRegion, and(bSameDistribution, NotAbsoluteMatch));
+			tsr_tensor_bool AcceptHistory = and(ValidCurrentThinGeometryRegion, and(bSameDistribution, and(NotAbsoluteMatch, not(bHasPixelAnimation))));
 			NewSampleWeight = select(AcceptHistory, NewSampleWeight, tsr_tensor_half::Const(1.0));
 
 			// Update reprojected history and use it for weight calculation.
@@ -376,10 +382,9 @@ void MainCS(
 			Deconcatenate(WeightedAvg3x3(Concatenate(ThinGeometryEdge,HistoryRelaxationWeight),float3(0.8, 1.0, 0.8)),ThinGeometryEdge,HistoryRelaxationWeight);
 			HistoryRelaxationWeight = saturate(WeightedAvg3x3(saturate(HistoryRelaxationWeight * tsr_half(MaxRelaxationWeight)), float3(0.5, 1.0, 0.5)));
 
-			// Not similar patches should reset relaxation weight to 0
+			// Not similar patches or has pixel animation should reset relaxation weight to 0
 			// to use more strict history clamping.
-			HistoryRelaxationWeight = select(bSameDistribution, HistoryRelaxationWeight, tsr_tensor_half::Const(0));
-			
+			HistoryRelaxationWeight = select(and(bSameDistribution, not(bHasPixelAnimation)), HistoryRelaxationWeight, tsr_tensor_half::Const(0));
 			//Debug[0].SetComponent(2,HistoryRelaxationWeight);
 		}
 
diff --git a/Engine/Shaders/Private/VelocityShader.usf b/Engine/Shaders/Private/VelocityShader.usf
--- a/Engine/Shaders/Private/VelocityShader.usf
+++ b/Engine/Shaders/Private/VelocityShader.usf
@@ -9,6 +9,9 @@
 #include "/Engine/Generated/VertexFactory.ush"
 #include "VelocityCommon.ush"
 
+#ifndef VELOCITY_CLIPPED_DEPTH_PASS
+#define VELOCITY_CLIPPED_DEPTH_PASS 0
+#endif
 // Move all geometry a little bit towards the camera to not get z fighting with existing depth
 // buffer of the same meshes rendered with slightly different float computations.
 // This is a positive number as our z buffer is inverted for better precision.
@@ -170,8 +173,37 @@ void MainPixelShader(
 	const float MaterialOpacity = GetMaterialOpacity(PixelMaterialInputs);
 #endif
 
+#if VELOCITY_CLIPPED_DEPTH_PASS && (!TRANSLUCENCY_VELOCITY_FROM_DEPTH)
+	// Set a bit for those BeforeDoF translucent pixels without writing velocity
+	#if MATERIALBLENDING_TRANSLUCENT || MATERIALBLENDING_ADDITIVE || MATERIALBLENDING_MODULATE || SUBSTRATE_TRANSLUCENT_MATERIAL
+        #if COMPILER_SUPPORTS_UINT64_IMAGE_ATOMICS
+        // We set the bit for those with Opacity value < Opacity Mask Clip Value.
+		ClipLODTransition(MaterialParameters);
+		float ActualMaterialOpacity = bThinTranslucentMaterial ? 1.0f : MaterialOpacity;
+		const bool bWriteFlag = (ActualMaterialOpacity > 0 && ActualMaterialOpacity < (1.0 / 255.0 + GetMaterialOpacityMaskClipValue()));
+		clip(select(bWriteFlag, 1.0f, -1.0f));
+
+		uint2 Position = uint2(SvPosition.xy - 0.5f) + View.ViewRectMin.xy;
+        
+        // Only consider if the current translucent pixel is in front.	
+		float OpaqueDepth = CalcSceneDepth(Position);
+
+		float CurrentTranslucencyDepth = 0;
+		ApplyPixelDepthOffsetToMaterialParameters(MaterialParameters, PixelMaterialInputs, CurrentTranslucencyDepth);
+		CurrentTranslucencyDepth = ConvertFromDeviceZ(CurrentTranslucencyDepth); 
+
+		uint WriteFlag = select(CurrentTranslucencyDepth < OpaqueDepth, 1u, 0);
+		const uint2 EncodedOutputBit = uint2(0u, WriteFlag << 16u); // Mask the lowest bit of alpha.
+		ImageInterlockedOrUInt64(VelocityClippedDepth.RWVelocity, Position, PackUlongType(EncodedOutputBit));
+		#endif
+	#endif
+	
+	clip(-1);
+#else 
 	GetMaterialClippingVelocity(MaterialParameters, PixelMaterialInputs, bThinTranslucentMaterial, MaterialOpacity);
-#endif
+#endif //VELOCITY_CLIPPED_DEPTH_PASS
+
+#endif //MATERIALBLENDING_MASKED_USING_COVERAGE
 
 #if !TRANSLUCENCY_VELOCITY_FROM_DEPTH
 	// 3d velocity, includes camera an object motion
@@ -181,5 +213,5 @@ void MainPixelShader(
 #else
 	// We only want velocity generated from camera movement and depth so write zero which is the special clear color value.
 	OutColor = 0;
-#endif
+#endif 
 }
diff --git a/Engine/Shaders/Public/Platform/D3D/D3DCommon.ush b/Engine/Shaders/Public/Platform/D3D/D3DCommon.ush
--- a/Engine/Shaders/Public/Platform/D3D/D3DCommon.ush
+++ b/Engine/Shaders/Public/Platform/D3D/D3DCommon.ush
@@ -140,6 +140,11 @@ void UEReportAssertWithPayload(uint AssertID, uint4 Payload)
 		{
 			InterlockedAdd(Dest[Coord], Value);
 		}
+
+		void ImageInterlockedOrUInt64(RWTexture2D<UlongType> Dest, uint2 Coord, UlongType Value)
+		{
+			InterlockedOr(Dest[Coord], Value);
+		}
 	#endif
 #endif
 
diff --git a/Engine/Shaders/Public/Platform/Metal/MetalCommon.ush b/Engine/Shaders/Public/Platform/Metal/MetalCommon.ush
--- a/Engine/Shaders/Public/Platform/Metal/MetalCommon.ush
+++ b/Engine/Shaders/Public/Platform/Metal/MetalCommon.ush
@@ -116,6 +116,11 @@
 			InterlockedAdd(Dest[Coord], Value);
 		}
 
+		void ImageInterlockedOrUInt64(RWTexture2D<UlongType> Dest, uint2 Coord, UlongType Value)
+		{
+			InterlockedOr(Dest[Coord], Value);
+		}
+
     #endif
 #endif
 
diff --git a/Engine/Shaders/Public/Platform/Vulkan/VulkanCommon.ush b/Engine/Shaders/Public/Platform/Vulkan/VulkanCommon.ush
--- a/Engine/Shaders/Public/Platform/Vulkan/VulkanCommon.ush
+++ b/Engine/Shaders/Public/Platform/Vulkan/VulkanCommon.ush
@@ -38,6 +38,11 @@ void ImageInterlockedAddUInt64(RWTexture2D<UlongType> Dest, uint2 Coord, UlongTy
 {
 	InterlockedAdd(Dest[Coord], Value);
 }
+
+void ImageInterlockedOrUInt64(RWTexture2D<UlongType> Dest, uint2 Coord, UlongType Value)
+{
+	InterlockedOr(Dest[Coord], Value);
+}
 #endif
 
 #if UE_BINDLESS_ENABLED
diff --git a/Engine/Source/Runtime/Renderer/Private/DeferredShadingRenderer.cpp b/Engine/Source/Runtime/Renderer/Private/DeferredShadingRenderer.cpp
--- a/Engine/Source/Runtime/Renderer/Private/DeferredShadingRenderer.cpp
+++ b/Engine/Source/Runtime/Renderer/Private/DeferredShadingRenderer.cpp
@@ -3515,6 +3515,8 @@ void FDeferredShadingSceneRenderer::Render(FRDGBuilder& GraphBuilder, const FSce
 					// Rebuild scene textures to include newly allocated velocity.
 					SceneTextures.UniformBuffer = CreateSceneTextureUniformBuffer(GraphBuilder, &SceneTextures, FeatureLevel, SceneTextures.SetupMode);
 				}
+
+				RenderVelocities(GraphBuilder, Views, SceneTextures, EVelocityPass::TranslucentClippedDepth, false, /*bBindRenderTarget=*/ false);
 			}
 		}
 		else if (GetHairStrandsComposition() == EHairStrandsCompositionType::AfterTranslucent)
diff --git a/Engine/Source/Runtime/Renderer/Private/PostProcess/TemporalAA.h b/Engine/Source/Runtime/Renderer/Private/PostProcess/TemporalAA.h
--- a/Engine/Source/Runtime/Renderer/Private/PostProcess/TemporalAA.h
+++ b/Engine/Source/Runtime/Renderer/Private/PostProcess/TemporalAA.h
@@ -179,6 +179,8 @@ bool IsTSRLensDistortionEnabled(EShaderPlatform ShaderPlatform);
 /** Returns whether a given view need to measure the scene for moire anti-flickering. */
 bool NeedTSRAntiFlickeringPass(const FViewInfo& View);
 
+bool NeedTSRThinGeometryDetectionPass(const FViewInfo& View);
+
 /** Returns whether TSR internal visualization is enabled on the view. */
 bool IsVisualizeTSREnabled(const FViewInfo& View);
 
diff --git a/Engine/Source/Runtime/Renderer/Private/PostProcess/TemporalSuperResolution.cpp b/Engine/Source/Runtime/Renderer/Private/PostProcess/TemporalSuperResolution.cpp
--- a/Engine/Source/Runtime/Renderer/Private/PostProcess/TemporalSuperResolution.cpp
+++ b/Engine/Source/Runtime/Renderer/Private/PostProcess/TemporalSuperResolution.cpp
@@ -997,6 +997,7 @@ class FTSRDetectThinGeometryCS : public FTSRConvolutionNetworkShader
 		SHADER_PARAMETER_RDG_TEXTURE(Texture2D, SceneDepthTexture)
 		SHADER_PARAMETER_RDG_TEXTURE_SRV(Texture2D, CurrentCoverageTexture)
 		SHADER_PARAMETER_RDG_TEXTURE_SRV(Texture2D, ReprojectedHistoryCoverageTexture)
+		SHADER_PARAMETER_RDG_TEXTURE(Texture2D, DecimateMaskTexture)
 
 		SHADER_PARAMETER_RDG_TEXTURE_UAV(RWTexture2DArray, R8Output)
 		SHADER_PARAMETER_RDG_TEXTURE_UAV(RWTexture2DArray, HistoryCoverageOutput)
@@ -1579,6 +1580,13 @@ bool NeedTSRAntiFlickeringPass(const FViewInfo& View)
 	return GetMainTAAPassConfig(View) == EMainTAAPassConfig::TSR && View.Family->EngineShowFlags.PostProcessing;
 }
 
+bool NeedTSRThinGeometryDetectionPass(const FViewInfo& View)
+{
+	return GetMainTAAPassConfig(View) == EMainTAAPassConfig::TSR && 
+		CVarTSRThinGeometryDetection.GetValueOnRenderThread() != 0 &&
+		View.Family->EngineShowFlags.PostProcessing;
+}
+
 bool ShouldAddTSRMainFlickeringLumaPass()
 {
 	return  (CVarTSRFlickeringEnable.GetValueOnRenderThread()!=0 && CVarTSRFlickeringPeriod.GetValueOnRenderThread() != 0.0f) || ShouldApplySkyRelaxation();
@@ -2625,6 +2633,7 @@ FDefaultTemporalUpscaler::FOutputs AddTemporalSuperResolutionPasses(
 			PassParameters->ErrorMultiplier = PassConfig.ThinGeometryErrorMultiplier;
 			PassParameters->MaxRelaxationWeight = FMath::Clamp(CVarTSRThinGeometryCoverageMaxRelaxationWeight.GetValueOnRenderThread(), 0.0f, 1.0f);
 			PassParameters->SceneDepthTexture = PassInputs.SceneDepth.Texture;
+			PassParameters->DecimateMaskTexture = DecimateMaskTexture;
 
 			// Coverage texture
 			{
diff --git a/Engine/Source/Runtime/Renderer/Private/SceneVisibility.cpp b/Engine/Source/Runtime/Renderer/Private/SceneVisibility.cpp
--- a/Engine/Source/Runtime/Renderer/Private/SceneVisibility.cpp
+++ b/Engine/Source/Runtime/Renderer/Private/SceneVisibility.cpp
@@ -1416,6 +1416,7 @@ void FRelevancePacket::ComputeRelevance(FDynamicPrimitiveIndexList& DynamicPrimi
 	const bool bMobileBasePassAlwaysUsesCSM = (ShadingPath == EShadingPath::Mobile) && MobileBasePassAlwaysUsesCSM(Scene.GetShaderPlatform());
 	const bool bVelocityPassWritesDepth = Scene.EarlyZPassMode == DDM_AllOpaqueNoVelocity;
 	const bool bIsTranslucentHoldoutEnabled = IsPrimitiveAlphaHoldoutEnabled(View);
+	const bool bIsTranslucentClippedDepthEnabled = NeedTSRThinGeometryDetectionPass(View);
 	const bool bHLODActive = Scene.SceneLODHierarchy.IsActive();
 	const FHLODVisibilityState* const HLODState = bHLODActive && ViewState ? &ViewState->HLODVisibilityState : nullptr;
 	float MaxDrawDistanceScale = GetCachedScalabilityCVars().ViewDistanceScale;
@@ -1628,6 +1629,12 @@ void FRelevancePacket::ComputeRelevance(FDynamicPrimitiveIndexList& DynamicPrimi
 											FTranslucentVelocityMeshProcessor::PrimitiveHasVelocityForFrame(PrimitiveSceneProxy))
 										{
 											DrawCommandPacket.AddCommandsForMesh(PrimitiveIndex, PrimitiveSceneInfo, StaticMeshRelevance, StaticMesh, CullingPayloadFlags, Scene, bCanCache, EMeshPass::TranslucentVelocity);
+											
+											// Only Before DoF Translucency needs this pass
+											if (ViewRelevance.bNormalTranslucency && bIsTranslucentClippedDepthEnabled)
+											{
+												DrawCommandPacket.AddCommandsForMesh(PrimitiveIndex, PrimitiveSceneInfo, StaticMeshRelevance, StaticMesh, CullingPayloadFlags, Scene, bCanCache, EMeshPass::TranslucentVelocityClippedDepth);
+											}
 										}
 									}
 								}
@@ -2433,6 +2440,12 @@ static void ComputeDynamicMeshRelevance(
 			{
 				PassMask.Set(EMeshPass::TranslucentVelocity);
 				View.NumVisibleDynamicMeshElements[EMeshPass::TranslucentVelocity] += NumElements;
+
+				if (ViewRelevance.bNormalTranslucency)
+				{
+					PassMask.Set(EMeshPass::TranslucentVelocityClippedDepth);
+					View.NumVisibleDynamicMeshElements[EMeshPass::TranslucentVelocityClippedDepth] += NumElements;
+				}
 			}
 
 			if (ViewRelevance.bUsesSingleLayerWaterMaterial)
diff --git a/Engine/Source/Runtime/Renderer/Private/VelocityRendering.cpp b/Engine/Source/Runtime/Renderer/Private/VelocityRendering.cpp
--- a/Engine/Source/Runtime/Renderer/Private/VelocityRendering.cpp
+++ b/Engine/Source/Runtime/Renderer/Private/VelocityRendering.cpp
@@ -55,6 +55,20 @@ static TAutoConsoleVariable<int32> CVarParallelVelocity(
 	ECVF_RenderThreadSafe
 	);
 
+static TAutoConsoleVariable<int32> CVarVelocityOutputTranslucentClippedDepthSupported(
+	TEXT("r.Velocity.OutputTranslucentClippedDepth.Supported"),
+	0,
+	TEXT("Whether the translucent velocity clipped depth pass is supported on the current platform.\n"),
+	ECVF_ReadOnly);
+
+static TAutoConsoleVariable<int32> CVarVelocityOutputTranslucentClippedDepthEnabled(
+	TEXT("r.Velocity.OutputTranslucentClippedDepth.Enabled"),
+	1,
+	TEXT("Enable/Disable the translucent velocity clipped depth pass on the fly.\n")\
+	TEXT("0: Skip this pass.\n")\
+	TEXT("1: Provide functions e.g., mark before DoF translucency pixel animation for opacity below clip value.\n"),
+	ECVF_RenderThreadSafe);
+
 DECLARE_GPU_DRAWCALL_STAT_NAMED(RenderVelocities, TEXT("Render Velocities"));
 
 /** Validate that deprecated CVars are no longer set. */
@@ -126,14 +140,39 @@ public:
 	{}
 };
 
+bool SupportsTranslucentClippedDepth(EShaderPlatform ShaderPlatform)
+{
+	// Translucent clipped depth requires atomics on uint64. So holds the same requirement of Nanite.
+	static FShaderPlatformCachedIniValue<int32> PerPlatformCVar(TEXT("r.Velocity.OutputTranslucentClippedDepth.Supported"));
+	return FDataDrivenShaderPlatformInfo::GetSupportsNanite(ShaderPlatform) && (PerPlatformCVar.Get(ShaderPlatform) != 0);
+}
+
 class FVelocityPS : public FMeshMaterialShader
 {
 public:
 	DECLARE_SHADER_TYPE(FVelocityPS, MeshMaterial);
 
+	class FDimensionVelocityClippedDepthPass : SHADER_PERMUTATION_BOOL("VELOCITY_CLIPPED_DEPTH_PASS");
+	using FPermutationDomain = TShaderPermutationDomain<FDimensionVelocityClippedDepthPass>;
+
 	static bool ShouldCompilePermutation(const FMeshMaterialShaderPermutationParameters& Parameters)
 	{
-		return FVelocityVS::ShouldCompilePermutation(Parameters);
+		if (FVelocityVS::ShouldCompilePermutation(Parameters))
+		{
+			FPermutationDomain PermutationVector(Parameters.PermutationId);
+			const bool bIsVelocityClippedDepthPass = PermutationVector.Get<FDimensionVelocityClippedDepthPass>();
+
+			if (bIsVelocityClippedDepthPass)
+			{
+				return SupportsTranslucentClippedDepth(Parameters.Platform);
+			}
+			else
+			{
+				return true;
+			}
+		}
+		
+		return false;
 	}
 
 	static void ModifyCompilationEnvironment(const FMaterialShaderPermutationParameters& Parameters, FShaderCompilerEnvironment& OutEnvironment)
@@ -141,8 +180,14 @@ public:
 		FMaterialShader::ModifyCompilationEnvironment(Parameters, OutEnvironment);
 		OutEnvironment.SetRenderTargetOutputFormat(0, PF_A16B16G16R16);
 
-		// We support velocity on thin trnaslucent only with masking, and only if the material is only made of thin translucent shading model.
+		// We support velocity on thin translucent only with masking, and only if the material is only made of thin translucent shading model.
 		OutEnvironment.SetDefine(TEXT("VELOCITY_THIN_TRANSLUCENT_MODE"), Parameters.MaterialParameters.ShadingModels.HasOnlyShadingModel(MSM_ThinTranslucent));
+
+		FPermutationDomain PermutationVector(Parameters.PermutationId);
+		if (PermutationVector.Get<FDimensionVelocityClippedDepthPass>())
+		{
+			OutEnvironment.CompilerFlags.Add(CFLAG_AllowTypedUAVLoads);
+		}
 	}
 
 	FVelocityPS() = default;
@@ -151,9 +196,12 @@ public:
 	{}
 };
 
+BEGIN_GLOBAL_SHADER_PARAMETER_STRUCT(FVelocityClippedDepthUniformParameters, )
+	SHADER_PARAMETER_RDG_TEXTURE_UAV(RWTexture2D<UlongType>, RWVelocity)
+END_GLOBAL_SHADER_PARAMETER_STRUCT()
+
 IMPLEMENT_SHADER_TYPE(,FVelocityVS, TEXT("/Engine/Private/VelocityShader.usf"), TEXT("MainVertexShader"), SF_Vertex);
 IMPLEMENT_SHADER_TYPE(,FVelocityPS, TEXT("/Engine/Private/VelocityShader.usf"), TEXT("MainPixelShader"), SF_Pixel);
-IMPLEMENT_SHADERPIPELINE_TYPE_VSPS(VelocityPipeline, FVelocityVS, FVelocityPS, true);
 
 EMeshPass::Type GetMeshPassFromVelocityPass(EVelocityPass VelocityPass)
 {
@@ -163,11 +211,24 @@ EMeshPass::Type GetMeshPassFromVelocityPass(EVelocityPass VelocityPass)
 		return EMeshPass::Velocity;
 	case EVelocityPass::Translucent:
 		return EMeshPass::TranslucentVelocity;
+	case EVelocityPass::TranslucentClippedDepth:
+		return EMeshPass::TranslucentVelocityClippedDepth;
 	}
 	check(false);
 	return EMeshPass::Velocity;
 }
 
+static const TCHAR* GetVelocityPassName(EVelocityPass VelocityPass)
+{
+	static const TCHAR* const kPassNames[] = {
+		TEXT("Opaque"),
+		TEXT("Translucent"),
+		TEXT("TranslucentClippedDepth")
+	};
+	static_assert(UE_ARRAY_COUNT(kPassNames) == int32(EVelocityPass::Count), "Fix me");
+	return kPassNames[int32(VelocityPass)];
+}
+
 bool FDeferredShadingSceneRenderer::ShouldRenderVelocities() const
 {
 	if (!FVelocityRendering::IsVelocityPassSupported(ShaderPlatform) || ViewFamily.UseDebugViewPS())
@@ -234,10 +295,57 @@ bool FMobileSceneRenderer::ShouldRenderVelocities() const
 	return bNeedsVelocity;
 }
 
+IMPLEMENT_STATIC_UNIFORM_BUFFER_SLOT(VelocityClippedDepth);
+IMPLEMENT_STATIC_UNIFORM_BUFFER_STRUCT(FVelocityClippedDepthUniformParameters, "VelocityClippedDepth", VelocityClippedDepth);
+
+bool NeedVelocityDepth(EShaderPlatform ShaderPlatform)
+{
+	// Lumen needs velocity depth
+	return (DoesProjectSupportDistanceFields() && FDataDrivenShaderPlatformInfo::GetSupportsLumenGI(ShaderPlatform)) 
+		|| FDataDrivenShaderPlatformInfo::GetSupportsRayTracing(ShaderPlatform);
+}
+
+static FRDGTextureUAVRef CreateDummyVelocityUAV(FRDGBuilder& GraphBuilder, EPixelFormat PixelFormat)
+{
+	FRDGTextureDesc Desc = FRDGTextureDesc::Create2D(
+		FIntPoint(1, 1),
+		PixelFormat,
+		FClearValueBinding::None,
+		/* InFlags = */ TexCreate_ShaderResource | TexCreate_UAV);
+
+	FRDGTextureRef DummyTexture = GraphBuilder.CreateTexture(Desc, TEXT("VelocityClippedDepth.UAVDummy"));
+
+	return GraphBuilder.CreateUAV(FRDGTextureUAVDesc(DummyTexture, 0, PixelFormat));
+};
+
+TRDGUniformBufferRef<FVelocityClippedDepthUniformParameters> BindTranslucentVelocityClippedDepthPassUniformParameters(FRDGBuilder& GraphBuilder, 
+	const FSceneTextures& SceneTextures, bool bWriteVelocity, EShaderPlatform ShaderPlatform)
+{
+	FVelocityClippedDepthUniformParameters& VelocityClippedDepthUniformParameters = *GraphBuilder.AllocParameters<FVelocityClippedDepthUniformParameters>();
+
+	const bool bNeedVelocityDepth = NeedVelocityDepth(ShaderPlatform);
+	const bool bIsOpenGLPlatform = IsOpenGLPlatform(ShaderPlatform);
+
+	if (bNeedVelocityDepth)
+	{
+		VelocityClippedDepthUniformParameters.RWVelocity = bWriteVelocity ?
+			GraphBuilder.CreateUAV(SceneTextures.Velocity) :
+			CreateDummyVelocityUAV(GraphBuilder, SceneTextures.Velocity->Desc.Format);
+	}
+	else
+	{
+		const EPixelFormat DummpyPixelFormat = bIsOpenGLPlatform ? PF_R16G16B16A16_UINT : PF_A16B16G16R16;
+		VelocityClippedDepthUniformParameters.RWVelocity = CreateDummyVelocityUAV(GraphBuilder, DummpyPixelFormat);
+	}
+
+	return GraphBuilder.CreateUniformBuffer(&VelocityClippedDepthUniformParameters);
+}
+
 BEGIN_SHADER_PARAMETER_STRUCT(FVelocityPassParameters, )
 	SHADER_PARAMETER_STRUCT_INCLUDE(FViewShaderParameters, View)
 	SHADER_PARAMETER_STRUCT_INCLUDE(FSceneTextureShaderParameters, SceneTextures)
 	SHADER_PARAMETER_STRUCT_INCLUDE(FInstanceCullingDrawParams, InstanceCullingDrawParams)
+	SHADER_PARAMETER_RDG_UNIFORM_BUFFER(FVelocityClippedDepthUniformParameters, VelocityClippedDepth)
 	RENDER_TARGET_BINDING_SLOTS()
 END_SHADER_PARAMETER_STRUCT()
 
@@ -249,6 +357,14 @@ void FSceneRenderer::RenderVelocities(
 	bool bForceVelocity,
 	bool bBindRenderTarget)
 {
+	const bool bIsTranslucentClippedDepthPass = VelocityPass == EVelocityPass::TranslucentClippedDepth;
+	const bool bSupportsTranslucentClippedDepth = SupportsTranslucentClippedDepth(ShaderPlatform);
+	const bool bIsTranslucentClippedDepthEnabled = CVarVelocityOutputTranslucentClippedDepthEnabled.GetValueOnRenderThread() != 0;
+	if (bIsTranslucentClippedDepthPass && (!bSupportsTranslucentClippedDepth || !bIsTranslucentClippedDepthEnabled))
+	{
+		return;
+	}
+
 	RDG_CSV_STAT_EXCLUSIVE_SCOPE(GraphBuilder, RenderVelocities);
 	SCOPED_NAMED_EVENT(FSceneRenderer_RenderVelocities, FColor::Emerald);
 	SCOPE_CYCLE_COUNTER(STAT_RenderVelocities);
@@ -256,13 +372,15 @@ void FSceneRenderer::RenderVelocities(
 	// Create mask for which GPUs we need clearing on
 	uint32 bNeedsClearMask = HasBeenProduced(SceneTextures.Velocity) ? 0 : ((1u << GNumExplicitGPUsForRendering) - 1);
 
-	RDG_EVENT_SCOPE_STAT(GraphBuilder, RenderVelocities, "RenderVelocities");
+	RDG_EVENT_SCOPE_STAT(GraphBuilder, RenderVelocities, "RenderVelocities(%s)",GetVelocityPassName(VelocityPass));
 	RDG_GPU_STAT_SCOPE(GraphBuilder, RenderVelocities);
 
 	const EMeshPass::Type MeshPass = GetMeshPassFromVelocityPass(VelocityPass);
+
 	FExclusiveDepthStencil ExclusiveDepthStencil = (VelocityPass == EVelocityPass::Opaque && !(Scene->EarlyZPassMode == DDM_AllOpaqueNoVelocity))
 														? FExclusiveDepthStencil::DepthRead_StencilWrite
 														: FExclusiveDepthStencil::DepthWrite_StencilWrite;
+	ExclusiveDepthStencil = bIsTranslucentClippedDepthPass ? FExclusiveDepthStencil::DepthRead_StencilNop : ExclusiveDepthStencil;
 
 	for (int32 ViewIndex = 0; ViewIndex < InViews.Num(); ViewIndex++)
 	{
@@ -314,6 +432,9 @@ void FSceneRenderer::RenderVelocities(
 				PassParameters->RenderTargets[0] = FRenderTargetBinding(SceneTextures.Velocity, (bNeedsClearMask & View.GPUMask.GetNative()) ? ERenderTargetLoadAction::EClear : ERenderTargetLoadAction::ELoad);
 				bNeedsClearMask &= ~View.GPUMask.GetNative();
 			}
+
+			PassParameters->VelocityClippedDepth = BindTranslucentVelocityClippedDepthPassUniformParameters(GraphBuilder,SceneTextures, bIsTranslucentClippedDepthPass, ShaderPlatform);
+
 			PassParameters->RenderTargets.MultiViewCount = (View.bIsMobileMultiViewEnabled) ? 2 : (View.Aspects.IsMobileMultiViewEnabled() ? 1 : 0);
 
 			if (bIsParallelVelocity)
@@ -381,9 +502,7 @@ void FMobileSceneRenderer::RenderVelocityPass(FRHICommandList& RHICmdList, const
 
 EPixelFormat FVelocityRendering::GetFormat(EShaderPlatform ShaderPlatform)
 {
-	// Lumen needs velocity depth
-	const bool bNeedVelocityDepth = (DoesProjectSupportDistanceFields() && FDataDrivenShaderPlatformInfo::GetSupportsLumenGI(ShaderPlatform)) 
-		|| FDataDrivenShaderPlatformInfo::GetSupportsRayTracing(ShaderPlatform);
+	bool bNeedVelocityDepth = NeedVelocityDepth(ShaderPlatform); 
 
 	// Android GLES platform doesn't support R16G16_UNORM and R16G16B16A16_UNORM format, use R16G16_UINT or R16G16B16A16_UINT instead.
 	const bool bIsOpenGLPlatform = IsOpenGLPlatform(ShaderPlatform);
@@ -400,7 +519,8 @@ EPixelFormat FVelocityRendering::GetFormat(EShaderPlatform ShaderPlatform)
 ETextureCreateFlags FVelocityRendering::GetCreateFlags(EShaderPlatform ShaderPlatform)
 {
 	const ETextureCreateFlags FastVRamFlag = BasePassCanOutputVelocity(ShaderPlatform) ? GFastVRamConfig.GBufferVelocity : TexCreate_None;
-	return TexCreate_RenderTargetable | TexCreate_UAV | TexCreate_ShaderResource | FastVRamFlag;
+	const ETextureCreateFlags Atomic64CompatibleFlag = NeedVelocityDepth(ShaderPlatform) ? ETextureCreateFlags::Atomic64Compatible : TexCreate_None;
+	return TexCreate_RenderTargetable | TexCreate_UAV | TexCreate_ShaderResource | FastVRamFlag | Atomic64CompatibleFlag;
 }
 
 FRDGTextureDesc FVelocityRendering::GetRenderTargetDesc(EShaderPlatform ShaderPlatform, FIntPoint Extent, const bool bRequireMultiView)
@@ -755,16 +875,22 @@ bool GetVelocityPassShaders(
 	const FMaterial& Material,
 	const FVertexFactoryType* VertexFactoryType,
 	ERHIFeatureLevel::Type FeatureLevel,
+	bool bIsTranslucentClippedDepthPass,
 	TShaderRef<FVelocityVS>& VertexShader,
 	TShaderRef<FVelocityPS>& PixelShader)
 {
 	FMaterialShaderTypes ShaderTypes;
 
-	// Don't use pipeline if we have hull/domain shaders
-	ShaderTypes.PipelineType = &VelocityPipeline;
-
 	ShaderTypes.AddShaderType<FVelocityVS>();
-	ShaderTypes.AddShaderType<FVelocityPS>();
+	
+	int32 PermutationId = 0;
+	if (bIsTranslucentClippedDepthPass)
+	{
+		FVelocityPS::FPermutationDomain PermutationVector;
+		PermutationVector.Set<FVelocityPS::FDimensionVelocityClippedDepthPass>(bIsTranslucentClippedDepthPass);
+		PermutationId = PermutationVector.ToDimensionValueId();
+	}
+	ShaderTypes.AddShaderType<FVelocityPS>(PermutationId);
 
 	FMaterialShaders Shaders;
 	if (!Material.TryGetShaders(ShaderTypes, VertexFactoryType, Shaders))
@@ -793,10 +919,12 @@ bool FVelocityMeshProcessor::Process(
 		FVelocityVS,
 		FVelocityPS> VelocityPassShaders;
 
+	bool bIsTranslucentClippedDepthPass = (MeshPassType == EMeshPass::TranslucentVelocityClippedDepth);
 	if (!GetVelocityPassShaders(
 		MaterialResource,
 		VertexFactory->GetType(),
 		FeatureLevel,
+		bIsTranslucentClippedDepthPass,
 		VelocityPassShaders.VertexShader,
 		VelocityPassShaders.PixelShader))
 	{
@@ -865,10 +993,13 @@ bool FVelocityMeshProcessor::CollectPSOInitializersInternal(
 		FVelocityVS,
 		FVelocityPS> VelocityPassShaders;
 
+	bool bIsTranslucentClippedDepthPass = (MeshPassType == EMeshPass::TranslucentVelocityClippedDepth);
+
 	if (!GetVelocityPassShaders(
 		MaterialResource,
 		VertexFactoryData.VertexFactoryType,
 		FeatureLevel,
+		bIsTranslucentClippedDepthPass,
 		VelocityPassShaders.VertexShader,
 		VelocityPassShaders.PixelShader))
 	{
@@ -902,10 +1033,12 @@ bool FVelocityMeshProcessor::CollectPSOInitializersInternal(
 	return true;
 }
 
-FVelocityMeshProcessor::FVelocityMeshProcessor(EMeshPass::Type MeshPassType, const FScene* Scene, ERHIFeatureLevel::Type FeatureLevel, const FSceneView* InViewIfDynamicMeshCommand, const FMeshPassProcessorRenderState& InPassDrawRenderState, FMeshPassDrawListContext* InDrawListContext)
-	: FMeshPassProcessor(MeshPassType, Scene, FeatureLevel, InViewIfDynamicMeshCommand, InDrawListContext)
+FVelocityMeshProcessor::FVelocityMeshProcessor(EMeshPass::Type InMeshPassType, const FScene* Scene, ERHIFeatureLevel::Type FeatureLevel, const FSceneView* InViewIfDynamicMeshCommand, const FMeshPassProcessorRenderState& InPassDrawRenderState, FMeshPassDrawListContext* InDrawListContext)
+	: FMeshPassProcessor(InMeshPassType, Scene, FeatureLevel, InViewIfDynamicMeshCommand, InDrawListContext)
 	, PassDrawRenderState(InPassDrawRenderState)
-{}
+{
+	MeshPassType = InMeshPassType;
+}
 
 FOpaqueVelocityMeshProcessor::FOpaqueVelocityMeshProcessor(const FScene* Scene, ERHIFeatureLevel::Type FeatureLevel, const FSceneView* InViewIfDynamicMeshCommand, const FMeshPassProcessorRenderState& InPassDrawRenderState, FMeshPassDrawListContext* InDrawListContext, EDepthDrawingMode InEarlyZPassMode)
 	: FVelocityMeshProcessor(EMeshPass::Velocity, Scene, FeatureLevel, InViewIfDynamicMeshCommand, InPassDrawRenderState, InDrawListContext)
@@ -931,8 +1064,8 @@ FMeshPassProcessor* CreateVelocityPassProcessor(ERHIFeatureLevel::Type FeatureLe
 REGISTER_MESHPASSPROCESSOR_AND_PSOCOLLECTOR(VelocityPass, CreateVelocityPassProcessor, EShadingPath::Deferred,  EMeshPass::Velocity, EMeshPassFlags::CachedMeshCommands | EMeshPassFlags::MainView);
 REGISTER_MESHPASSPROCESSOR_AND_PSOCOLLECTOR(MobileVelocityPass, CreateVelocityPassProcessor, EShadingPath::Mobile,  EMeshPass::Velocity, EMeshPassFlags::CachedMeshCommands | EMeshPassFlags::MainView);
 
-FTranslucentVelocityMeshProcessor::FTranslucentVelocityMeshProcessor(const FScene* Scene, ERHIFeatureLevel::Type FeatureLevel, const FSceneView* InViewIfDynamicMeshCommand, const FMeshPassProcessorRenderState& InPassDrawRenderState, FMeshPassDrawListContext* InDrawListContext)
-	: FVelocityMeshProcessor(EMeshPass::TranslucentVelocity, Scene, FeatureLevel, InViewIfDynamicMeshCommand, InPassDrawRenderState, InDrawListContext)
+FTranslucentVelocityMeshProcessor::FTranslucentVelocityMeshProcessor(const FScene* Scene, ERHIFeatureLevel::Type FeatureLevel, const FSceneView* InViewIfDynamicMeshCommand, const FMeshPassProcessorRenderState& InPassDrawRenderState, FMeshPassDrawListContext* InDrawListContext, EMeshPass::Type MeshPass)
+	: FVelocityMeshProcessor(MeshPass, Scene, FeatureLevel, InViewIfDynamicMeshCommand, InPassDrawRenderState, InDrawListContext)
 {}
 
 FMeshPassProcessor* CreateTranslucentVelocityPassProcessor(ERHIFeatureLevel::Type FeatureLevel, const FScene* Scene, const FSceneView* InViewIfDynamicMeshCommand, FMeshPassDrawListContext* InDrawListContext)
@@ -941,8 +1074,19 @@ FMeshPassProcessor* CreateTranslucentVelocityPassProcessor(ERHIFeatureLevel::Typ
 	VelocityPassState.SetBlendState(TStaticBlendState<CW_RGBA>::GetRHI());
 	VelocityPassState.SetDepthStencilState(TStaticDepthStencilState<true, CF_DepthNearOrEqual>::GetRHI());
 
-	return new FTranslucentVelocityMeshProcessor(Scene, FeatureLevel, InViewIfDynamicMeshCommand, VelocityPassState, InDrawListContext);
+	return new FTranslucentVelocityMeshProcessor(Scene, FeatureLevel, InViewIfDynamicMeshCommand, VelocityPassState, InDrawListContext, EMeshPass::TranslucentVelocity);
+}
+
+FMeshPassProcessor* CreateTranslucentVelocityClippedDepthPassProcessor(ERHIFeatureLevel::Type FeatureLevel, const FScene* Scene, const FSceneView* InViewIfDynamicMeshCommand, FMeshPassDrawListContext* InDrawListContext)
+{
+	FMeshPassProcessorRenderState VelocityPassState;
+	VelocityPassState.SetBlendState(TStaticBlendState<CW_RGBA>::GetRHI());
+	VelocityPassState.SetDepthStencilState(TStaticDepthStencilState<false, CF_DepthNearOrEqual>::GetRHI());
+
+	return new FTranslucentVelocityMeshProcessor(Scene, FeatureLevel, InViewIfDynamicMeshCommand, VelocityPassState, InDrawListContext, EMeshPass::TranslucentVelocityClippedDepth);
 }
 
 REGISTER_MESHPASSPROCESSOR_AND_PSOCOLLECTOR(TranslucentVelocityPass, CreateTranslucentVelocityPassProcessor, EShadingPath::Deferred, EMeshPass::TranslucentVelocity, EMeshPassFlags::CachedMeshCommands | EMeshPassFlags::MainView);
+REGISTER_MESHPASSPROCESSOR_AND_PSOCOLLECTOR(TranslucentVelocityClippedDepthPass, CreateTranslucentVelocityClippedDepthPassProcessor, EShadingPath::Deferred, EMeshPass::TranslucentVelocityClippedDepth, EMeshPassFlags::CachedMeshCommands | EMeshPassFlags::MainView);
 REGISTER_MESHPASSPROCESSOR_AND_PSOCOLLECTOR(MobileTranslucentVelocityPass, CreateTranslucentVelocityPassProcessor, EShadingPath::Mobile, EMeshPass::TranslucentVelocity, EMeshPassFlags::CachedMeshCommands | EMeshPassFlags::MainView);
+//TODO: Add Mobile translucent velocity clipped depth pass process when it is ready to support.
diff --git a/Engine/Source/Runtime/Renderer/Private/VelocityRendering.h b/Engine/Source/Runtime/Renderer/Private/VelocityRendering.h
--- a/Engine/Source/Runtime/Renderer/Private/VelocityRendering.h
+++ b/Engine/Source/Runtime/Renderer/Private/VelocityRendering.h
@@ -24,6 +24,9 @@ enum class EVelocityPass : uint32
 	// Renders a separate velocity / depth pass for translucency AFTER the translucent pass.
 	Translucent,
 
+	// Renders a secondary translucent velocity / depth pass for pixels that are clipped by opacity mask clip value (e.g., write has pixel animation mask).
+	TranslucentClippedDepth,
+
 	Count
 };
 
@@ -94,6 +97,7 @@ protected:
 
 protected:
 	EDepthDrawingMode EarlyZPassMode = DDM_None;
+	EMeshPass::Type MeshPassType = EMeshPass::Type::Num;
 };
 
 /**
@@ -152,7 +156,8 @@ public:
 		ERHIFeatureLevel::Type FeatureLevel,
 		const FSceneView* InViewIfDynamicMeshCommand,
 		const FMeshPassProcessorRenderState& InPassDrawRenderState,
-		FMeshPassDrawListContext* InDrawListContext);
+		FMeshPassDrawListContext* InDrawListContext,
+		EMeshPass::Type Meshpass);
 
 	/** Returns true if the object is capable of having velocity for any frame. */
 	static bool PrimitiveCanHaveVelocity(EShaderPlatform ShaderPlatform, const FPrimitiveSceneProxy* PrimitiveSceneProxy);
diff --git a/Engine/Source/Runtime/Renderer/Public/MeshPassProcessor.h b/Engine/Source/Runtime/Renderer/Public/MeshPassProcessor.h
--- a/Engine/Source/Runtime/Renderer/Public/MeshPassProcessor.h
+++ b/Engine/Source/Runtime/Renderer/Public/MeshPassProcessor.h
@@ -49,6 +49,7 @@ namespace EMeshPass
 		Distortion,
 		Velocity,
 		TranslucentVelocity,
+		TranslucentVelocityClippedDepth, /**  A pass to handle pixels with depth below the opaque threshold*/
 		TranslucencyStandard,
 		TranslucencyStandardModulate,
 		TranslucencyAfterDOF,
@@ -107,6 +108,7 @@ inline const TCHAR* GetMeshPassName(EMeshPass::Type MeshPass)
 	case EMeshPass::Distortion: return TEXT("Distortion");
 	case EMeshPass::Velocity: return TEXT("Velocity");
 	case EMeshPass::TranslucentVelocity: return TEXT("TranslucentVelocity");
+	case EMeshPass::TranslucentVelocityClippedDepth: return TEXT("TranslucentVelocityClippedDepth");
 	case EMeshPass::TranslucencyStandard: return TEXT("TranslucencyStandard");
 	case EMeshPass::TranslucencyStandardModulate: return TEXT("TranslucencyStandardModulate");
 	case EMeshPass::TranslucencyAfterDOF: return TEXT("TranslucencyAfterDOF");
@@ -142,9 +144,9 @@ inline const TCHAR* GetMeshPassName(EMeshPass::Type MeshPass)
 	}
 
 #if WITH_EDITOR
-	static_assert(EMeshPass::Num == 39 + 4, "Need to update switch(MeshPass) after changing EMeshPass"); // GUID to prevent incorrect auto-resolves, please change when changing the expression: {674D7D62-CFD8-4971-9A8D-CD91E5612CD8}
+	static_assert(EMeshPass::Num == 40 + 4, "Need to update switch(MeshPass) after changing EMeshPass"); // GUID to prevent incorrect auto-resolves, please change when changing the expression: {674D7D62-CFD8-4971-9A8D-CD91E5612CD8}
 #else
-	static_assert(EMeshPass::Num == 39, "Need to update switch(MeshPass) after changing EMeshPass"); // GUID to prevent incorrect auto-resolves, please change when changing the expression: {674D7D62-CFD8-4971-9A8D-CD91E5612CD8}
+	static_assert(EMeshPass::Num == 40, "Need to update switch(MeshPass) after changing EMeshPass"); // GUID to prevent incorrect auto-resolves, please change when changing the expression: {674D7D62-CFD8-4971-9A8D-CD91E5612CD8}
 #endif
 
 	checkf(0, TEXT("Missing case for EMeshPass %u"), (uint32)MeshPass);
