Fix shader compilation error with on non-deferred platform.

diff --git a/Engine/Source/Runtime/Renderer/Private/LightRendering.cpp b/Engine/Source/Runtime/Renderer/Private/LightRendering.cpp
--- a/Engine/Source/Runtime/Renderer/Private/LightRendering.cpp
+++ b/Engine/Source/Runtime/Renderer/Private/LightRendering.cpp
@@ -993,7 +993,7 @@ class FDeferredLightPS : public FGlobalShader
 		{
 			if (Substrate::IsSubstrateEnabled())
 			{
-				return Substrate::IsSubstrateBlendableGBufferEnabled(Parameters.Platform);
+				return IsFeatureLevelSupported(Parameters.Platform, ERHIFeatureLevel::SM5) && Substrate::IsSubstrateBlendableGBufferEnabled(Parameters.Platform);
 			}
 
 			// Anisotropic materials do not currently support rect lights
