Fix MTLPixelFormatRGBA16Unorm is not a valid format to be used with texture atomics in the velocity pass.

#tests Successfully launch on Mac SM6

diff --git a/Engine/Source/Runtime/Renderer/Private/VelocityRendering.cpp b/Engine/Source/Runtime/Renderer/Private/VelocityRendering.cpp
--- a/Engine/Source/Runtime/Renderer/Private/VelocityRendering.cpp
+++ b/Engine/Source/Runtime/Renderer/Private/VelocityRendering.cpp
@@ -324,17 +324,14 @@ TRDGUniformBufferRef<FVelocityClippedDepthUniformParameters> BindTranslucentVelo
 	FVelocityClippedDepthUniformParameters& VelocityClippedDepthUniformParameters = *GraphBuilder.AllocParameters<FVelocityClippedDepthUniformParameters>();
 
 	const bool bNeedVelocityDepth = NeedVelocityDepth(ShaderPlatform);
-	const bool bIsOpenGLPlatform = IsOpenGLPlatform(ShaderPlatform);
-
-	if (bNeedVelocityDepth)
+	
+	if (bNeedVelocityDepth && bWriteVelocity)
 	{
-		VelocityClippedDepthUniformParameters.RWVelocity = bWriteVelocity ?
-			GraphBuilder.CreateUAV(SceneTextures.Velocity) :
-			CreateDummyVelocityUAV(GraphBuilder, SceneTextures.Velocity->Desc.Format);
+		VelocityClippedDepthUniformParameters.RWVelocity = GraphBuilder.CreateUAV(SceneTextures.Velocity);
 	}
 	else
 	{
-		const EPixelFormat DummpyPixelFormat = bIsOpenGLPlatform ? PF_R16G16B16A16_UINT : PF_A16B16G16R16;
+		const EPixelFormat DummpyPixelFormat = GPixelFormats[PF_R64_UINT].Supported ? PF_R64_UINT : PF_R32G32_UINT;
 		VelocityClippedDepthUniformParameters.RWVelocity = CreateDummyVelocityUAV(GraphBuilder, DummpyPixelFormat);
 	}
 
@@ -519,7 +516,7 @@ EPixelFormat FVelocityRendering::GetFormat(EShaderPlatform ShaderPlatform)
 ETextureCreateFlags FVelocityRendering::GetCreateFlags(EShaderPlatform ShaderPlatform)
 {
 	const ETextureCreateFlags FastVRamFlag = BasePassCanOutputVelocity(ShaderPlatform) ? GFastVRamConfig.GBufferVelocity : TexCreate_None;
-	const ETextureCreateFlags Atomic64CompatibleFlag = NeedVelocityDepth(ShaderPlatform) ? ETextureCreateFlags::Atomic64Compatible : TexCreate_None;
+	const ETextureCreateFlags Atomic64CompatibleFlag = (NeedVelocityDepth(ShaderPlatform) && SupportsTranslucentClippedDepth(ShaderPlatform)) ? ETextureCreateFlags::Atomic64Compatible : TexCreate_None;
 	return TexCreate_RenderTargetable | TexCreate_UAV | TexCreate_ShaderResource | FastVRamFlag | Atomic64CompatibleFlag;
 }
 
