Fixed: RDG validation assert around InstanceCulling resources usage in a mobile renderer

diff --git a/Engine/Source/Runtime/Renderer/Private/InstanceCulling/InstanceCullingContext.cpp b/Engine/Source/Runtime/Renderer/Private/InstanceCulling/InstanceCullingContext.cpp
--- a/Engine/Source/Runtime/Renderer/Private/InstanceCulling/InstanceCullingContext.cpp
+++ b/Engine/Source/Runtime/Renderer/Private/InstanceCulling/InstanceCullingContext.cpp
@@ -1351,7 +1351,13 @@ FInstanceCullingDeferredContext *FInstanceCullingContext::CreateDeferredContext(
 	return DeferredContext;
 }
 
-
+void FInstanceCullingContext::SetDeferredContextCullingParams(const FInstanceCullingDeferredContext& DeferredContext, FInstanceCullingDrawParams& Parameters)
+{
+	Parameters.DrawIndirectArgsBuffer = DeferredContext.DrawIndirectArgsBuffer;
+	Parameters.InstanceIdOffsetBuffer = DeferredContext.InstanceDataBuffer;
+	Parameters.InstanceCulling = DeferredContext.UniformBuffer;
+	Parameters.BatchedPrimitive = DeferredContext.BatchedPrimitive;
+}
 
 class FClearIndirectArgInstanceCountCs : public FGlobalShader
 {
diff --git a/Engine/Source/Runtime/Renderer/Private/InstanceCulling/InstanceCullingManager.cpp b/Engine/Source/Runtime/Renderer/Private/InstanceCulling/InstanceCullingManager.cpp
--- a/Engine/Source/Runtime/Renderer/Private/InstanceCulling/InstanceCullingManager.cpp
+++ b/Engine/Source/Runtime/Renderer/Private/InstanceCulling/InstanceCullingManager.cpp
@@ -83,6 +83,13 @@ void FInstanceCullingManager::SetDummyCullingParams(FRDGBuilder& GraphBuilder, F
 {
 	Parameters.Scene = SceneUniforms.GetBuffer(GraphBuilder);
 	Parameters.InstanceCulling = GetDummyInstanceCullingUniformBuffer();
+
+	// Mobile renderer renders several mesh passes in a single RDG pass, however RDG pass can accept only a single FInstanceCullingDrawParams
+	// This makes sure that merged culling parameters are always end up in a RDG pass struct for a mobile renderer (see BuildMeshRenderingCommands)
+	if (DeferredContext && IsMobilePlatform(Scene.GetShaderPlatform()))
+	{
+		FInstanceCullingContext::SetDeferredContextCullingParams(*DeferredContext, Parameters);
+	}
 }
 
 bool FInstanceCullingManager::AllowBatchedBuildRenderingCommands(const FGPUScene& GPUScene)
diff --git a/Engine/Source/Runtime/Renderer/Public/InstanceCulling/InstanceCullingContext.h b/Engine/Source/Runtime/Renderer/Public/InstanceCulling/InstanceCullingContext.h
--- a/Engine/Source/Runtime/Renderer/Public/InstanceCulling/InstanceCullingContext.h
+++ b/Engine/Source/Runtime/Renderer/Public/InstanceCulling/InstanceCullingContext.h
@@ -217,6 +217,8 @@ public:
 		const FGPUScene& GPUScene,
 		FInstanceCullingManager& InstanceCullingManager);
 
+	static void SetDeferredContextCullingParams(const FInstanceCullingDeferredContext& DeferredContext, FInstanceCullingDrawParams& Parameters);
+
 	/**
 	 * Helper function to add a pass to zero the instance count in the indirect args.
 	 */
