SceneCapture: View rectangle fix for camera constrained custom render pass.

We now rely on the view rectangle size to ignore aspect ratio constraints (black bars) when bIgnoreScreenPercentage is enabled on custom render passes that copy the main view resolution.

diff --git a/Engine/Source/Runtime/Renderer/Private/SceneRendering.cpp b/Engine/Source/Runtime/Renderer/Private/SceneRendering.cpp
--- a/Engine/Source/Runtime/Renderer/Private/SceneRendering.cpp
+++ b/Engine/Source/Runtime/Renderer/Private/SceneRendering.cpp
@@ -3453,7 +3453,9 @@ void FSceneRenderer::PrepareViewRectsForRendering()
 			{
 				if (SceneCaptureUserData.bIgnoreScreenPercentage)
 				{
-					View.ViewRect = GetDownscaledViewRect(Views[0].UnscaledViewRect, Views[0].GetUnscaledFamilyViewRect().Max, SceneCaptureUserData.SceneTextureDivisor);
+					const FIntRect SourceViewRect = FIntRect(FInt32Point::ZeroValue, Views[0].UnscaledViewRect.Size());
+					
+					View.ViewRect = GetDownscaledViewRect(SourceViewRect, Views[0].GetUnscaledFamilyViewRect().Size(), SceneCaptureUserData.SceneTextureDivisor);
 				}
 				else
 				{
