Fix resolution scale not being handled properly for scaling hair strands size causing incorrect aliasing.

diff --git a/Engine/Source/Runtime/Renderer/Private/HairStrands/HairStrandsVisibility.cpp b/Engine/Source/Runtime/Renderer/Private/HairStrands/HairStrandsVisibility.cpp
--- a/Engine/Source/Runtime/Renderer/Private/HairStrands/HairStrandsVisibility.cpp
+++ b/Engine/Source/Runtime/Renderer/Private/HairStrands/HairStrandsVisibility.cpp
@@ -46,6 +46,9 @@ static FAutoConsoleVariableRef CVarHairVisibilityCompute_MeanSamplePerPixel(TEXT
 static int32 GHairClearVisibilityBuffer = 0;
 static FAutoConsoleVariableRef CVarHairClearVisibilityBuffer(TEXT("r.HairStrands.Visibility.Clear"), GHairClearVisibilityBuffer, TEXT("Clear hair strands visibility buffer"));
 
+static int32 GHairResolutionScaleAware = 1;
+static FAutoConsoleVariableRef CVarHairResolutionScaleAware(TEXT("r.HairStrands.Visibility.ResolutionScaleAware"), GHairResolutionScaleAware, TEXT("Scale hair strands based on resolution scale"));
+
 static TAutoConsoleVariable<int32> CVarHairVelocityMagnitudeScale(
 	TEXT("r.HairStrands.VelocityMagnitudeScale"),
 	100,  // Tuned by eye, based on heavy motion (strong head shack)
@@ -301,8 +304,10 @@ void SetUpViewHairRenderInfo(const FViewInfo& ViewInfo, bool bEnableMSAA, FVecto
 	// In the case we render coverage, we need to override some view uniform shader parameters to account for the change in MSAA sample count.
 	const uint32 HairVisibilitySampleCount = bEnableMSAA ? GetMaxSamplePerPixel(ViewInfo.GetShaderPlatform()) : 1;	// The coverage pass does not use MSAA
 	const float RasterizationScaleOverride = 0.0f;	// no override
-	FMinHairRadiusAtDepth1 MinHairRadius = ComputeMinStrandRadiusAtDepth1(
-		FIntPoint(ViewInfo.UnconstrainedViewRect.Width(), ViewInfo.UnconstrainedViewRect.Height()), ViewInfo.FOV, HairVisibilitySampleCount, RasterizationScaleOverride, ViewInfo.ViewMatrices.GetOrthoDimensions().X);
+
+	const FIntRect EffectiveViewRect = GHairResolutionScaleAware > 0 ? ViewInfo.ViewRect : ViewInfo.UnconstrainedViewRect;
+	const FMinHairRadiusAtDepth1 MinHairRadius = ComputeMinStrandRadiusAtDepth1(
+		FIntPoint(EffectiveViewRect.Width(), EffectiveViewRect.Height()), ViewInfo.FOV, HairVisibilitySampleCount, RasterizationScaleOverride, ViewInfo.ViewMatrices.GetOrthoDimensions().X);
 
 	OutHairRenderInfo = PackHairRenderInfo(MinHairRadius.Primary, MinHairRadius.Stable, MinHairRadius.Velocity, VelocityMagnitudeScale);
 	OutHairRenderInfoBits = PackHairRenderInfoBits(!ViewInfo.IsPerspectiveProjection(), false);
