Fix for crash when viz HDR is used together with wireframe.
Now wireframe disables viz HDR.

#tests PC editor
#rnx

diff --git a/Engine/Source/Runtime/Renderer/Private/DebugProbeRendering.cpp b/Engine/Source/Runtime/Renderer/Private/DebugProbeRendering.cpp
--- a/Engine/Source/Runtime/Renderer/Private/DebugProbeRendering.cpp
+++ b/Engine/Source/Runtime/Renderer/Private/DebugProbeRendering.cpp
@@ -88,7 +88,8 @@ IMPLEMENT_GLOBAL_SHADER(FStampDeferredDebugProbePS, "/Engine/Private/DebugProbes
 static bool ViewRequiresAndSupportsIlluminanceMeter(const FViewInfo& View)
 {
 	if (View.Family->EngineShowFlags.VisualizeHDR		// This is the debug view contain the illuminance meter
-		&& IsIlluminanceMeterSupportedByView(View))		// For instance: forward shading does not work with illuminance meter (we need a material buffer to work with)
+		&& IsIlluminanceMeterSupportedByView(View)		// For instance: forward shading does not work with illuminance meter (we need a material buffer to work with)
+		&& !View.Family->EngineShowFlags.Wireframe)		// Disable when in wireframe mode
 	{
 		return true;
 	}
diff --git a/Engine/Source/Runtime/Renderer/Private/PostProcess/PostProcessing.cpp b/Engine/Source/Runtime/Renderer/Private/PostProcess/PostProcessing.cpp
--- a/Engine/Source/Runtime/Renderer/Private/PostProcess/PostProcessing.cpp
+++ b/Engine/Source/Runtime/Renderer/Private/PostProcess/PostProcessing.cpp
@@ -424,7 +424,7 @@ void AddPostProcessingPasses(
 
 	const bool bViewDebugMaterialsEnabled = View.RequiresDebugMaterials();
 	const FEngineShowFlags& EngineShowFlags = View.Family->EngineShowFlags;
-	const bool bVisualizeHDR = EngineShowFlags.VisualizeHDR;
+	const bool bVisualizeHDR = EngineShowFlags.VisualizeHDR && !EngineShowFlags.Wireframe;
 	const bool bViewFamilyOutputInHDR = View.Family->RenderTarget->GetSceneHDREnabled();
 	const bool bVisualizeGBufferOverview = IsVisualizeGBufferOverviewEnabled(View);
 	const bool bVisualizeGBufferDumpToFile = IsVisualizeGBufferDumpToFileEnabled(View);
@@ -589,7 +589,7 @@ void AddPostProcessingPasses(
 	PassSequence.SetEnabled(EPass::VisualizeSubsurface, EngineShowFlags.VisualizeSSS);
 	PassSequence.SetEnabled(EPass::VisualizeGBufferOverview, bVisualizeGBufferOverview || bVisualizeGBufferDumpToFile || bVisualizeGBufferDumpToPIpe);
 	PassSequence.SetEnabled(EPass::VisualizeLumenSceneOverview, (LumenVisualizeMode == VISUALIZE_MODE_OVERVIEW || LumenVisualizeMode == VISUALIZE_MODE_PERFORMANCE_OVERVIEW) && bPostProcessingEnabled);
-	PassSequence.SetEnabled(EPass::VisualizeHDR, EngineShowFlags.VisualizeHDR);
+	PassSequence.SetEnabled(EPass::VisualizeHDR, bVisualizeHDR);
 	PassSequence.SetEnabled(EPass::VisualizeMotionVectors, EngineShowFlags.VisualizeMotionVectors || EngineShowFlags.VisualizeReprojection);
 	PassSequence.SetEnabled(EPass::VisualizeTemporalUpscaler, EngineShowFlags.VisualizeTemporalUpscaler);
 #if WITH_EDITOR
@@ -2324,7 +2324,6 @@ void AddVisualizeCalibrationMaterialPostProcessingPasses(FRDGBuilder& GraphBuild
 	FScreenPassTexture SceneColor((*Inputs.SceneTextures)->SceneColorTexture, PrimaryViewRect);
 
 	const FEngineShowFlags& EngineShowFlags = View.Family->EngineShowFlags;
-	const bool bVisualizeHDR = EngineShowFlags.VisualizeHDR;
 	const bool bViewFamilyOutputInHDR = View.Family->RenderTarget->GetSceneHDREnabled();
 	const bool bOutputInHDR = IsPostProcessingOutputInHDR();
 
diff --git a/Engine/Source/Runtime/Renderer/Private/SceneRendering.cpp b/Engine/Source/Runtime/Renderer/Private/SceneRendering.cpp
--- a/Engine/Source/Runtime/Renderer/Private/SceneRendering.cpp
+++ b/Engine/Source/Runtime/Renderer/Private/SceneRendering.cpp
@@ -4775,7 +4775,8 @@ void FSceneRenderer::SetupMeshPass(FViewInfo& View, FExclusiveDepthStencil::Type
 bool FSceneRenderer::ShouldCompositeEditorPrimitives(const FViewInfo& View)
 {
 	const FEngineShowFlags& ShowFlags = View.Family->EngineShowFlags;
-	if (ShowFlags.VisualizeHDR || ShowFlags.VisualizeSkyLightIlluminance ||
+	if ((ShowFlags.VisualizeHDR && !ShowFlags.Wireframe) ||				// VisualizeHDR is skipped with in Wireframe mode.
+		ShowFlags.VisualizeSkyLightIlluminance ||
 		ShowFlags.VisualizePostProcessStack || View.Family->UseDebugViewPS())
 	{
 		// certain visualize modes get obstructed too much
