Fix crash in ComposeMeshEdges when the result from the prepass is missing

#tests Test Animation.AssetValidation.AnimRuntime.SkeletalMesh.PreviewScene.Profiles.Switching.UpdatesSceneInRealTime was run locally and passes
#rnx

diff --git a/Engine/Source/Runtime/Renderer/Private/MeshEdgesRendering.cpp b/Engine/Source/Runtime/Renderer/Private/MeshEdgesRendering.cpp
--- a/Engine/Source/Runtime/Renderer/Private/MeshEdgesRendering.cpp
+++ b/Engine/Source/Runtime/Renderer/Private/MeshEdgesRendering.cpp
@@ -187,15 +187,29 @@ const FMeshEdgesViewSettings& GetMeshEdgesViewSettings(const FSceneViewFamily& V
 	return FamilyData->ViewSettings[ViewIndex];
 }
 
-const FMeshEdgesViewSettings& GetMeshEdgesViewSettings(const FSceneView& View)
+static int FindViewIndex(const FSceneView& View)
 {
-	check(View.Family);
+	if (!View.Family)
+	{
+		return INDEX_NONE;
+	}
 
-	int ViewIndex = 0;
-	for (; ViewIndex < View.Family->Views.Num(); ViewIndex++)
+	for (int ViewIndex = 0; ViewIndex < View.Family->Views.Num(); ViewIndex++)
 	{
-		if (View.Family->Views[ViewIndex] == &View) break;
+		if (View.Family->Views[ViewIndex] == &View)
+		{
+			return ViewIndex;
+		}
 	}
+	return INDEX_NONE;
+}
+
+const FMeshEdgesViewSettings& GetMeshEdgesViewSettings(const FSceneView& View)
+{
+	check(View.Family);
+
+	int ViewIndex = FindViewIndex(View);
+	if (ViewIndex == INDEX_NONE) { ViewIndex = 0; }
 
 	return GetMeshEdgesViewSettings(*View.Family, ViewIndex);
 }
@@ -364,6 +378,8 @@ void RenderMeshEdges(FSceneViewFamily& ViewFamily)
 		SceneRenderBuilder.Execute();
 	};
 
+	ViewFamilyData->ViewRects.SetNumZeroed(ViewFamily.Views.Num());
+
 	SceneRenderBuilder.AddRenderer(SceneRenderer, TEXT("RenderMeshEdges"),
 		[ViewFamilyData] (FRDGBuilder& GraphBuilder, const FSceneRenderFunctionInputs& Inputs)
 	{
@@ -374,7 +390,11 @@ void RenderMeshEdges(FSceneViewFamily& ViewFamily)
 
 		for (const FViewInfo& ViewInfo : Inputs.Renderer->Views)
 		{
-			ViewFamilyData->ViewRects.Emplace(ViewInfo.ViewRect);
+			int ViewIndex = FindViewIndex(ViewInfo);
+			if (ViewIndex != INDEX_NONE)
+			{
+				ViewFamilyData->ViewRects[ViewIndex] = ViewInfo.ViewRect;
+			}
 		}
 
 		return true;
@@ -393,10 +413,11 @@ void ComposeMeshEdges(FRDGBuilder& GraphBuilder,
 		return;
 	}
 
-	int ViewIndex = 0;
-	for (; ViewIndex < ViewFamily.Views.Num(); ViewIndex++)
+	int ViewIndex = FindViewIndex(View);
+	if (ViewIndex == INDEX_NONE)
 	{
-		if (ViewFamily.Views[ViewIndex] == &View) break;
+		// Could not find view
+		return;
 	}
 
 	if (ViewIndex >= ViewFamily.Views.Num())
@@ -407,6 +428,11 @@ void ComposeMeshEdges(FRDGBuilder& GraphBuilder,
 
 	const FMeshEdgesViewFamilyData* ViewFamilyData = ViewFamily.GetExtentionData<FMeshEdgesViewFamilyData>();
 	check(ViewFamilyData); // should have been created in RenderMeshEdges
+	if (ViewFamilyData->ViewRects.Num() == 0)
+	{
+		// Cannot compose if the prepass did not render
+		return;
+	}
 	check(ViewFamilyData->ViewRects.Num() > ViewIndex);
 	check(ViewFamilyData->WireframeColor);
 	check(ViewFamilyData->WireframeDepth);
