Fix VolumetricFogScreenToResourceUV not taking rounding to tile size into account.

diff --git a/Engine/Source/Runtime/Renderer/Private/SceneRendering.cpp b/Engine/Source/Runtime/Renderer/Private/SceneRendering.cpp
--- a/Engine/Source/Runtime/Renderer/Private/SceneRendering.cpp
+++ b/Engine/Source/Runtime/Renderer/Private/SceneRendering.cpp
@@ -1644,7 +1644,7 @@ void FViewInfo::SetupUniformBufferParameters(
 		ViewUniformShaderParameters.VolumetricFogPrevUVMaxForTemporalBlend = ViewState->VolumetricFogPrevUVMaxForTemporalBlend;
 		ViewUniformShaderParameters.VolumetricFogPrevResourceGridSize = FVector3f(ViewState->VolumetricFogPrevResourceGridSize);
 	}
-	ViewUniformShaderParameters.VolumetricFogScreenToResourceUV = ViewRectSize / (FVector2f(VolumetricFogResourceGridSize.X, VolumetricFogResourceGridSize.Y) * VolumetricFogResourceGridPixelSize);
+	ViewUniformShaderParameters.VolumetricFogScreenToResourceUV = FVector2f(VolumetricFogViewGridSize.X, VolumetricFogViewGridSize.Y) / FVector2f(VolumetricFogResourceGridSize.X, VolumetricFogResourceGridSize.Y);
 	ViewUniformShaderParameters.VolumetricFogUVMax = GetVolumetricFogUVMaxForSampling(ViewRectSize, VolumetricFogResourceGridSize, VolumetricFogResourceGridPixelSize);
 
 	SetupPrecomputedVolumetricLightmapUniformBufferParameters(Scene, Family->EngineShowFlags, ViewUniformShaderParameters);
