Material Cache, fixed last border bc block not getting compressed

diff --git a/Engine/Source/Runtime/Renderer/Private/MaterialCache/MaterialCacheRenderer.cpp b/Engine/Source/Runtime/Renderer/Private/MaterialCache/MaterialCacheRenderer.cpp
--- a/Engine/Source/Runtime/Renderer/Private/MaterialCache/MaterialCacheRenderer.cpp
+++ b/Engine/Source/Runtime/Renderer/Private/MaterialCache/MaterialCacheRenderer.cpp
@@ -1557,7 +1557,7 @@ static void MaterialCacheFinalizePages(FSceneRendererBase* Renderer, FRDGBuilder
 	);
 
 	const uint32 BlockSH = 2;
-	const uint32 TexelThreadCountX = GetMaterialCacheTileWidth() + GetMaterialCacheTileBorderWidth();
+	const uint32 TexelThreadCountX = GetMaterialCacheTileWidth() + GetMaterialCacheTileBorderWidth() * 2;
 	const uint32 BlockThreadCountX = TexelThreadCountX >> BlockSH;
 	
 	FMaterialCacheABufferWritePagesCS::FParameters* PassParameters = GraphBuilder.AllocParameters<FMaterialCacheABufferWritePagesCS::FParameters>();
