Add platform cached ini value to determine whether to compile 128 bit base pass pixel shader variations on platforms which require them.
These are infrequently needed and turning them off can save 50k shaders / 15 MiB at runtime
r.128BitBPPSCompilation.Allow (default is true, for backwards compatability)

#tests Cooked PF and RR on affected platforms

diff --git a/Engine/Source/Runtime/RenderCore/Private/RenderUtils.cpp b/Engine/Source/Runtime/RenderCore/Private/RenderUtils.cpp
--- a/Engine/Source/Runtime/RenderCore/Private/RenderUtils.cpp
+++ b/Engine/Source/Runtime/RenderCore/Private/RenderUtils.cpp
@@ -1670,6 +1670,12 @@ RENDERCORE_API bool AllowTranslucencyPerObjectShadows(const FStaticShaderPlatfor
 	return IsFeatureLevelSupported(Platform, ERHIFeatureLevel::SM5) && GAllowTranslucencyShadowsInProject != 0;
 }
 
+RENDERCORE_API bool Allow128BitBasePassPSCompilation(const FStaticShaderPlatform Platform)
+{
+	static FShaderPlatformCachedIniValue<bool> PerPlatformCVar(TEXT("r.128BitBPPSCompilation.Allow"));
+	return PerPlatformCVar.Get(Platform);
+}
+
 RENDERCORE_API bool ShouldGenerateRayTracingProxiesByDefault()
 {
 	return GRayTracingProxiesProjectEnabled;
diff --git a/Engine/Source/Runtime/RenderCore/Public/RenderUtils.h b/Engine/Source/Runtime/RenderCore/Public/RenderUtils.h
--- a/Engine/Source/Runtime/RenderCore/Public/RenderUtils.h
+++ b/Engine/Source/Runtime/RenderCore/Public/RenderUtils.h
@@ -548,6 +548,12 @@ RENDERCORE_API bool DoesRuntimeSupportOnePassPointLightShadows(EShaderPlatform P
  */
 RENDERCORE_API bool AllowTranslucencyPerObjectShadows(const FStaticShaderPlatform Platform);
 
+/**
+* Read-only switch to check if compiling explicit 128BitRT BasePass Pixel Shaders is allowed
+* They are infrequently necessary, and can save 50k shaders / 15 MiB / platform on large projects
+* Check bRequiresExplicit128bitRT in a platform's DataDrivenPlatformInfo.ini to see if that platform requires explicit 128Bit shaders
+*/
+RENDERCORE_API bool Allow128BitBasePassPSCompilation(const FStaticShaderPlatform Platform);
 
 /**Note, this should only be used when a platform requires special shader compilation for 32 bit pixel format render targets.
 Does not replace pixel format associations across the board**/
diff --git a/Engine/Source/Runtime/Renderer/Private/BasePassRendering.cpp b/Engine/Source/Runtime/Renderer/Private/BasePassRendering.cpp
--- a/Engine/Source/Runtime/Renderer/Private/BasePassRendering.cpp
+++ b/Engine/Source/Runtime/Renderer/Private/BasePassRendering.cpp
@@ -117,6 +117,13 @@ static TAutoConsoleVariable<int32> CVarPSOPrecacheAlphaColorChannel(
 	ECVF_ReadOnly
 );
 
+static TAutoConsoleVariable<bool> CVarAllow128BitBPPSCompilation(
+	TEXT("r.128BitBPPSCompilation.Allow"),
+	true,
+	TEXT("Whether or not to allow 128 Bit BasePass Pixel Shader compilation on platforms which require them.\n"),
+	ECVF_RenderThreadSafe | ECVF_ReadOnly
+);
+
 // Scene color alpha is used during scene captures and planar reflections.  1 indicates background should be shown, 0 indicates foreground is fully present.
 static const float kSceneColorClearAlpha = 1.0f;
 
@@ -178,7 +185,8 @@ F128BitRTBasePassPS::F128BitRTBasePassPS(const ShaderMetaType::CompiledShaderIni
 
 bool F128BitRTBasePassPS::ShouldCompilePermutation(const FMeshMaterialShaderPermutationParameters& Parameters)
 {
-	return FDataDrivenShaderPlatformInfo::GetRequiresExplicit128bitRT(Parameters.Platform);
+	return FDataDrivenShaderPlatformInfo::GetRequiresExplicit128bitRT(Parameters.Platform) && 
+		   Allow128BitBasePassPSCompilation(Parameters.Platform);
 }
 
 void F128BitRTBasePassPS::ModifyCompilationEnvironment(const FMeshMaterialShaderPermutationParameters& Parameters, FShaderCompilerEnvironment& OutEnvironment)
