Rename FLightOcclusionType to ELightOcclusionType.

diff --git a/Engine/Source/Runtime/RenderCore/Private/RayTracingGeometry.cpp b/Engine/Source/Runtime/RenderCore/Private/RayTracingGeometry.cpp
--- a/Engine/Source/Runtime/RenderCore/Private/RayTracingGeometry.cpp
+++ b/Engine/Source/Runtime/RenderCore/Private/RayTracingGeometry.cpp
@@ -32,20 +32,17 @@ PRAGMA_DISABLE_DEPRECATION_WARNINGS // Remove this in UE 5.6 once RayTracingGeom
 
 bool FRayTracingGeometry::HasValidInitializer() const
 {
-	if (Initializer.OfflineData != nullptr)
-	{
-		return true;
-	}
-
 	bool bAllSegmentsAreValid = Initializer.Segments.Num() > 0;
 	for (const FRayTracingGeometrySegment& Segment : Initializer.Segments)
 	{
 		if (!Segment.VertexBuffer)
 		{
 			bAllSegmentsAreValid = false;
+			check(Initializer.OfflineData == nullptr);
 			break;
 		}
 	}
+
 	return bAllSegmentsAreValid;
 }
 
diff --git a/Engine/Source/Runtime/Renderer/Private/DeferredShadingRenderer.cpp b/Engine/Source/Runtime/Renderer/Private/DeferredShadingRenderer.cpp
--- a/Engine/Source/Runtime/Renderer/Private/DeferredShadingRenderer.cpp
+++ b/Engine/Source/Runtime/Renderer/Private/DeferredShadingRenderer.cpp
@@ -3981,7 +3981,7 @@ void FDeferredShadingSceneRenderer::InitializeRayTracingFlags_RenderThread()
 	// We currently don't need a full list of RT lights, only whether there are any RT lights at all.
 	for (const FLightSceneInfoCompact& LightSceneInfoCompact : Scene->Lights)
 	{
-		if (GetLightOcclusionType(LightSceneInfoCompact, ViewFamily) == FLightOcclusionType::Raytraced)
+		if (GetLightOcclusionType(LightSceneInfoCompact, ViewFamily) == ELightOcclusionType::Raytraced)
 		{
 			bRayTracingShadows = true;
 			break;
diff --git a/Engine/Source/Runtime/Renderer/Private/LightRendering.cpp b/Engine/Source/Runtime/Renderer/Private/LightRendering.cpp
--- a/Engine/Source/Runtime/Renderer/Private/LightRendering.cpp
+++ b/Engine/Source/Runtime/Renderer/Private/LightRendering.cpp
@@ -729,7 +729,7 @@ FDeferredLightUniformStruct GetSimpleDeferredLightParameters(
 	return GetSimpleDeferredLightParameters(View, SimpleLight, SimpleLightPerViewData.Position);
 }
 
-FLightOcclusionType GetLightOcclusionType(const FLightSceneProxy& Proxy, const FSceneViewFamily& ViewFamily)
+ELightOcclusionType GetLightOcclusionType(const FLightSceneProxy& Proxy, const FSceneViewFamily& ViewFamily)
 {
 	bool bUseRaytracing = false;
 #if RHI_RAYTRACING
@@ -739,13 +739,13 @@ FLightOcclusionType GetLightOcclusionType(const FLightSceneProxy& Proxy, const F
 
 	if (MegaLightsMode != EMegaLightsMode::Disabled)
 	{
-		return MegaLightsMode == EMegaLightsMode::EnabledVSM ? FLightOcclusionType::MegaLightsVSM : FLightOcclusionType::MegaLights;
+		return MegaLightsMode == EMegaLightsMode::EnabledVSM ? ELightOcclusionType::MegaLightsVSM : ELightOcclusionType::MegaLights;
 	}
 
-	return bUseRaytracing ? FLightOcclusionType::Raytraced : FLightOcclusionType::Shadowmap;
+	return bUseRaytracing ? ELightOcclusionType::Raytraced : ELightOcclusionType::Shadowmap;
 }
 
-FLightOcclusionType GetLightOcclusionType(const FLightSceneInfoCompact& LightInfo, const FSceneViewFamily& ViewFamily)
+ELightOcclusionType GetLightOcclusionType(const FLightSceneInfoCompact& LightInfo, const FSceneViewFamily& ViewFamily)
 {
 	bool bUseRaytracing = false;
 #if RHI_RAYTRACING
@@ -755,10 +755,10 @@ FLightOcclusionType GetLightOcclusionType(const FLightSceneInfoCompact& LightInf
 
 	if (MegaLightsMode != EMegaLightsMode::Disabled)
 	{
-		return MegaLightsMode == EMegaLightsMode::EnabledVSM ? FLightOcclusionType::MegaLightsVSM : FLightOcclusionType::MegaLights;
+		return MegaLightsMode == EMegaLightsMode::EnabledVSM ? ELightOcclusionType::MegaLightsVSM : ELightOcclusionType::MegaLights;
 	}
 
-	return bUseRaytracing ? FLightOcclusionType::Raytraced : FLightOcclusionType::Shadowmap;
+	return bUseRaytracing ? ELightOcclusionType::Raytraced : ELightOcclusionType::Shadowmap;
 }
 
 float GetLightFadeFactor(const FSceneView& View, const FLightSceneProxy* Proxy)
@@ -1635,14 +1635,14 @@ void FDeferredShadingSceneRenderer::RenderLights(
 				const FLightSceneInfo& LightSceneInfo = *SortedLightInfo.LightSceneInfo;
 				const FLightSceneProxy& LightSceneProxy = *LightSceneInfo.Proxy;
 				const FVisibleLightInfo& VisibleLightInfo = VisibleLightInfos[LightSceneInfo.Id];
-				const FLightOcclusionType OcclusionType = GetLightOcclusionType(LightSceneProxy, ViewFamily);
+				const ELightOcclusionType OcclusionType = GetLightOcclusionType(LightSceneProxy, ViewFamily);
 
 				const bool bDrawShadows = SortedLightInfo.SortKey.Fields.bShadowed;
 				const bool bDrawLightFunction = SortedLightInfo.SortKey.Fields.bLightFunction && (!bUseLightFunctionAtlas || !SortedLightInfo.bIsCompatibleWithLightFunctionAtlas);
 				const bool bDrawPreviewIndicator = ViewFamily.EngineShowFlags.PreviewShadowsIndicator && !LightSceneInfo.IsPrecomputedLightingValid() && LightSceneProxy.HasStaticShadowing();
 				const bool bDrawHairShadow = bDrawShadows && bUseHairLighting;
 				const bool bUseHairDeepShadow = bDrawShadows && bUseHairLighting && LightSceneProxy.CastsHairStrandsDeepShadow();
-				const bool bRunHairLighting = bUseHairLighting && (bIsHairPass || OcclusionType == FLightOcclusionType::Raytraced);
+				const bool bRunHairLighting = bUseHairLighting && (bIsHairPass || OcclusionType == ELightOcclusionType::Raytraced);
 				const bool bDrawFirstPersonSelfShadow = bDrawShadows && !bIsHairPass && ShouldRenderFirstPersonSelfShadowForLight(*this, ViewFamily, Views, LightSceneInfo);
 
 				bool bUsedShadowMaskTexture = false;
@@ -1652,7 +1652,7 @@ void FDeferredShadingSceneRenderer::RenderLights(
 
 				// Raytraced shadow light for hair are handled/interleaved with regular light, as the raytraced shadow masks 
 				// for gbuffer & hair are computed by the same single pass.
-				if (bIsHairPass && OcclusionType == FLightOcclusionType::Raytraced)
+				if (bIsHairPass && OcclusionType == ELightOcclusionType::Raytraced)
 				{
 					continue;
 				}
@@ -1670,7 +1670,7 @@ void FDeferredShadingSceneRenderer::RenderLights(
 					bElideScreenShadowMask = 
 						CVarOnePassProjectionSkipScreenShadowMask.GetValueOnRenderThread() != 0 &&
 						ShadowSceneRenderer.UsePackedShadowMaskBits() &&
-						OcclusionType == FLightOcclusionType::Shadowmap &&
+						OcclusionType == ELightOcclusionType::Shadowmap &&
 						!(bDirectLighting && bDrawLightFunction) &&
 						!bDrawPreviewIndicator &&
 						SortedLightInfo.SortKey.Fields.LightType != LightType_Directional &&
@@ -1715,7 +1715,7 @@ void FDeferredShadingSceneRenderer::RenderLights(
 
 						const bool bDenoiserCompatible = !LightRequiresDenoiser(LightSceneInfo) || IScreenSpaceDenoiser::EShadowRequirements::PenumbraAndClosestOccluder == DenoiserToUse->GetShadowRequirements(View, LightSceneInfo, RayTracingConfig);
 
-						const bool bWantsBatchedShadow = OcclusionType == FLightOcclusionType::Raytraced && 
+						const bool bWantsBatchedShadow = OcclusionType == ELightOcclusionType::Raytraced &&
 							bDoShadowBatching &&
 							bDenoiserCompatible &&
 							SortedLightInfo.SortKey.Fields.bShadowed;
@@ -1785,8 +1785,8 @@ void FDeferredShadingSceneRenderer::RenderLights(
 									continue;
 								}
 
-								const FLightOcclusionType BatchOcclusionType = GetLightOcclusionType(*BatchLightSceneInfo.Proxy, ViewFamily);
-								if (BatchOcclusionType != FLightOcclusionType::Raytraced)
+								const ELightOcclusionType BatchOcclusionType = GetLightOcclusionType(*BatchLightSceneInfo.Proxy, ViewFamily);
+								if (BatchOcclusionType != ELightOcclusionType::Raytraced)
 								{
 									continue;
 								}
@@ -1943,7 +1943,7 @@ void FDeferredShadingSceneRenderer::RenderLights(
 					}
 					else 
 #endif // RHI_RAYTRACING
-					if (OcclusionType == FLightOcclusionType::Raytraced)
+					if (OcclusionType == ELightOcclusionType::Raytraced)
 					{
 						FSceneTextureParameters SceneTextureParameters = GetSceneTextureParameters(GraphBuilder, SceneTextures.UniformBuffer);
 
diff --git a/Engine/Source/Runtime/Renderer/Private/LightRendering.h b/Engine/Source/Runtime/Renderer/Private/LightRendering.h
--- a/Engine/Source/Runtime/Renderer/Private/LightRendering.h
+++ b/Engine/Source/Runtime/Renderer/Private/LightRendering.h
@@ -337,13 +337,13 @@ public:
 	static FParameters GetParameters(const FViewInfo& View, const FLightSceneInfo* LightSceneInfo, bool bBindViewUniform = true);
 };
 
-enum class FLightOcclusionType : uint8
+enum class ELightOcclusionType : uint8
 {
 	Shadowmap,
 	Raytraced,
 	MegaLights,			// Light handled through MegaLights raytracing
 	MegaLightsVSM,		// Light projection handled through MegaLights using a VSM source data
 };
-FLightOcclusionType GetLightOcclusionType(const FLightSceneProxy& Proxy, const FSceneViewFamily& ViewFamily);
-FLightOcclusionType GetLightOcclusionType(const FLightSceneInfoCompact& LightInfo, const FSceneViewFamily& ViewFamily);
 
+ELightOcclusionType GetLightOcclusionType(const FLightSceneProxy& Proxy, const FSceneViewFamily& ViewFamily);
+ELightOcclusionType GetLightOcclusionType(const FLightSceneInfoCompact& LightInfo, const FSceneViewFamily& ViewFamily);
diff --git a/Engine/Source/Runtime/Renderer/Private/ShadowRendering.cpp b/Engine/Source/Runtime/Renderer/Private/ShadowRendering.cpp
--- a/Engine/Source/Runtime/Renderer/Private/ShadowRendering.cpp
+++ b/Engine/Source/Runtime/Renderer/Private/ShadowRendering.cpp
@@ -2015,8 +2015,8 @@ bool FSceneRenderer::CheckForProjectedShadows( const FLightSceneInfo* LightScene
 
 	// If light has ray-traced occlusion enabled, then it will project some shadows. No need 
 	// for doing a lookup through shadow maps data
-	const FLightOcclusionType LightOcclusionType = GetLightOcclusionType(*LightSceneInfo->Proxy, ViewFamily);
-	if (LightOcclusionType == FLightOcclusionType::Raytraced)
+	const ELightOcclusionType LightOcclusionType = GetLightOcclusionType(*LightSceneInfo->Proxy, ViewFamily);
+	if (LightOcclusionType == ELightOcclusionType::Raytraced)
 		return true;
 
 	// Find the projected shadows cast by this light.
diff --git a/Engine/Source/Runtime/Renderer/Private/ShadowSetup.cpp b/Engine/Source/Runtime/Renderer/Private/ShadowSetup.cpp
--- a/Engine/Source/Runtime/Renderer/Private/ShadowSetup.cpp
+++ b/Engine/Source/Runtime/Renderer/Private/ShadowSetup.cpp
@@ -6171,10 +6171,10 @@ void FSceneRenderer::CreateDynamicShadows(FDynamicShadowsTaskData& TaskData)
 
 				FVisibleLightInfo& VisibleLightInfo = VisibleLightInfos[LightSceneInfo->Id];
 
-				const FLightOcclusionType OcclusionType = GetLightOcclusionType(LightSceneInfoCompact, ViewFamily);
+				const ELightOcclusionType OcclusionType = GetLightOcclusionType(LightSceneInfoCompact, ViewFamily);
 
 				// Register visible lights for allowing hair strands to cast shadow (non-directional light)
-				if (bHairStrands && OcclusionType == FLightOcclusionType::Raytraced && LightSceneInfo->Proxy->GetLightType() != LightType_Directional)
+				if (bHairStrands && OcclusionType == ELightOcclusionType::Raytraced && LightSceneInfo->Proxy->GetLightType() != LightType_Directional)
 				{
 					if (LightSceneInfoCompact.bCastStaticShadow || LightSceneInfoCompact.bCastDynamicShadow)
 					{
@@ -6191,8 +6191,8 @@ void FSceneRenderer::CreateDynamicShadows(FDynamicShadowsTaskData& TaskData)
 					}
 				}
 
-				if (OcclusionType != FLightOcclusionType::Shadowmap &&
-					OcclusionType != FLightOcclusionType::MegaLightsVSM)
+				if (OcclusionType != ELightOcclusionType::Shadowmap &&
+					OcclusionType != ELightOcclusionType::MegaLightsVSM)
 				{
 					continue;
 				}
diff --git a/Engine/Source/Runtime/Renderer/Private/Shadows/ShadowSceneRenderer.cpp b/Engine/Source/Runtime/Renderer/Private/Shadows/ShadowSceneRenderer.cpp
--- a/Engine/Source/Runtime/Renderer/Private/Shadows/ShadowSceneRenderer.cpp
+++ b/Engine/Source/Runtime/Renderer/Private/Shadows/ShadowSceneRenderer.cpp
@@ -489,13 +489,13 @@ TSharedPtr<FVirtualShadowMapPerLightCacheEntry> FShadowSceneRenderer::AddLocalLi
 	}
 
 	// TODO: This is remarkably slow and shouldn't really need to be evaluated multiple times
-	FLightOcclusionType OcclusionType = GetLightOcclusionType(*LightSceneInfo->Proxy, SceneRenderer.ViewFamily);
+	ELightOcclusionType OcclusionType = GetLightOcclusionType(*LightSceneInfo->Proxy, SceneRenderer.ViewFamily);
 	// Depending on which type of projection we're going to use, mark that we need to associated path for later
-	if (OcclusionType == FLightOcclusionType::Shadowmap)
+	if (OcclusionType == ELightOcclusionType::Shadowmap)
 	{
 		bNeedVSMOnePassProjection = true;
 	}
-	else if (OcclusionType == FLightOcclusionType::MegaLightsVSM)
+	else if (OcclusionType == ELightOcclusionType::MegaLightsVSM)
 	{
 		bNeedMegaLightsProjection = true;
 	}
@@ -598,11 +598,11 @@ void FShadowSceneRenderer::AddDirectionalLightShadow(FLightSceneInfo& LightScene
 		}
 	}
 
-	FLightOcclusionType OcclusionType = GetLightOcclusionType(*LightSceneInfo.Proxy, SceneRenderer.ViewFamily);
+	ELightOcclusionType OcclusionType = GetLightOcclusionType(*LightSceneInfo.Proxy, SceneRenderer.ViewFamily);
 	// Mark that we need MegaLights projection pass if this light uses it
 	// We don't need to explicitly mark VSM projection here since directional lights don't go through one pass projection,
 	// and thus will naturally go down the path that will render into the screen shadow mask.
-	if (OcclusionType == FLightOcclusionType::MegaLightsVSM)
+	if (OcclusionType == ELightOcclusionType::MegaLightsVSM)
 	{
 		bNeedMegaLightsProjection = true;
 	}
diff --git a/Engine/Source/Runtime/Renderer/Private/TranslucentLighting.cpp b/Engine/Source/Runtime/Renderer/Private/TranslucentLighting.cpp
--- a/Engine/Source/Runtime/Renderer/Private/TranslucentLighting.cpp
+++ b/Engine/Source/Runtime/Renderer/Private/TranslucentLighting.cpp
@@ -2688,8 +2688,8 @@ void FDeferredShadingSceneRenderer::RenderTranslucencyLightingVolume(FRDGBuilder
 			}
 
 			const bool bDrawShadows = SortedLightInfo.SortKey.Fields.bShadowed;
-			const FLightOcclusionType OcclusionType = GetLightOcclusionType(LightSceneProxy, ViewFamily);
-			const bool bSupportShadowMaps = bDrawShadows && OcclusionType == FLightOcclusionType::Shadowmap;
+			const ELightOcclusionType OcclusionType = GetLightOcclusionType(LightSceneProxy, ViewFamily);
+			const bool bSupportShadowMaps = bDrawShadows && OcclusionType == ELightOcclusionType::Shadowmap;
 
 			// Collect all the light injection data
 			CollectLightForTranslucencyLightingVolumeInjection(&LightSceneInfo, bSupportShadowMaps, Collector);
diff --git a/Engine/Source/Runtime/Renderer/Private/VolumetricFog.cpp b/Engine/Source/Runtime/Renderer/Private/VolumetricFog.cpp
--- a/Engine/Source/Runtime/Renderer/Private/VolumetricFog.cpp
+++ b/Engine/Source/Runtime/Renderer/Private/VolumetricFog.cpp
@@ -619,7 +619,7 @@ const FProjectedShadowInfo* GetShadowForInjectionIntoVolumetricFog(const FVisibl
 
 bool LightHasRayTracedShadows(const FLightSceneInfo* LightSceneInfo, const FSceneViewFamily& ViewFamily)
 {
-	return GetLightOcclusionType(*LightSceneInfo->Proxy, ViewFamily) == FLightOcclusionType::Raytraced && GVolumetricFogInjectRaytracedLights;
+	return GetLightOcclusionType(*LightSceneInfo->Proxy, ViewFamily) == ELightOcclusionType::Raytraced && GVolumetricFogInjectRaytracedLights;
 }
 
 bool LightNeedsSeparateInjectionIntoVolumetricFogForOpaqueShadow(const FViewInfo& View, const FLightSceneInfo* LightSceneInfo, const FVisibleLightInfo& VisibleLightInfo, const FScene& InScene)
