Match mobile primary and secondary upscaler setting with Desktop Renderer.

diff --git a/Engine/Source/Runtime/Renderer/Private/PostProcess/PostProcessing.cpp b/Engine/Source/Runtime/Renderer/Private/PostProcess/PostProcessing.cpp
--- a/Engine/Source/Runtime/Renderer/Private/PostProcess/PostProcessing.cpp
+++ b/Engine/Source/Runtime/Renderer/Private/PostProcess/PostProcessing.cpp
@@ -2506,8 +2506,9 @@ void AddMobilePostProcessingPasses(FRDGBuilder& GraphBuilder, FScene* Scene, con
 
 	bool bUseHighResolutionScreenshotMask = IsHighResolutionScreenshotMaskEnabled(View);
 
-	bool bShouldPrimaryUpscale = (View.PrimaryScreenPercentageMethod == EPrimaryScreenPercentageMethod::SpatialUpscale && View.UnscaledViewRect != View.ViewRect) || View.LensDistortionLUT.IsEnabled();
-	bShouldPrimaryUpscale |= View.Family->GetPrimarySpatialUpscalerInterface() != nullptr;
+	// Set bShouldPrimaryUpscale like Desktop Renderer
+	// compare size with Secondary, not Unscaled
+	bool bShouldPrimaryUpscale = (View.PrimaryScreenPercentageMethod == EPrimaryScreenPercentageMethod::SpatialUpscale && View.GetSecondaryViewRectSize() != View.ViewRect.Size()) || View.LensDistortionLUT.IsEnabled();
 
 	PassSequence.SetEnabled(EPass::Tonemap, bUseToneMapper);
 	PassSequence.SetEnabled(EPass::HighResolutionScreenshotMask, bUseHighResolutionScreenshotMask);
@@ -2525,7 +2526,10 @@ void AddMobilePostProcessingPasses(FRDGBuilder& GraphBuilder, FScene* Scene, con
 	PassSequence.SetEnabled(EPass::DebugPrimitive, FSceneRenderer::ShouldCompositeDebugPrimitivesInPostProcess(View));
 #endif
 	PassSequence.SetEnabled(EPass::PrimaryUpscale, bShouldPrimaryUpscale);
-	PassSequence.SetEnabled(EPass::SecondaryUpscale, View.Family->GetSecondarySpatialUpscalerInterface() != nullptr);
+	// Match Desktop Renderer, calling View.RequiresSecondaryUpscale()
+	// Both Desktop and Mobile Renderer call Third-party if enabled (no matter what scaling).
+	// Is it really necessary to call the Secondary Spatial if the scaling is 100% ?
+	PassSequence.SetEnabled(EPass::SecondaryUpscale, View.RequiresSecondaryUpscale() || View.Family->GetSecondarySpatialUpscalerInterface() != nullptr);
 
 	PassSequence.SetEnabled(EPass::Visualize, View.Family->EngineShowFlags.ShaderComplexity);
 	PassSequence.SetEnabled(EPass::VisualizeLightGrid, ShouldVisualizeLightGrid(View.Family->GetShaderPlatform()));
@@ -3278,7 +3282,7 @@ void AddMobilePostProcessingPasses(FRDGBuilder& GraphBuilder, FScene* Scene, con
 	{
 		ISpatialUpscaler::FInputs PassInputs;
 		PassSequence.AcceptOverrideIfLastPass(EPass::PrimaryUpscale, PassInputs.OverrideOutput);
-		PassInputs.Stage = EUpscaleStage::PrimaryToOutput;
+		PassInputs.Stage = PassSequence.IsEnabled(EPass::SecondaryUpscale) ? EUpscaleStage::PrimaryToSecondary : EUpscaleStage::PrimaryToOutput;
 		PassInputs.SceneColor = SceneColor;
 		PassInputs.OverrideOutput.LoadAction = View.IsFirstInFamily() ? ERenderTargetLoadAction::EClear : ERenderTargetLoadAction::ELoad;
 
@@ -3299,7 +3303,7 @@ void AddMobilePostProcessingPasses(FRDGBuilder& GraphBuilder, FScene* Scene, con
 			}
 			else
 			{
-				check(SceneColor.ViewRect.Size() == View.UnscaledViewRect.Size());
+				check(SceneColor.ViewRect.Size() == View.GetSecondaryViewRectSize());
 			}
 		}
 		else
