Move PrepareViewStateForVisibility back to after ViewExtensions' PreRenderViewFamily_RenderThread and PreRenderView_RenderThread in FSceneRenderer::OnRenderBegin.

PrepareViewStateForVisibility must be performed after those ViewExtension functions to ensure it picks up any adjustments to the projection matrices, such as HMD late update (see FDefaultXRCamera::PreRenderView_RenderThread).

diff --git a/Engine/Source/Runtime/Renderer/Private/SceneRendering.cpp b/Engine/Source/Runtime/Renderer/Private/SceneRendering.cpp
--- a/Engine/Source/Runtime/Renderer/Private/SceneRendering.cpp
+++ b/Engine/Source/Runtime/Renderer/Private/SceneRendering.cpp
@@ -3932,11 +3932,11 @@ IVisibilityTaskData* FSceneRenderer::OnRenderBegin(FRDGBuilder& GraphBuilder, co
 	UE::Tasks::FTaskEvent GPUSceneUpdateTaskPrerequisites{ UE_SOURCE_LOCATION };
 	SceneUpdateParameters.GPUSceneUpdateTaskPrerequisites = GPUSceneUpdateTaskPrerequisites;
 
-	UE::Tasks::FTask PrepareViewStateTask;
+	UE::Tasks::FTask PrepareSceneTexturesConfigTask;
 
 	if (SceneUpdateInputs)
 	{
-		PrepareViewStateTask = UE::Tasks::Launch(UE_SOURCE_LOCATION, [SceneUpdateInputs]
+		PrepareSceneTexturesConfigTask = UE::Tasks::Launch(UE_SOURCE_LOCATION, [SceneUpdateInputs]
 		{
 			TRACE_CPUPROFILER_EVENT_SCOPE(PrepareViewRects);
 			FTaskTagScope TagScope(ETaskTag::EParallelRenderingThread);
@@ -3959,8 +3959,6 @@ IVisibilityTaskData* FSceneRenderer::OnRenderBegin(FRDGBuilder& GraphBuilder, co
 					CustomRenderPass.ViewFamily.SceneTexturesConfig.NumSamples = 1;
 					CustomRenderPass.ViewFamily.SceneTexturesConfig.EditorPrimitiveNumSamples = 1;
 				}
-	
-				Renderer->PrepareViewStateForVisibility(SceneTexturesConfig);
 			}
 
 		}, UE::Tasks::ETaskPriority::Normal, bIsMobilePlatform ? UE::Tasks::EExtendedTaskPriority::Inline : UE::Tasks::EExtendedTaskPriority::None);
@@ -3968,7 +3966,7 @@ IVisibilityTaskData* FSceneRenderer::OnRenderBegin(FRDGBuilder& GraphBuilder, co
 
 	SceneUpdateParameters.Callbacks.PostStaticMeshUpdate = [&] (const UE::Tasks::FTask& StaticMeshUpdateTask)
 	{
-		PrepareViewStateTask.Wait();
+		PrepareSceneTexturesConfigTask.Wait();
 
 #if RHI_RAYTRACING
 		if (SceneUpdateInputs)
@@ -3998,6 +3996,15 @@ IVisibilityTaskData* FSceneRenderer::OnRenderBegin(FRDGBuilder& GraphBuilder, co
 			}
 		}
 
+		if (SceneUpdateInputs)
+		{
+			for (FSceneRenderer* Renderer : SceneUpdateInputs->Renderers)
+			{
+				const FSceneTexturesConfig& SceneTexturesConfig = Renderer->GetActiveSceneTexturesConfig();
+				Renderer->PrepareViewStateForVisibility(SceneTexturesConfig);
+			}
+		}
+
 		if (ViewFamily.EngineShowFlags.LensDistortion && FPaniniProjectionConfig::IsEnabledByCVars())
 		{
 			const FPaniniProjectionConfig PaniniProjection = FPaniniProjectionConfig::ReadCVars();
