[Light grid] Clamp diffuse and specular scale reciprocals to encoding format min so we don't get division-by-zero later in shaders

#tests tested in editor

diff --git a/Engine/Source/Runtime/Renderer/Private/LightGridInjection.cpp b/Engine/Source/Runtime/Renderer/Private/LightGridInjection.cpp
--- a/Engine/Source/Runtime/Renderer/Private/LightGridInjection.cpp
+++ b/Engine/Source/Runtime/Renderer/Private/LightGridInjection.cpp
@@ -480,8 +480,8 @@ static void PackLightData(
 
 	// Pack specular scale and IES profile index.
 	// Simple light specular and diffuse scales can be greater than one. If that happens, we encode reciprocals instead
-	const float SpecularScale = SimpleLight.SpecularScale > 1.0f ? 1.0f / SimpleLight.SpecularScale : SimpleLight.SpecularScale;
-	const float DiffuseScale  = SimpleLight.DiffuseScale > 1.0f ? 1.0f / SimpleLight.DiffuseScale : SimpleLight.DiffuseScale;
+	const float SpecularScale = SimpleLight.SpecularScale > 1.0f ? FMath::Max(1.0f / SimpleLight.SpecularScale, 1.f / 1023.f) : SimpleLight.SpecularScale;
+	const float DiffuseScale  = SimpleLight.DiffuseScale > 1.0f ? FMath::Max(1.0f / SimpleLight.DiffuseScale, 1.f / 1023.f) : SimpleLight.DiffuseScale;
 	const float IESAtlasIndex = INDEX_NONE;
 
 	// Offset IESAtlasIndex here in order to preserve INDEX_NONE = -1 after encoding
