Fixed: Incorrect render target setup for pre-caching translucent PSOs on platforms that require DepthAux

diff --git a/Engine/Source/Runtime/Renderer/Private/MobileBasePass.cpp b/Engine/Source/Runtime/Renderer/Private/MobileBasePass.cpp
--- a/Engine/Source/Runtime/Renderer/Private/MobileBasePass.cpp
+++ b/Engine/Source/Runtime/Renderer/Private/MobileBasePass.cpp
@@ -1054,7 +1054,8 @@ void FMobileBasePassMeshProcessor::CollectPSOInitializers(const FSceneTexturesCo
 
 	FGraphicsPipelineRenderTargetsInfo RenderTargetsInfo;
 	RenderTargetsInfo.NumSamples = 1;
-	if (bTranslucentBasePass)
+	if (bTranslucentBasePass && 
+		!SceneTexturesConfig.bRequiresDepthAux) // Really need to detect if we use single pass rendering
 	{		
 		EPixelFormat SceneColorFormat = SceneTexturesConfig.ColorFormat;
 		ETextureCreateFlags SceneColorCreateFlags = SceneTexturesConfig.ColorCreateFlags;
