Fix incorrect CVar naming.

diff --git a/Engine/Source/Runtime/Renderer/Private/VolumetricCloudRendering.cpp b/Engine/Source/Runtime/Renderer/Private/VolumetricCloudRendering.cpp
--- a/Engine/Source/Runtime/Renderer/Private/VolumetricCloudRendering.cpp
+++ b/Engine/Source/Runtime/Renderer/Private/VolumetricCloudRendering.cpp
@@ -273,8 +273,8 @@ static TAutoConsoleVariable<int32> CVarVolumetricCloudAllowAnalyticDerivatives(
 
 static bool ProjectSupportVolumetricCloud(const EShaderPlatform ShaderPlatform)
 {
-	static FShaderPlatformCachedIniValue<int32> CVarSubstrateBlendableGBufferPlatform(TEXT("r.VolumetricCloud.Support"));
-	return CVarSubstrateBlendableGBufferPlatform.Get(ShaderPlatform) > 0;
+	static FShaderPlatformCachedIniValue<int32> CVar(TEXT("r.VolumetricCloud.Support"));
+	return CVar.Get(ShaderPlatform) > 0;
 }
 
 static bool ShouldUseComputeForCloudTracing(const FStaticFeatureLevel InFeatureLevel)
