SMAA - Add Luminance mode

diff --git a/Engine/Shaders/Private/SMAA/SMAAEdgeDetectionShader.usf b/Engine/Shaders/Private/SMAA/SMAAEdgeDetectionShader.usf
--- a/Engine/Shaders/Private/SMAA/SMAAEdgeDetectionShader.usf
+++ b/Engine/Shaders/Private/SMAA/SMAAEdgeDetectionShader.usf
@@ -9,6 +9,8 @@ struct FSMAAEdgeDetectionScreenVertexOutput
 	float4 Position : SV_POSITION;
 };
 
+int EdgeMode;
+
 /*-----------------------------------------------------------------------------
 	vertex shader
 -----------------------------------------------------------------------------*/
@@ -37,7 +39,15 @@ struct FSMAAEdgeDetectionScreenVertexOutput
 		FSMAAEdgeDetectionScreenVertexOutput Input,
 		out float2 OutColor : SV_Target0)
 	{
-		OutColor.rg = SMAAColorEdgeDetectionPS(Input.UV, Input.Offset, Input_Texture);
+		BRANCH
+		if (EdgeMode == 0) // Color
+		{
+			OutColor.rg = SMAAColorEdgeDetectionPS(Input.UV, Input.Offset, Input_Texture);
+		}
+		else // if (EdgeMode == 1) // Luma (uses Linear Color as input texture)
+		{
+			OutColor.rg = SMAALumaEdgeDetectionPS(Input.UV, Input.Offset, Input_Texture);
+		}
 	}
 
 #endif // MainPS
diff --git a/Engine/Source/Runtime/Renderer/Private/PostProcess/SubpixelMorphologicalAA.cpp b/Engine/Source/Runtime/Renderer/Private/PostProcess/SubpixelMorphologicalAA.cpp
--- a/Engine/Source/Runtime/Renderer/Private/PostProcess/SubpixelMorphologicalAA.cpp
+++ b/Engine/Source/Runtime/Renderer/Private/PostProcess/SubpixelMorphologicalAA.cpp
@@ -27,6 +27,21 @@ TAutoConsoleVariable<int32> CVarSMAADebugVisualization(
 	TEXT(" 3: Scene Color + Smoothened edges highlights \n"),
 	ECVF_RenderThreadSafe);
 
+TAutoConsoleVariable<int32> CVarSMAAEdgeDetectionMode(
+	TEXT("r.SMAA.EdgeMode"), 0,
+	TEXT("Edge detection mode used.\n")
+	TEXT(" 0: Color \n")
+	TEXT(" 1: Luminance \n"),
+	ECVF_RenderThreadSafe);
+
+enum class ESMAAEdgeMode : uint8
+{
+	Color,
+	Luminance,
+
+	Count
+};
+
 #define DECLARE_SMAA_QUALITY_PERMUTATIONS               \
 class FQuality : SHADER_PERMUTATION_INT("SMAA_QUALITY", 4);
 
@@ -49,6 +64,7 @@ public:
 		SHADER_PARAMETER_STRUCT(FScreenPassTextureInput, Input)
 		SHADER_PARAMETER_SAMPLER(SamplerState, PointSampler)
 		SHADER_PARAMETER(FVector4f, RTMetrics)
+		SHADER_PARAMETER(int32, EdgeMode)
 		RENDER_TARGET_BINDING_SLOTS()
 	END_SHADER_PARAMETER_STRUCT()
 };
@@ -185,6 +201,7 @@ void AddSMAAEdgeDetectionPass(FRDGBuilder& GraphBuilder, const FSceneView& InSce
 	PassParameters->Input = GetScreenPassTextureInput(SceneColorTexture, PointClampSampler);
 	PassParameters->RTMetrics = RTMetrics;
 	PassParameters->PointSampler = PointClampSampler;
+	PassParameters->EdgeMode = FMath::Clamp(CVarSMAAEdgeDetectionMode->GetInt(), 0, (int)ESMAAEdgeMode::Count - 1);
 
 	FSMAAEdgeDetectionPS::FPermutationDomain PixelPermutationVector;
 	SetupSMAAQualityPermutations<FSMAAEdgeDetectionPS>(PixelPermutationVector, Quality);
