Fix incorrectly reported overflow with megalights active
Only run VSM coarse page marking once when there are multiple views as it is view-independent.

diff --git a/Engine/Shaders/Private/VirtualShadowMaps/VirtualShadowMapProjection.usf b/Engine/Shaders/Private/VirtualShadowMaps/VirtualShadowMapProjection.usf
--- a/Engine/Shaders/Private/VirtualShadowMaps/VirtualShadowMapProjection.usf
+++ b/Engine/Shaders/Private/VirtualShadowMaps/VirtualShadowMapProjection.usf
@@ -813,23 +813,27 @@ void VirtualShadowMapProjection(
 
 #if ONE_PASS_PROJECTION
 	const FCulledLightsGridHeader LightGridHeader = VirtualShadowMapGetLightsGridHeader(LocalPixelPos, SceneDepth);	
-	// We can only handle so many lights in our output encoding right now, so no purpose in computing more
-	uint LightCount = min(GetPackedShadowMaskMaxLightCount(), LightGridHeader.NumLights);
-	bool bLightCountOverflow = GetPackedShadowMaskMaxLightCount() < LightGridHeader.NumLights;
-	StatsBufferInterlockedEnableFlags(VSM_STAT_OVERFLOW_FLAGS, VSM_STAT_OVERFLOW_FLAG_OPP_MAX_LIGHTS, bLightCountOverflow, true);
+	// Allow one extra loop iteration to detect potential overflow
+	const uint LoopCount = min(GetPackedShadowMaskMaxLightCount() + 1, LightGridHeader.NumLights);
 
 	uint4 ShadowMask = InitializePackedShadowMask();
-
 	LOOP
-	for (uint Index = 0; Index < LightCount; ++Index)
+	for (uint Index = 0; Index < LoopCount; ++Index)
 	{
 		const FLocalLightData LightData = VirtualShadowMapGetLocalLightData(LightGridHeader, Index);
-		int VirtualShadowMapId = LightData.Internal.VirtualShadowMapId;
+		
+		// Megalights are sorted to the end (See LightSceneInfo.h SortyKey) so if we hit one we're done
+		if (UnpackIsHandledByMegaLights(LightData))
+		{
+			break;
+		}
 
-		if (VirtualShadowMapId != INDEX_NONE)
+		if (Index < GetPackedShadowMaskMaxLightCount())
 		{
-			FLightShaderParameters Light = ConvertFromLocal( LightData );
+			// VSM pruned light grid, so this should always be valid
+			int VirtualShadowMapId = LightData.Internal.VirtualShadowMapId;
 
+			FLightShaderParameters Light = ConvertFromLocal(LightData);
 			FVirtualShadowMapSampleResult Result = ProjectLight(
 				VirtualShadowMapId,
 				Light,
@@ -848,6 +852,11 @@ void VirtualShadowMapProjection(
 			OutputVisualize( VirtualShadowMapId, PixelPos - View.ViewRectMin.xy, TranslatedWorldPosition, Result );
 #endif // VISUALIZE_OUTPUT
 		}
+		else
+		{
+			// Record overflow
+			StatsBufferInterlockedEnableFlags(VSM_STAT_OVERFLOW_FLAGS, VSM_STAT_OVERFLOW_FLAG_OPP_MAX_LIGHTS, true, true);
+		}
 	}
 
 	OutShadowMaskBits[PixelPos] = ShadowMask;
diff --git a/Engine/Source/Runtime/Renderer/Private/VirtualShadowMaps/VirtualShadowMapArray.cpp b/Engine/Source/Runtime/Renderer/Private/VirtualShadowMaps/VirtualShadowMapArray.cpp
--- a/Engine/Source/Runtime/Renderer/Private/VirtualShadowMaps/VirtualShadowMapArray.cpp
+++ b/Engine/Source/Runtime/Renderer/Private/VirtualShadowMaps/VirtualShadowMapArray.cpp
@@ -2254,6 +2254,31 @@ void FVirtualShadowMapArray::BuildPageAllocations(
 	NanitePerformanceFeedbackRDG = GraphBuilder.CreateBuffer(FRDGBufferDesc::CreateStructuredDesc(sizeof(uint32), NanitePerformanceFeedbackSize), TEXT("NanitePerformanceFeedback"));
 	AddClearUAVPass(GraphBuilder, GraphBuilder.CreateUAV(NanitePerformanceFeedbackRDG), 0);
 
+	// Mark coarse pages (view-independent)
+	// NOTE: Must do this *first*. In the case where bIncludeNonNaniteGeometry is false we need to ensure that the request
+	// can be over-written by any pixel pages that *do* want Non-Nanite geometry. We avoid writing with atomics since that
+	// is much slower.
+	// Because of this we also cannot overlap this pass with the following ones.	
+	// Note: always run this pass such that the distant lights may be marked if need be
+	{
+		FMarkCoarsePagesCS::FParameters* PassParameters = GraphBuilder.AllocParameters< FMarkCoarsePagesCS::FParameters >();
+		PassParameters->VirtualShadowMap = GetUncachedUniformBuffer(GraphBuilder);
+		PassParameters->OutPageRequestFlags = GraphBuilder.CreateUAV(PageRequestFlagsRDG);
+		PassParameters->OutPageReceiverMasks = GraphBuilder.CreateUAV(PageReceiverMasksRDG);
+		PassParameters->bMarkCoarsePagesLocal = CVarMarkCoarsePagesLocal.GetValueOnRenderThread();
+		PassParameters->bIncludeNonNaniteGeometry = CVarCoarsePagesIncludeNonNanite.GetValueOnRenderThread();
+
+		auto ComputeShader = GetGlobalShaderMap(Scene.GetFeatureLevel())->GetShader<FMarkCoarsePagesCS>();
+
+		FComputeShaderUtils::AddPass(
+			GraphBuilder,
+			RDG_EVENT_NAME("MarkCoarsePages"),
+			ComputeShader,
+			PassParameters,
+			FIntVector(FMath::DivideAndRoundUp(uint32(GetNumShadowMaps()), FMarkCoarsePagesCS::DefaultCSGroupX), 1, 1)
+		);
+	}
+
 	for (int32 ViewIndex = 0; ViewIndex < Views.Num(); ++ViewIndex)
 	{
 		RDG_EVENT_SCOPE_CONDITIONAL(GraphBuilder, Views.Num() > 1, "View%d", ViewIndex);
@@ -2276,32 +2301,6 @@ void FVirtualShadowMapArray::BuildPageAllocations(
 		FRDGBufferRef ScreenSpaceGridBoundsRDG = nullptr;
 			
 		{
-			// Mark coarse pages
-			// NOTE: Must do this *first*. In the case where bIncludeNonNaniteGeometry is false we need to ensure that the request
-			// can be over-written by any pixel pages that *do* want Non-Nanite geometry. We avoid writing with atomics since that
-			// is much slower.
-			// Because of this we also cannot overlap this pass with the following ones.
-			bool bMarkCoarsePagesLocal = CVarMarkCoarsePagesLocal.GetValueOnRenderThread() != 0;
-			// Note: always run this pass such that the distant lights may be marked if need be
-			{
-				FMarkCoarsePagesCS::FParameters* PassParameters = GraphBuilder.AllocParameters< FMarkCoarsePagesCS::FParameters >();
-				PassParameters->VirtualShadowMap = GetUncachedUniformBuffer(GraphBuilder);
-				PassParameters->OutPageRequestFlags = GraphBuilder.CreateUAV(PageRequestFlagsRDG);
-				PassParameters->OutPageReceiverMasks = GraphBuilder.CreateUAV(PageReceiverMasksRDG);
-				PassParameters->bMarkCoarsePagesLocal = bMarkCoarsePagesLocal ? 1 : 0;
-				PassParameters->bIncludeNonNaniteGeometry = CVarCoarsePagesIncludeNonNanite.GetValueOnRenderThread();
-
-				auto ComputeShader = View.ShaderMap->GetShader<FMarkCoarsePagesCS>();
-
-				FComputeShaderUtils::AddPass(
-					GraphBuilder,
-					RDG_EVENT_NAME("MarkCoarsePages"),
-					ComputeShader,
-					PassParameters,
-					FIntVector(FMath::DivideAndRoundUp(uint32(GetNumShadowMaps()), FMarkCoarsePagesCS::DefaultCSGroupX), 1, 1)
-				);
-			}
-
 			// Prune light grid to remove lights without VSMs
 			{
 				const bool bLightGridUses16BitBuffers = LightGridUses16BitBuffers(View.GetShaderPlatform());
