Fix RHI validation error with rect light, potential causing GPU hang with MegaLights.

diff --git a/Engine/Source/Runtime/Renderer/Private/RectLightTextureManager.cpp b/Engine/Source/Runtime/Renderer/Private/RectLightTextureManager.cpp
--- a/Engine/Source/Runtime/Renderer/Private/RectLightTextureManager.cpp
+++ b/Engine/Source/Runtime/Renderer/Private/RectLightTextureManager.cpp
@@ -1887,6 +1887,12 @@ void UpdateAtlasTexture(FRDGBuilder& GraphBuilder, const ERHIFeatureLevel::Type
 		AddSlotsPass(GraphBuilder, ShaderMap, GRectLightTextureManager.AtlasLayout.SourceTextureMIPBias, RefreshSlots, AtlasTexture);
 		FilterSlotsPass(GraphBuilder, ShaderMap, RefreshSlots, AtlasTexture);
 	}
+
+	// If the atlas was registered into RDG make sure it gets transitioned to external access on all pipes.
+	if (FRDGTexture* AtlasTexture = GraphBuilder.FindExternalTexture(GRectLightTextureManager.AtlasTexture))
+	{
+		GraphBuilder.UseExternalAccessMode(AtlasTexture, ERHIAccess::SRVMask, ERHIPipeline::All);
+	}
 }
 
 void AddDebugPass(FRDGBuilder& GraphBuilder, const FViewInfo& View, FRDGTextureRef OutputTexture)
