Simplify verbose code to get the size of ElementType

diff --git a/Engine/Plugins/Animation/ControlRig/Source/ControlRig/Public/ControlRigOverride.h b/Engine/Plugins/Animation/ControlRig/Source/ControlRig/Public/ControlRigOverride.h
--- a/Engine/Plugins/Animation/ControlRig/Source/ControlRig/Public/ControlRigOverride.h
+++ b/Engine/Plugins/Animation/ControlRig/Source/ControlRig/Public/ControlRigOverride.h
@@ -365,8 +365,7 @@ private:
 	{
 		const uint32 PathHash = CityHash32(
 			(const char*)InPath.GetCharArray().GetData(),
-			TRemoveReference<decltype(InPath.GetCharArray())>::Type::GetTypeSize() *
-			InPath.GetCharArray().Num());
+			sizeof(InPath[0]) * InPath.GetCharArray().Num());
 		return HashCombineFast(GetTypeHash(InSubjectKey), PathHash);
 	}
 
diff --git a/Engine/Plugins/Experimental/LearningAgents/Source/Learning/Private/LearningAction.cpp b/Engine/Plugins/Experimental/LearningAgents/Source/Learning/Private/LearningAction.cpp
--- a/Engine/Plugins/Experimental/LearningAgents/Source/Learning/Private/LearningAction.cpp
+++ b/Engine/Plugins/Experimental/LearningAgents/Source/Learning/Private/LearningAction.cpp
@@ -1472,8 +1472,7 @@ namespace UE::Learning::Action
 			const FString NameString = Name.ToString().ToLower();
 			return (int32)CityHash32(
 				(const char*)NameString.GetCharArray().GetData(),
-				TRemoveReference<decltype(NameString.GetCharArray())>::Type::GetTypeSize() *
-				NameString.GetCharArray().Num());
+				sizeof(NameString[0]) * NameString.GetCharArray().Num());
 		}
 
 		static inline int32 HashInt(const int32 Int)
diff --git a/Engine/Plugins/Experimental/LearningAgents/Source/Learning/Private/LearningObservation.cpp b/Engine/Plugins/Experimental/LearningAgents/Source/Learning/Private/LearningObservation.cpp
--- a/Engine/Plugins/Experimental/LearningAgents/Source/Learning/Private/LearningObservation.cpp
+++ b/Engine/Plugins/Experimental/LearningAgents/Source/Learning/Private/LearningObservation.cpp
@@ -1130,8 +1130,7 @@ namespace UE::Learning::Observation
 			const FString NameString = Name.ToString().ToLower();
 			return (int32)CityHash32(
 				(const char*)NameString.GetCharArray().GetData(), 
-				TRemoveReference<decltype(NameString.GetCharArray())>::Type::GetTypeSize() *
-				NameString.GetCharArray().Num());
+				sizeof(NameString[0]) * NameString.GetCharArray().Num());
 		}
 
 		static inline int32 HashInt(const int32 Int)
diff --git a/Engine/Plugins/Online/OnlineSubsystem/Source/Public/OnlineSubsystemTypes.h b/Engine/Plugins/Online/OnlineSubsystem/Source/Public/OnlineSubsystemTypes.h
--- a/Engine/Plugins/Online/OnlineSubsystem/Source/Public/OnlineSubsystemTypes.h
+++ b/Engine/Plugins/Online/OnlineSubsystem/Source/Public/OnlineSubsystemTypes.h
@@ -478,7 +478,7 @@ public:
 
 	virtual int32 GetSize() const override
 	{
-		return TRemoveReference<decltype(UniqueNetIdStr.GetCharArray())>::Type::GetTypeSize() * UniqueNetIdStr.GetCharArray().Num();
+		return sizeof(UniqueNetIdStr[0]) * UniqueNetIdStr.GetCharArray().Num();
 	}
 
 	virtual bool IsValid() const override
diff --git a/Engine/Source/Runtime/Renderer/Private/InstanceCulling/InstanceCullingContext.cpp b/Engine/Source/Runtime/Renderer/Private/InstanceCulling/InstanceCullingContext.cpp
--- a/Engine/Source/Runtime/Renderer/Private/InstanceCulling/InstanceCullingContext.cpp
+++ b/Engine/Source/Runtime/Renderer/Private/InstanceCulling/InstanceCullingContext.cpp
@@ -1066,10 +1066,10 @@ FInstanceCullingDeferredContext *FInstanceCullingContext::CreateDeferredContext(
 #define INST_CULL_CREATE_STRUCT_BUFF_ARGS_BIN_INDEX(ArrayName) \
 	GraphBuilder, \
 	TEXT("InstanceCulling.") TEXT(#ArrayName), \
-	TRemoveReference<decltype(DeferredContext->ArrayName[BinIndex])>::Type::GetTypeSize(), \
+	sizeof(DeferredContext->ArrayName[BinIndex][0]), \
 	INST_CULL_CALLBACK_BIN_INDEX(DeferredContext->ArrayName[BinIndex].Num()), \
 	INST_CULL_CALLBACK_BIN_INDEX(DeferredContext->ArrayName[BinIndex].GetData()), \
-	INST_CULL_CALLBACK_BIN_INDEX(DeferredContext->ArrayName[BinIndex].Num() * TRemoveReference<decltype(DeferredContext->ArrayName[BinIndex])>::Type::GetTypeSize())
+	INST_CULL_CALLBACK_BIN_INDEX(DeferredContext->ArrayName[BinIndex].Num() * sizeof(DeferredContext->ArrayName[BinIndex][0]))
 
 	const ERHIFeatureLevel::Type FeatureLevel = GPUScene.GetFeatureLevel();
 	const EShaderPlatform ShaderPlatform = GPUScene.GetShaderPlatform();
