Missing AliasbleFormat PF_X24_G8 on TextureDesc for Composite.PrimitivesDepth

diff --git a/Engine/Source/Runtime/Renderer/Private/PostProcess/PostProcessCompositePrimitivesCommon.cpp b/Engine/Source/Runtime/Renderer/Private/PostProcess/PostProcessCompositePrimitivesCommon.cpp
--- a/Engine/Source/Runtime/Renderer/Private/PostProcess/PostProcessCompositePrimitivesCommon.cpp
+++ b/Engine/Source/Runtime/Renderer/Private/PostProcess/PostProcessCompositePrimitivesCommon.cpp
@@ -39,13 +39,14 @@ const FViewInfo* CreateCompositePrimitiveView(const FViewInfo& ParentView, FIntR
 
 FRDGTextureRef CreateCompositeDepthTexture(FRDGBuilder& GraphBuilder, FIntPoint Extent, uint32 NumMSAASamples)
 {
-	const FRDGTextureDesc DepthDesc = FRDGTextureDesc::Create2D(
+	FRDGTextureDesc DepthDesc = FRDGTextureDesc::Create2D(
 		Extent,
 		PF_DepthStencil,
 		FClearValueBinding::DepthFar,
 		TexCreate_ShaderResource | TexCreate_DepthStencilTargetable,
 		1,
 		NumMSAASamples);
+	DepthDesc.AliasableFormats.Add(PF_X24_G8);
 
 	return GraphBuilder.CreateTexture(DepthDesc, TEXT("Composite.PrimitivesDepth"));
 }
@@ -314,4 +315,4 @@ void PopulateDepthPass(FRDGBuilder& GraphBuilder,
 		TStaticDepthStencilState<true, CF_Always>::GetRHI());
 }
 
-#endif //UE_ENABLE_DEBUG_DRAWING
\ No newline at end of file
+#endif //UE_ENABLE_DEBUG_DRAWING
