Path Tracer: Fix upscaling method being selected incorrectly when using the path tracer

Setting View.AntiAliasingMethod to none was no longer sufficient, and could also be bypassed when the ScreenPercentage show flag was false.

#rnx

diff --git a/Engine/Source/Runtime/Renderer/Private/SceneRendering.cpp b/Engine/Source/Runtime/Renderer/Private/SceneRendering.cpp
--- a/Engine/Source/Runtime/Renderer/Private/SceneRendering.cpp
+++ b/Engine/Source/Runtime/Renderer/Private/SceneRendering.cpp
@@ -3256,22 +3256,6 @@ void FSceneRenderer::PrepareViewRectsForRendering()
 		DynamicResolutionFractions = ViewFamily.ScreenPercentageInterface->GetResolutionFractions_RenderThread();
 	}
 
-	// If not supporting screen percentage, bypass all computation.
-	if (!ViewFamily.EngineShowFlags.ScreenPercentage)
-	{
-		DynamicResolutionFractions[GDynamicPrimaryResolutionFraction] = 1.0f;
-
-		// The base pass have to respect FSceneView::UnscaledViewRect.
-		for (FViewInfo& View : Views)
-		{
-			View.ViewRect = View.UnscaledViewRect;
-		}
-
-		ComputeFamilySize();
-
-		return;
-	}
-
 	// Checks that view rects were still not initialized.
 	for (FViewInfo& View : Views)
 	{
@@ -3291,10 +3275,27 @@ void FSceneRenderer::PrepareViewRectsForRendering()
 			{
 				// Disable anti-aliasing if we are not going to be able to apply final post process effects
 				View.AntiAliasingMethod = AAM_None;
+				View.PrimaryScreenPercentageMethod = EPrimaryScreenPercentageMethod::SpatialUpscale;
 			}
 		}
 	}
 
+	// If not supporting screen percentage, bypass all computation.
+	if (!ViewFamily.EngineShowFlags.ScreenPercentage)
+	{
+		DynamicResolutionFractions[GDynamicPrimaryResolutionFraction] = 1.0f;
+
+		// The base pass have to respect FSceneView::UnscaledViewRect.
+		for (FViewInfo& View : Views)
+		{
+			View.ViewRect = View.UnscaledViewRect;
+		}
+
+		ComputeFamilySize();
+
+		return;
+	}
+
 	float PrimaryResolutionFraction = DynamicResolutionFractions[GDynamicPrimaryResolutionFraction];
 	{
 		// Ensure screen percentage show flag is respected. Prefer to check() rather rendering at a differen screen percentage
