PathTracer: Add support for external realtime denoisers when using the path tracing primary rays debug view

This is mostly experimental for now, but provides some of the standard G-Buffer like textures that realtime denoisers expect.

#rnx

diff --git a/Engine/Shaders/Private/PathTracing/PathTracingDebug.usf b/Engine/Shaders/Private/PathTracing/PathTracingDebug.usf
--- a/Engine/Shaders/Private/PathTracing/PathTracingDebug.usf
+++ b/Engine/Shaders/Private/PathTracing/PathTracingDebug.usf
@@ -36,6 +36,11 @@
 #endif
 
 RWTexture2D<float4> RWSceneColor;
+RWTexture2D<float4> RWSceneDiffuseColor;
+RWTexture2D<float4> RWSceneSpecularColor;
+RWTexture2D<float4> RWSceneNormal;
+RWTexture2D<float>  RWSceneRoughness;
+RWTexture2D<float>  RWSceneLinearDepth;
 RWTexture2D<float>  RWSceneDepth;
 RaytracingAccelerationStructure TLAS;
 
@@ -95,6 +100,11 @@ float3 TraceShadowRay(FRayDesc Ray, float PathRoughness, float PreviousPathRough
 
 struct FPathTracingResult {
 	float4 Radiance;
+	float3 DiffuseColor;
+	float3 SpecularColor;
+	float3 WorldNormal;
+	float Roughness;
+	float LinearDepth;
 	float Z;
 };
 
@@ -218,6 +228,7 @@ FPathTracingResult PathTracingDebugCore(FRayDesc Ray, uint2 LaunchIndex, int Sam
 	const float AlphaMin = 1.0 / 100.0; // Allow blending over 100 frames, higher than original publication because the chains get reset during BSDF sampling
 	const float MaterialProb = 0.15f; 	// Chance of BSDF sampling (TODO: adapt this as a function of roughness)
 
+	LOOP
 	for (int Bounce = 0; Bounce <= MaxBounces; Bounce++)
 	{
 		// Trace Ray
@@ -287,8 +298,16 @@ FPathTracingResult PathTracingDebugCore(FRayDesc Ray, uint2 LaunchIndex, int Sam
 		if (bCameraRay)
 		{
 			const float3 ViewZ = View.ViewToTranslatedWorld[2].xyz;
-			Result.Z = ConvertToDeviceZ(dot(TranslatedWorldPos, ViewZ));
+			const float LinearZ = dot(TranslatedWorldPos, ViewZ);
+			Result.Z = ConvertToDeviceZ(LinearZ);
 			Result.Radiance.a = 0; // opaque alpha (we aren't handling translucency right now)
+			// G-Buffer like data for denoisers
+			Result.WorldNormal = PackedPayload.UnpackWorldNormal();
+			Result.Roughness = PackedPayload.UnpackRoughness();
+
+			Result.DiffuseColor = PackedPayload.UnpackBaseColor() * (1 - PackedPayload.UnpackMetallic());
+			Result.SpecularColor = lerp(0.08 * PackedPayload.UnpackSpecular(), PackedPayload.UnpackBaseColor(), PackedPayload.UnpackMetallic());
+			Result.LinearDepth = LinearZ;
 		}
 
 		// Add radiance (emissive surfaces)
@@ -778,19 +797,40 @@ RAY_TRACING_ENTRY_RAYGEN(PathTracingDebugRG)
 #if PATH_TRACER_PRIMARY_RAYS
 	FPathTracingResult Result = (FPathTracingResult) 0;
 	
-	// TODO: launch separate kernels and accumulate instead
+	// TODO: launch separate kernels and accumulate instead?
+	LOOP
 	for (int SampleIndex = 0; SampleIndex < SamplesPerPixel; SampleIndex++)
 	{
 		FPathTracingResult SampleResult = PathTracingDebugCore(CameraRay, LaunchIndex, SampleIndex);
 
 		Result.Radiance += SampleResult.Radiance;
 		Result.Z = max(SampleResult.Z, Result.Z); // inverted Z so larger numbers are closer
+
+		// G-Buffer like data for denoisers
+		Result.LinearDepth += SampleResult.LinearDepth;
+		Result.DiffuseColor += SampleResult.DiffuseColor;
+		Result.SpecularColor += SampleResult.SpecularColor;
+		Result.WorldNormal += SampleResult.WorldNormal;
+		Result.Roughness += SampleResult.Roughness;
 	}
-	Result.Radiance *= 1.0 / SamplesPerPixel;
+	const float NormalizationFactor = 1.0 / SamplesPerPixel;
+	Result.Radiance *= NormalizationFactor;
+	Result.LinearDepth *= NormalizationFactor;
+	Result.DiffuseColor *= NormalizationFactor;
+	Result.SpecularColor *= NormalizationFactor;
+	Result.WorldNormal = normalize(Result.WorldNormal); // averaged roughness better?
+	Result.Roughness *= NormalizationFactor;
+
 #else
 	FPathTracingResult Result = PathTracingDebugCore(CameraRay, LaunchIndex);
 #endif
 
-	RWSceneColor[DispatchIdx] = Result.Radiance;
+	// Write to output textures
+	RWSceneColor[DispatchIdx] = Result.Radiance;	
 	RWSceneDepth[DispatchIdx] = Result.Z;
+	RWSceneDiffuseColor[DispatchIdx] = float4(Result.DiffuseColor.xyz, 1.0f);
+	RWSceneSpecularColor[DispatchIdx] = float4(Result.SpecularColor.xyz, 1.0f);
+	RWSceneNormal[DispatchIdx] = float4(Result.WorldNormal, Result.Roughness); // TODO: which copy of roughness is actually used?
+	RWSceneRoughness[DispatchIdx] = Result.Roughness;
+	RWSceneLinearDepth[DispatchIdx] = Result.LinearDepth;
 }
diff --git a/Engine/Source/Runtime/Renderer/Internal/PathTracingResources.h b/Engine/Source/Runtime/Renderer/Internal/PathTracingResources.h
--- a/Engine/Source/Runtime/Renderer/Internal/PathTracingResources.h
+++ b/Engine/Source/Runtime/Renderer/Internal/PathTracingResources.h
@@ -14,3 +14,12 @@ struct FPathTracingResources
 	
 	bool bPostProcessEnabled = false;
 };
+
+struct FPathTracingRealtimeDenoiserResources
+{
+	FRDGTexture* DiffuseColor = nullptr;
+	FRDGTexture* SpecularColor = nullptr;
+	FRDGTexture* Roughness = nullptr;
+	FRDGTexture* Normal = nullptr;
+	FRDGTexture* LinearDepth = nullptr;
+};
diff --git a/Engine/Source/Runtime/Renderer/Private/PathTracing.cpp b/Engine/Source/Runtime/Renderer/Private/PathTracing.cpp
--- a/Engine/Source/Runtime/Renderer/Private/PathTracing.cpp
+++ b/Engine/Source/Runtime/Renderer/Private/PathTracing.cpp
@@ -628,6 +628,9 @@ bool PathTracing::IsOutputInvalidateAllowed(EInvalidateReason InvalidateReason)
 	return true;
 }
 
+// Placeholder for experimenting with newer versions of external denoisers
+RENDERER_API FPathTracingRealtimeDenoiserResources GPathTracingRealtimeDenoiserResources;
+
 BEGIN_SHADER_PARAMETER_STRUCT(FPathTracingData, )
 	SHADER_PARAMETER(float, BlendFactor)
 	SHADER_PARAMETER(float, BaseExposure)
@@ -1457,6 +1460,11 @@ class FPathTracingDebugRG : public FGlobalShader
 
 	BEGIN_SHADER_PARAMETER_STRUCT(FParameters, )
 		SHADER_PARAMETER_RDG_TEXTURE_UAV(RWTexture2D<float4>, RWSceneColor)
+		SHADER_PARAMETER_RDG_TEXTURE_UAV(RWTexture2D<float4>, RWSceneDiffuseColor)
+		SHADER_PARAMETER_RDG_TEXTURE_UAV(RWTexture2D<float4>, RWSceneSpecularColor)
+		SHADER_PARAMETER_RDG_TEXTURE_UAV(RWTexture2D<float4>, RWSceneNormal)
+		SHADER_PARAMETER_RDG_TEXTURE_UAV(RWTexture2D<float>,  RWSceneRoughness)
+		SHADER_PARAMETER_RDG_TEXTURE_UAV(RWTexture2D<float>,  RWSceneLinearDepth)
 		SHADER_PARAMETER_RDG_TEXTURE_UAV(RWTexture2D<float>,  RWSceneDepth)
 		SHADER_PARAMETER_RDG_BUFFER_SRV(RaytracingAccelerationStructure, TLAS)
 		SHADER_PARAMETER_RDG_BUFFER_UAV(RWStructuredBuffer<FMarkovChainState>, MarkovChainStateGrid)
@@ -3280,10 +3288,11 @@ void FDeferredShadingSceneRenderer::PreparePathTracing(const FViewInfo& View, co
 	{
 		if (int DebugMode = GetPathTracingVisualizationMode(); DebugMode >= 0)
 		{
+			const bool bUsePrimaryRays = DebugMode == PATH_TRACER_DEBUG_VIZ_PRIMARY_RAYS;
 			const bool bUseSER = GRHIGlobals.SupportsShaderExecutionReordering && CVarPathTracingShaderExecutionReordering.GetValueOnRenderThread();
 			FPathTracingDebugRG::FPermutationDomain PermutationVector;
-			PermutationVector.Set<FPathTracingDebugRG::FUseSER>(bUseSER);
-			PermutationVector.Set<FPathTracingDebugRG::FUsePrimaryRays>(DebugMode == PATH_TRACER_DEBUG_VIZ_PRIMARY_RAYS);
+			PermutationVector.Set<FPathTracingDebugRG::FUseSER>(bUseSER && bUsePrimaryRays);
+			PermutationVector.Set<FPathTracingDebugRG::FUsePrimaryRays>(bUsePrimaryRays);
 			auto RayGenShader = GetGlobalShaderMap(ViewFamily.GetShaderPlatform())->GetShader<FPathTracingDebugRG>(PermutationVector);
 			OutRayGenShaders.Add(RayGenShader.GetRayTracingShader());
 		}
@@ -3482,6 +3491,31 @@ void FDeferredShadingSceneRenderer::RenderPathTracing(
 			FClearValueBinding::None,
 			TexCreate_ShaderResource | TexCreate_UAV);
 		FRDGTextureRef DepthTexture    = GraphBuilder.CreateTexture(DepthDesc, TEXT("PathTracer.Depth"));
+		
+		FRDGTextureDesc AlbedoNormalDesc = FRDGTextureDesc::Create2D(
+			View.ViewRect.Size(),
+			PF_FloatRGBA,
+			FClearValueBinding::None,
+			TexCreate_ShaderResource | TexCreate_UAV);
+		
+		FRDGTextureDesc RoughnessDesc(FRDGTextureDesc::Create2D(
+			View.ViewRect.Size(),
+			PF_R32_FLOAT, // TODO PF_B8G8R8A8 enough?
+			FClearValueBinding::None,
+			TexCreate_ShaderResource | TexCreate_UAV
+		));
+
+		FRDGTextureRef DiffuseColorTexture = GraphBuilder.CreateTexture(AlbedoNormalDesc, TEXT("PathTracer.DiffuseColor"));
+		FRDGTextureRef SpecularColorTexture = GraphBuilder.CreateTexture(AlbedoNormalDesc, TEXT("PathTracer.SpecularColor"));
+		FRDGTextureRef NormalTexture = GraphBuilder.CreateTexture(AlbedoNormalDesc, TEXT("PathTracer.Normal_Aleks"));
+		FRDGTextureRef RoughnessTexture = GraphBuilder.CreateTexture(RoughnessDesc, TEXT("PathTracer.Roughness"));
+		FRDGTextureRef LinearDepthTexture = GraphBuilder.CreateTexture(RoughnessDesc, TEXT("PathTracer.LinearDepth"));
+
+		GPathTracingRealtimeDenoiserResources.DiffuseColor = DiffuseColorTexture;
+		GPathTracingRealtimeDenoiserResources.SpecularColor = SpecularColorTexture;
+		GPathTracingRealtimeDenoiserResources.Normal = NormalTexture;
+		GPathTracingRealtimeDenoiserResources.Roughness = RoughnessTexture;
+		GPathTracingRealtimeDenoiserResources.LinearDepth = LinearDepthTexture;		
 
 		FPathTracingState* PathTracingState = GetPathTracingStateFromView(View);
 		FRDGBuffer* MarkovChainStateGrid = nullptr;
@@ -3523,16 +3557,23 @@ void FDeferredShadingSceneRenderer::RenderPathTracing(
 			&PassParameters->SceneVisibleLightCount,
 			&PassParameters->SceneLightCount,
 			&PassParameters->SceneLights
-		);
+		);				
+
+		PassParameters->RWSceneDiffuseColor = GraphBuilder.CreateUAV(DiffuseColorTexture);
+		PassParameters->RWSceneSpecularColor = GraphBuilder.CreateUAV(SpecularColorTexture);
+		PassParameters->RWSceneNormal = GraphBuilder.CreateUAV(NormalTexture);
+		PassParameters->RWSceneRoughness = GraphBuilder.CreateUAV(RoughnessTexture);
+		PassParameters->RWSceneLinearDepth = GraphBuilder.CreateUAV(LinearDepthTexture);
 
 		PassParameters->DecalParameters = View.RayTracingDecalUniformBuffer;
 		PassParameters->Scene = View.GetSceneUniforms().GetBuffer(GraphBuilder);
 		PassParameters->NaniteRayTracing = Nanite::GRayTracingManager.GetUniformBuffer();
 
+		const bool bUsePrimaryRays = DebugMode == PATH_TRACER_DEBUG_VIZ_PRIMARY_RAYS;
 		const bool bUseSER = GRHIGlobals.SupportsShaderExecutionReordering && CVarPathTracingShaderExecutionReordering.GetValueOnRenderThread();
 		FPathTracingDebugRG::FPermutationDomain PermutationVector;
-		PermutationVector.Set<FPathTracingDebugRG::FUseSER>(bUseSER);
-		PermutationVector.Set<FPathTracingDebugRG::FUsePrimaryRays>(DebugMode == PATH_TRACER_DEBUG_VIZ_PRIMARY_RAYS);
+		PermutationVector.Set<FPathTracingDebugRG::FUseSER>(bUseSER && bUsePrimaryRays);
+		PermutationVector.Set<FPathTracingDebugRG::FUsePrimaryRays>(bUsePrimaryRays);
 		TShaderMapRef<FPathTracingDebugRG> RayGenShader(View.ShaderMap, PermutationVector);
 		ClearUnusedGraphResources(RayGenShader, PassParameters);
 		GraphBuilder.AddPass(
