Fix missing alias for X24_G8 on LevelInstanceDepth

diff --git a/Engine/Source/Runtime/Renderer/Private/PostProcess/PostProcessVisualizeLevelInstance.cpp b/Engine/Source/Runtime/Renderer/Private/PostProcess/PostProcessVisualizeLevelInstance.cpp
--- a/Engine/Source/Runtime/Renderer/Private/PostProcess/PostProcessVisualizeLevelInstance.cpp
+++ b/Engine/Source/Runtime/Renderer/Private/PostProcess/PostProcessVisualizeLevelInstance.cpp
@@ -81,6 +81,7 @@ FScreenPassTexture AddVisualizeLevelInstancePass(
 			DepthStencilDesc.ClearValue = FClearValueBinding((float)ERHIZBuffer::FarPlane, 0);
 			DepthStencilDesc.Flags = TexCreate_DepthStencilTargetable | TexCreate_ShaderResource;
 			DepthStencilDesc.NumSamples = NumSamples;
+			DepthStencilDesc.AliasableFormats.Add(PF_X24_G8);
 
 			DepthStencilTexture = GraphBuilder.CreateTexture(DepthStencilDesc, TEXT("LevelInstanceDepth"));
 		}
@@ -174,4 +175,4 @@ FScreenPassTexture AddVisualizeLevelInstancePass(
 	return MoveTemp(Output);
 }
 
-#endif
\ No newline at end of file
+#endif
