[CustomRenderPass] Fixed ViewRect and ViewRectWithSecondaryViews not being properly initialized for CRP views.

diff --git a/Engine/Source/Runtime/Renderer/Private/SceneRendering.cpp b/Engine/Source/Runtime/Renderer/Private/SceneRendering.cpp
--- a/Engine/Source/Runtime/Renderer/Private/SceneRendering.cpp
+++ b/Engine/Source/Runtime/Renderer/Private/SceneRendering.cpp
@@ -3017,6 +3017,8 @@ FSceneRenderer::FSceneRenderer(const FSceneViewFamily* InViewFamily, FHitProxyCo
 		ViewInfo->bDisableQuerySubmissions = true;
 		ViewInfo->bIgnoreExistingQueries = true;
 		ViewInfo->CustomRenderPass = CustomRenderPass;
+		ViewInfo->ViewRect = ViewInfo->UnscaledViewRect;
+		ViewInfo->ViewRectWithSecondaryViews = ViewInfo->UnscaledViewRect;
 		CustomRenderPass->Views.Add(ViewInfo);
 
 		NumAdditionalViews++;
