Don't allocate unnecessary receiver mask space if local light receiver masks are disabled.

diff --git a/Engine/Shaders/Private/VirtualShadowMaps/VirtualShadowMapPageManagement.usf b/Engine/Shaders/Private/VirtualShadowMaps/VirtualShadowMapPageManagement.usf
--- a/Engine/Shaders/Private/VirtualShadowMaps/VirtualShadowMapPageManagement.usf
+++ b/Engine/Shaders/Private/VirtualShadowMaps/VirtualShadowMapPageManagement.usf
@@ -153,7 +153,7 @@ void ProcessMaskMip4x4(RWTexture2D<uint> OutPageReceiverMaskMip, inout uint Mask
 // Gather a 2x2 footprint given the upper-left coordinate
 uint4 GatherReceiverMask(uint2 TexelCoord)
 {
-	return GatherPageTable(InPageReceiverMasks, TexelCoord, 0u, 0.5f * VirtualShadowMap.PageTableTextureSizeInvSize.zw);
+	return GatherPageTable(InPageReceiverMasks, TexelCoord, 0u, VirtualShadowMap.PageReceiverMaskTextureSizeInvSize.zw);
 }
 
 /**
@@ -181,6 +181,8 @@ void GenerateHierarchicalPageFlags(uint ThreadId : SV_DispatchThreadID)
 	uint Flag = InPageFlags[GlobalPageTableEntryIndex.GetResourceAddress()] & VSM_PAGE_FLAGS_BITS_MASK;
 	if (Flag != 0u)
 	{
+		FVirtualShadowMapProjectionShaderData ProjectionData = GetVirtualShadowMapProjectionData(VirtualShadowMapHandle);
+
 		uint MipLevel = MetaData.MipLevel;
 		uint2 PageAddress = MetaData.PageAddress;
 		// No hierarchy to propagate to for single-page pages.
@@ -235,7 +237,7 @@ void GenerateHierarchicalPageFlags(uint ThreadId : SV_DispatchThreadID)
 #undef PROCESS_LEVEL
 
 		BRANCH
-		if (VirtualShadowMap.bEnableReceiverMasks)
+		if (ProjectionData.bUseReceiverMask)
 		{
 			// 2x2 masks
 			uint2 MaskAddress = GlobalPageTableEntryIndex.GetResourceAddress() * 2u;
diff --git a/Engine/Shaders/Private/VirtualShadowMaps/VirtualShadowMapPageMarking.usf b/Engine/Shaders/Private/VirtualShadowMaps/VirtualShadowMapPageMarking.usf
--- a/Engine/Shaders/Private/VirtualShadowMaps/VirtualShadowMapPageMarking.usf
+++ b/Engine/Shaders/Private/VirtualShadowMaps/VirtualShadowMapPageMarking.usf
@@ -135,7 +135,7 @@ void MarkPage(FVirtualShadowMapHandle VirtualShadowMapHandle, uint MipLevel, flo
 	MarkPageAddress(PageOffset, Flags);
 
 	BRANCH
-	if (VirtualShadowMap.bEnableReceiverMasks)
+	if (ProjectionData.bUseReceiverMask)
 	{
 		MarkPageReceiverMask(PageOffset, VirtualAddress);
 	}
@@ -280,7 +280,7 @@ void MarkPageFroxel(
 	}
 
 	BRANCH
-	if (VirtualShadowMap.bEnableReceiverMasks)
+	if (ProjectionData.bUseReceiverMask)
 	{
 		// Rect in 2x page space storing 4x4 masks in each uint
 		uint4 MaskFPRect = uint4(Rect.Pixels) >> (VSM_LOG2_PAGE_SIZE - 1u);
@@ -992,7 +992,7 @@ void MarkCoarsePages(uint DispatchThreadId : SV_DispatchThreadID)
 			MarkPageAddress(CalcPageOffset(VirtualShadowMapHandle, 0, PageAddressLowHigh.zw), Flags);
 
 			BRANCH
-			if (VirtualShadowMap.bEnableReceiverMasks)
+			if (ProjectionData.bUseReceiverMask)
 			{
 				MarkFullPageReceiverMask(CalcPageOffset(VirtualShadowMapHandle, 0, PageAddressLowHigh.xy));
 				MarkFullPageReceiverMask(CalcPageOffset(VirtualShadowMapHandle, 0, PageAddressLowHigh.xw));
@@ -1008,7 +1008,7 @@ void MarkCoarsePages(uint DispatchThreadId : SV_DispatchThreadID)
 		FVSMPageOffset PageOffset = CalcPageOffset(VirtualShadowMapHandle, VSM_MAX_MIP_LEVELS - 1, uint2(0, 0));
 		MarkPageAddress(PageOffset, Flags);
 		BRANCH
-		if (VirtualShadowMap.bEnableReceiverMasks)
+		if (ProjectionData.bUseReceiverMask)
 		{
 			MarkFullPageReceiverMask(PageOffset);
 		}
diff --git a/Engine/Shaders/Private/VirtualShadowMaps/VirtualShadowMapPageOverlap.ush b/Engine/Shaders/Private/VirtualShadowMaps/VirtualShadowMapPageOverlap.ush
--- a/Engine/Shaders/Private/VirtualShadowMaps/VirtualShadowMapPageOverlap.ush
+++ b/Engine/Shaders/Private/VirtualShadowMaps/VirtualShadowMapPageOverlap.ush
@@ -69,7 +69,7 @@ uint4 GatherPageFlags(uint2 TexelCoord, uint HMipLevel)
 uint4 GatherReceiverMask(uint2 TexelCoord, uint HMipLevel)
 {
 	// Note: UV scaled by extra 0.5f to account for 2x size of mask texture
-	return GatherPageTable(VirtualShadowMap.PageReceiverMasks, TexelCoord, HMipLevel, 0.5f * VirtualShadowMap.PageTableTextureSizeInvSize.zw);
+	return GatherPageTable(VirtualShadowMap.PageReceiverMasks, TexelCoord, HMipLevel, VirtualShadowMap.PageReceiverMaskTextureSizeInvSize.zw);
 }
 
 // Returns true if ANY flags in FlagMask are set on at least one overlapped page (see GetPageFlagMaskForRendering below)
@@ -168,7 +168,7 @@ bool OverlapsAnyValidPage(FVirtualShadowMapHandle VirtualShadowMapHandle, uint M
 		(!bDetailGeometry || ((PageFlags & VSM_FLAG_DETAIL_GEOMETRY) != 0)))
 	{
 		BRANCH
-		if (VirtualShadowMap.bEnableReceiverMasks && bUseReceiverMask)
+		if (bUseReceiverMask)
 		{
 			// Rect in 2x page space storing 4x4 masks in each uint
 			uint4 MaskFPRect = uint4(ClampedScreenRect.Pixels) >> (VSM_LOG2_PAGE_SIZE - 1u);
diff --git a/Engine/Source/Runtime/Renderer/Private/Shadows/ShadowSceneRenderer.cpp b/Engine/Source/Runtime/Renderer/Private/Shadows/ShadowSceneRenderer.cpp
--- a/Engine/Source/Runtime/Renderer/Private/Shadows/ShadowSceneRenderer.cpp
+++ b/Engine/Source/Runtime/Renderer/Private/Shadows/ShadowSceneRenderer.cpp
@@ -626,7 +626,7 @@ void FShadowSceneRenderer::AllocateVirtualShadowMapIds()
 			}
 
 			const TSharedPtr<FVirtualShadowMapPerLightCacheEntry>& CacheEntry = ClipmapInfo.Clipmap->GetCacheEntry();
-			const int32 NextVirtualShadowMapId = VirtualShadowMapArray.Allocate(false, CacheEntry->ShadowMapEntries.Num());
+			const int32 NextVirtualShadowMapId = VirtualShadowMapArray.AllocateDirectional(CacheEntry->ShadowMapEntries.Num());
 			CacheEntry->UpdateVirtualShadowMapId(NextVirtualShadowMapId);
 		}
 	}
@@ -635,7 +635,7 @@ void FShadowSceneRenderer::AllocateVirtualShadowMapIds()
 	for (const FLocalLightShadowFrameSetup& FrameSetup : LocalLights)
 	{
 		const TSharedPtr<FVirtualShadowMapPerLightCacheEntry>& CacheEntry = FrameSetup.PerLightCacheEntry;
-		const int32 NextVirtualShadowMapId = VirtualShadowMapArray.Allocate(CacheEntry->bIsDistantLight, CacheEntry->ShadowMapEntries.Num());
+		const int32 NextVirtualShadowMapId = VirtualShadowMapArray.AllocateLocal(CacheEntry->bIsDistantLight, CacheEntry->ShadowMapEntries.Num());
 		CacheEntry->UpdateVirtualShadowMapId(NextVirtualShadowMapId);
 	}
 
diff --git a/Engine/Source/Runtime/Renderer/Private/VirtualShadowMaps/VirtualShadowMapArray.cpp b/Engine/Source/Runtime/Renderer/Private/VirtualShadowMaps/VirtualShadowMapArray.cpp
--- a/Engine/Source/Runtime/Renderer/Private/VirtualShadowMaps/VirtualShadowMapArray.cpp
+++ b/Engine/Source/Runtime/Renderer/Private/VirtualShadowMaps/VirtualShadowMapArray.cpp
@@ -669,7 +669,6 @@ void FVirtualShadowMapArray::Initialize(
 	UniformParameters.CoarsePagePixelThresholdDynamic = CVarCoarsePagePixelThresholdDynamic.GetValueOnRenderThread();
 	UniformParameters.CoarsePagePixelThresholdStatic = CVarCoarsePagePixelThresholdStatic.GetValueOnRenderThread();
 	UniformParameters.CoarsePagePixelThresholdDynamicNanite = CVarCoarsePagePixelThresholdDynamicNanite.GetValueOnRenderThread();
-	UniformParameters.bEnableReceiverMasks = IsVirtualShadowMapDirectionalReceiverMaskEnabled() || IsVirtualShadowMapLocalReceiverMaskEnabled();
 	UniformParameters.MipModeLocal = CVarMarkPixelPagesMipModeLocal.GetValueOnRenderThread();
 
 	UniformParameters.SceneFrameNumber = Scene.GetFrameNumberRenderThread();
@@ -710,7 +709,7 @@ void FVirtualShadowMapArray::Initialize(
 
 	if (bShouldRetainCacheInfo)
 	{
-		// Always reserve IDs for the single-page SMs.
+		// Always reserve slots for the single-page SMs.
 		NumShadowMapSlots = VSM_MAX_SINGLE_PAGE_SHADOW_MAPS;
 
 		// Fixed physical page pool width, we adjust the height to accomodate the requested maximum
@@ -757,6 +756,7 @@ void FVirtualShadowMapArray::Initialize(
 		UniformParameters.PageTableRowMask = CacheManager->PrevUniformParameters.PageTableRowMask;
 		UniformParameters.PageTableRowShift = CacheManager->PrevUniformParameters.PageTableRowShift;
 		UniformParameters.PageTableTextureSizeInvSize = CacheManager->PrevUniformParameters.PageTableTextureSizeInvSize;
+		UniformParameters.PageReceiverMaskTextureSizeInvSize = CacheManager->PrevUniformParameters.PageReceiverMaskTextureSizeInvSize;
 
 		// TODO: Parameterize this in a useful way; potentially modify it automatically
 		// when there are fewer lights in the scene and/or clustered shading settings differ.
@@ -813,7 +813,7 @@ void FVirtualShadowMapArray::Initialize(
 	UpdateCachedUniformBuffers(GraphBuilder);
 }
 
-int32 FVirtualShadowMapArray::Allocate(bool bSinglePageShadowMap, int32 Count)
+int32 FVirtualShadowMapArray::AllocateShadowMapSlots(bool bSinglePageShadowMap, int32 Count)
 {
 	check(IsEnabled());
 	int32 VirtualShadowMapId = INDEX_NONE;
@@ -831,10 +831,33 @@ int32 FVirtualShadowMapArray::Allocate(bool bSinglePageShadowMap, int32 Count)
 		VirtualShadowMapId = NumShadowMapSlots;
 		NumShadowMapSlots += Count;
 	}
-
 	return VirtualShadowMapId;
 }
 
+int32 FVirtualShadowMapArray::AllocateDirectional(int32 Count)
+{	
+	check(IsEnabled());
+	check(NumUnreferencedShadowMaps == 0);
+	check(NumLocalShadowMaps == 0);
+	NumDirectionalShadowMaps += Count;
+	return AllocateShadowMapSlots(false, Count);
+}
+
+int32 FVirtualShadowMapArray::AllocateLocal(bool bSinglePageShadowMap, int32 Count)
+{
+	check(IsEnabled());
+	check(NumUnreferencedShadowMaps == 0);
+	NumLocalShadowMaps += Count;
+	return AllocateShadowMapSlots(bSinglePageShadowMap, Count);
+}
+
+int32 FVirtualShadowMapArray::AllocateUnreferenced(bool bSinglePageShadowMap, int32 Count)
+{
+	check(IsEnabled());
+	NumUnreferencedShadowMaps += Count;
+	return AllocateShadowMapSlots(bSinglePageShadowMap, Count);
+}
+
 FVirtualShadowMapArray::~FVirtualShadowMapArray()
 {
 }
@@ -2017,7 +2040,6 @@ void FVirtualShadowMapArray::BuildPageAllocations(
 	FIntPoint PageTableTextureSize(NumPageTablesEntriesPerRow * VSM_PAGE_TABLE_TEX2D_SIZE_X, NumPageTableRows * VSM_PAGE_TABLE_TEX2D_SIZE_Y);
 	UniformParameters.PageTableRowShift = uint32(FMath::FloorLog2(NumPageTablesEntriesPerRow));
 	UniformParameters.PageTableRowMask = NumPageTablesEntriesPerRow - 1u;
-
 	UniformParameters.PageTableTextureSizeInvSize = FVector4f(FVector2f(PageTableTextureSize), FVector2f(1.0f, 1.0f) / FVector2f(PageTableTextureSize));
 
 	auto AllocatePageTable = [&](bool bAllocateMipLevels, EPixelFormat PixelFormat, const TCHAR* DebugName, uint32 SampleStride = 1u)
@@ -2029,6 +2051,7 @@ void FVirtualShadowMapArray::BuildPageAllocations(
 				FClearValueBinding(0), 
 				TexCreate_ShaderResource | TexCreate_UAV | TexCreate_AtomicCompatible,
 				// Technically one Hmip too little but the 1x1 level is meaningless as it mixes the info from all the various mip levels. 
+				// TODO: Should this be Log2(SampleStride)...? Only used with 2 currently
 				bAllocateMipLevels ? (VSM_LOG2_PAGE_SIZE + SampleStride / 2u): 1u), 
 			DebugName);
 	};
@@ -2091,11 +2114,40 @@ void FVirtualShadowMapArray::BuildPageAllocations(
 	// TODO: should be uint8 - but atomics prevent this
 	PageRequestFlagsRDG = AllocatePageTable(false, PF_R32_UINT, TEXT("Shadow.Virtual.PageRequestFlags"));
 	AddClearPageTableUAVPass(GraphBuilder, PageRequestFlagsRDG, 0u);
-	if (UniformParameters.bEnableReceiverMasks)
+
+	if (IsVirtualShadowMapDirectionalReceiverMaskEnabled() || IsVirtualShadowMapLocalReceiverMaskEnabled())
 	{
-		// TODO: should be uint16 - but atomics prevent this
-		PageReceiverMasksRDG = AllocatePageTable(true, PF_R32_UINT, TEXT("Shadow.Virtual.PageReceiverMasks"), 2u);
-		AddClearPageTableUAVPass(GraphBuilder, PageReceiverMasksRDG, 0u, 2u);
+		// If local light receiver masks are enabled (and present), allocate and clear it as another full page table
+		if (IsVirtualShadowMapLocalReceiverMaskEnabled())
+		{
+			PageReceiverMasksRDG = AllocatePageTable(true, PF_R32_UINT, TEXT("Shadow.Virtual.PageReceiverMasks"), 2u);
+			UniformParameters.PageReceiverMaskTextureSizeInvSize = UniformParameters.PageTableTextureSizeInvSize;
+			// TODO: If we don't actually have any local lights perhaps worth doing the simpler ClearUAVPass below instead?
+			AddClearPageTableUAVPass(GraphBuilder, PageReceiverMasksRDG, 0u, 2u);
+		}
+		else
+		{
+			// Directional goes before local so that if local receiver masks are disabled we don't need to allocate them
+			// NOTE: One extra for the single page entries
+			int32 NumDirectionalPageReceiverMaskEntriesRequired = NumDirectionalShadowMaps + 1;
+			const uint32 NumPageReceiverMaskRows = FMath::DivideAndRoundUp((uint32)NumDirectionalPageReceiverMaskEntriesRequired, NumPageTablesEntriesPerRow);
+			// Each page stores 2x2 masks, each representing a 4x4 tile
+			FIntPoint PageReceiverMaskTextureSize(2U * NumPageTablesEntriesPerRow * VSM_PAGE_TABLE_TEX2D_SIZE_X, 2U * NumPageReceiverMaskRows * VSM_PAGE_TABLE_TEX2D_SIZE_Y);
+			UniformParameters.PageReceiverMaskTextureSizeInvSize = FVector4f(FVector2f(PageReceiverMaskTextureSize), FVector2f(1.0f, 1.0f) / FVector2f(PageReceiverMaskTextureSize));
+
+			// TODO: should be uint16 - but atomics prevent this
+			PageReceiverMasksRDG = GraphBuilder.CreateTexture(
+				FRDGTextureDesc::Create2D(
+					PageReceiverMaskTextureSize,
+					PF_R32_UINT, 
+					FClearValueBinding(0), 
+					TexCreate_ShaderResource | TexCreate_UAV | TexCreate_AtomicCompatible,
+					// Technically one Hmip too few but the 1x1 level is meaningless as it mixes the info from all the various mip levels. 
+					// One extra mip level for the 2x2 tiles
+					(VSM_LOG2_PAGE_SIZE + 1)),
+				TEXT("Shadow.Virtual.PageReceiverMasks"));
+			AddClearUAVPass(GraphBuilder, GraphBuilder.CreateUAV(PageReceiverMasksRDG), FUintVector4(0, 0, 0, 0));
+		}
 	}
 	else
 	{
diff --git a/Engine/Source/Runtime/Renderer/Private/VirtualShadowMaps/VirtualShadowMapArray.h b/Engine/Source/Runtime/Renderer/Private/VirtualShadowMaps/VirtualShadowMapArray.h
--- a/Engine/Source/Runtime/Renderer/Private/VirtualShadowMaps/VirtualShadowMapArray.h
+++ b/Engine/Source/Runtime/Renderer/Private/VirtualShadowMaps/VirtualShadowMapArray.h
@@ -168,6 +168,7 @@ BEGIN_GLOBAL_SHADER_PARAMETER_STRUCT(FVirtualShadowMapUniformParameters, )
 	SHADER_PARAMETER(uint32, PageTableRowMask)
 	// xy == PageTableTextureSize, zw == 1.0f/PageTableTextureSize
 	SHADER_PARAMETER(FVector4f, PageTableTextureSizeInvSize) 
+	SHADER_PARAMETER(FVector4f, PageReceiverMaskTextureSizeInvSize)
 
 	// Set to 1 if r.Shadow.Virtual.NonNanite.IncludeInCoarsePages is set to 0 in order to signal that we want to use the legacy path for just excluding non-nanite
 	SHADER_PARAMETER(uint32, bExcludeNonNaniteFromCoarsePages)
@@ -179,7 +180,6 @@ BEGIN_GLOBAL_SHADER_PARAMETER_STRUCT(FVirtualShadowMapUniformParameters, )
 	SHADER_PARAMETER(uint32, SceneFrameNumber)
 
 	SHADER_PARAMETER(float, GlobalResolutionLodBias)
-	SHADER_PARAMETER(int32, bEnableReceiverMasks)
 	SHADER_PARAMETER(int32, MipModeLocal)
 
 	// SMRT parameters that are sometimes used globally
@@ -291,8 +291,11 @@ public:
 		return bEnabled;
 	}
 
-	// Returns the first in a continuously allocated range of new VirtualShadowMapIds
-	int32 Allocate(bool bSinglePageShadowMap, int32 Count);
+	// NOTE: All directional lights must be allocated before all local lights
+	// due to receiver mask allocation logic. Unreferenced lights must be allocated last.
+	int32 AllocateDirectional(int32 Count);
+	int32 AllocateLocal(bool bSinglePageShadowMap, int32 Count);
+	int32 AllocateUnreferenced(bool bSinglePageShadowMap, int32 Count);
 	
 	void UpdateNextData(int32 PrevVirtualShadowMapId, const FNextVirtualShadowMapData& Data);
 	
@@ -500,6 +503,8 @@ public:
 	bool bEnableNaniteVisualization = false;
 
 private:
+	int32 AllocateShadowMapSlots(bool bSinglePageShadowMap, int32 Count);
+
 	void UpdateVisualizeLight(
 		const TConstArrayView<FViewInfo> &Views,
 		const TConstArrayView<FVisibleLightInfo>& VisibleLightInfos);
@@ -530,8 +535,18 @@ private:
 	// This is primarily an indirection that allows us to reallocate/repack VirtualShadowMapIds each frame
 	TArray<FNextVirtualShadowMapData, SceneRenderingAllocator> NextData;
 
-	int32 NumShadowMapSlots = 0;
-	int32 NumSinglePageShadowMaps = 0;
+	// Allocation/page table order is:
+	// - Single page local light shadow maps (fixed number of slots)
+	// - Directional light shadow maps (always full)
+	// - Full local light shadow maps
+	// - Any unreferenced lights of any type (see FVirtualShadowMapCacheManager::UpdateUnreferencedCacheEntries)
+
+	int32 NumDirectionalShadowMaps = 0;
+	int32 NumLocalShadowMaps = 0;			// Includes both single page and full
+	int32 NumUnreferencedShadowMaps = 0;	// Includes both single page and full
+
+	int32 NumSinglePageShadowMaps = 0;	// Single page slot allocation counter
+	int32 NumShadowMapSlots = 0;		// Total/full allocation counter
 
 	// Gets created in dummy form at initialization time, then updated after VSM data is computed
 	TArray<TRDGUniformBufferRef<FVirtualShadowMapUniformParameters>> CachedUniformBuffers;
diff --git a/Engine/Source/Runtime/Renderer/Private/VirtualShadowMaps/VirtualShadowMapCacheManager.cpp b/Engine/Source/Runtime/Renderer/Private/VirtualShadowMaps/VirtualShadowMapCacheManager.cpp
--- a/Engine/Source/Runtime/Renderer/Private/VirtualShadowMaps/VirtualShadowMapCacheManager.cpp
+++ b/Engine/Source/Runtime/Renderer/Private/VirtualShadowMaps/VirtualShadowMapCacheManager.cpp
@@ -1284,7 +1284,7 @@ void FVirtualShadowMapArrayCacheManager::UpdateUnreferencedCacheEntries(
 			// Not active this render, but still recent enough to keep it and its pages alive
 			// Reallocate new VSM ID for this frame
 			int32 NumMaps = CacheEntry->ShadowMapEntries.Num();
-			const int32 NextVirtualShadowMapId = VirtualShadowMapArray.Allocate(CacheEntry->bIsDistantLight, NumMaps);
+			const int32 NextVirtualShadowMapId = VirtualShadowMapArray.AllocateUnreferenced(CacheEntry->bIsDistantLight, NumMaps);
 			CacheEntry->UpdateVirtualShadowMapId(NextVirtualShadowMapId);
 			for (int32 Map = 0; Map < NumMaps; ++Map)
 			{
