Path Tracing: Move rarely used options to experimental to reduce startup time

Some of the longest to compile modes are now gated under experimental which reduces the number of path tracing shader permutations compiled on startup from 32 down to 4 (for machines that support SER) or 2 otherwise.

Collapsing the cloud shader permutation was measured but found to slow down performance even when that codepath is not in use, so keeping it for now.

Tweak some cvar help text to get rid of redundant default value mentions (as this is now reported by the cvar system itself).

diff --git a/Engine/Source/Runtime/Renderer/Private/PathTracing.cpp b/Engine/Source/Runtime/Renderer/Private/PathTracing.cpp
--- a/Engine/Source/Runtime/Renderer/Private/PathTracing.cpp
+++ b/Engine/Source/Runtime/Renderer/Private/PathTracing.cpp
@@ -47,7 +47,7 @@ TAutoConsoleVariable<int32> CVarPathTracing(
 TAutoConsoleVariable<bool> CVarPathTracingExperimental(
 	TEXT("r.PathTracing.Experimental"),
 	false,
-	TEXT("Enables some experimental features of the path tracing renderer that require compiling additional permutations of the path tracer. (default: false)"),
+	TEXT("Enables experimental features or rarely used modes of the path tracing renderer that require compiling additional shader permutations. Enabling this will increase startup time by compiling additional shaders."),
 	ECVF_RenderThreadSafe | ECVF_ReadOnly
 );
 
@@ -149,35 +149,37 @@ TAutoConsoleVariable<int32> CVarPathTracingVisualizeIndirectGuidingHashSize(
 TAutoConsoleVariable<int32> CVarPathTracingCompactionDepth(
 	TEXT("r.PathTracing.CompactionDepth"),
 	-1,
-	TEXT("Enables path compaction to improve GPU occupancy for the path tracer. The value sets the bounce up to which compaction will happen, beyond that point the path tracer will handle bounces within the dispatch. (default: -1, all bounces)"),
+	TEXT("Enables path compaction to improve GPU occupancy for the path tracer. The value sets the bounce up to which compaction will happen, beyond that point the path tracer will handle bounces within the dispatch. (default: -1, all bounces)\n")
+	TEXT("Requires r.PathTracing.Experimental=true to modify.\n"),
 	ECVF_RenderThreadSafe
 );
 
 TAutoConsoleVariable<bool> CVarPathTracingTraceOpaqueFirst(
 	TEXT("r.PathTracing.TraceOpaqueFirst"),
 	true,
-	TEXT("Trace opaque geometry before translucent geometry. This allows the path tracer to setup a correct depth for DepthFade based effects (default: 1)"),
+	TEXT("Trace opaque geometry before translucent geometry. This allows the path tracer to setup a correct depth for DepthFade based effects.\n")
+	TEXT("Requires r.PathTracing.Experimental=true to modify.\n"),
 	ECVF_RenderThreadSafe
 );
 
 TAutoConsoleVariable<bool> CVarPathTracingShaderExecutionReordering(
 	TEXT("r.PathTracing.ShaderExecutionReordering"),
 	true,
-	TEXT("Enables Shader Execution Reordering to improve shader coherence for the path tracer. This variable only has effect if the underlying hardware supports it. (default: true (enabled))"),
+	TEXT("Enables Shader Execution Reordering to improve shader coherence for the path tracer. This variable only has effect if the current hardware supports it."),
 	ECVF_RenderThreadSafe
 );
 
 TAutoConsoleVariable<bool> CVarPathTracingIndirectDispatch(
 	TEXT("r.PathTracing.IndirectDispatch"),
 	false,
-	TEXT("Enables indirect dispatch (if supported by the hardware) for compacted path tracing (default: false (disabled))"),
+	TEXT("Enables indirect dispatch (if supported by the hardware) for compacted path tracing."),
 	ECVF_RenderThreadSafe
 );
 
 TAutoConsoleVariable<int32> CVarPathTracingFlushDispatch(
 	TEXT("r.PathTracing.FlushDispatch"),
 	2,
-	TEXT("Enables flushing of the command list after dispatch to reduce the likelyhood of TDRs on Windows (default: 2)\n")
+	TEXT("Enables flushing of the command list after dispatch to reduce the likelyhood of TDRs on Windows\n")
 	TEXT("0: off\n")
 	TEXT("1: flush after each dispatch\n")
 	TEXT("2: flush after each tile\n"),
@@ -187,7 +189,7 @@ TAutoConsoleVariable<int32> CVarPathTracingFlushDispatch(
 TAutoConsoleVariable<int32> CVarPathTracingDispatchSize(
 	TEXT("r.PathTracing.DispatchSize"),
 	2048,
-	TEXT("Controls the tile size used when rendering the image. Reducing this value may prevent GPU timeouts for heavy renders. (default = 2048)"),
+	TEXT("Controls the tile size used when rendering the image. Reducing this value may prevent GPU timeouts for heavy renders."),
 	ECVF_RenderThreadSafe
 );
 
@@ -215,7 +217,7 @@ TAutoConsoleVariable<float> CVarPathTracingFilterWidth(
 TAutoConsoleVariable<int32> CVarPathTracingMISMode(
 	TEXT("r.PathTracing.MISMode"),
 	2,
-	TEXT("Selects the sampling technique for light integration (default = 2 (MIS enabled))\n")
+	TEXT("Selects the sampling technique for light integration.\n")
 	TEXT("0: Material sampling\n")
 	TEXT("1: Light sampling\n")
 	TEXT("2: MIS betwen material and light sampling (default)\n"),
@@ -225,7 +227,7 @@ TAutoConsoleVariable<int32> CVarPathTracingMISMode(
 TAutoConsoleVariable<int32> CVarPathTracingVolumeMISMode(
 	TEXT("r.PathTracing.VolumeMISMode"),
 	1,
-	TEXT("Selects the sampling technique for volumetric integration of local lighting (default = 1)\n")
+	TEXT("Selects the sampling technique for volumetric integration of local lighting\n")
 	TEXT("0: Density sampling\n")
 	TEXT("1: Light sampling (default)\n"),
 	ECVF_RenderThreadSafe
@@ -234,14 +236,14 @@ TAutoConsoleVariable<int32> CVarPathTracingVolumeMISMode(
 TAutoConsoleVariable<int32> CVarPathTracingMaxRaymarchSteps(
 	TEXT("r.PathTracing.MaxRaymarchSteps"),
 	768,
-	TEXT("Upper limit on the number of ray marching steps in volumes. This limit should not be hit in most cases, but raising it can reduce bias in case it is. (default = 256)."),
+	TEXT("Upper limit on the number of ray marching steps in volumes. This limit should not be hit in most cases, but raising it can reduce bias in case it is."),
 	ECVF_RenderThreadSafe
 );
 
 TAutoConsoleVariable<int32> CVarPathTracingMISCompensation(
 	TEXT("r.PathTracing.MISCompensation"),
 	1,
-	TEXT("Activates MIS compensation for skylight importance sampling. (default = 1 (enabled))\n")
+	TEXT("Activates MIS compensation for skylight importance sampling.\n")
 	TEXT("This option only takes effect when r.PathTracing.MISMode = 2\n"),
 	ECVF_RenderThreadSafe
 );
@@ -249,7 +251,7 @@ TAutoConsoleVariable<int32> CVarPathTracingMISCompensation(
 TAutoConsoleVariable<int32> CVarPathTracingSkylightCaching(
 	TEXT("r.PathTracing.SkylightCaching"),
 	1,
-	TEXT("Attempts to re-use skylight data between frames. (default = 1 (enabled))\n")
+	TEXT("Attempts to re-use skylight data between frames.\n")
 	TEXT("When set to 0, the skylight texture and importance samping data will be regenerated every frame. This is mainly intended as a benchmarking and debugging aid\n"),
 	ECVF_RenderThreadSafe
 );
@@ -257,7 +259,7 @@ TAutoConsoleVariable<int32> CVarPathTracingSkylightCaching(
 TAutoConsoleVariable<int32> CVarPathTracingVisibleLights(
 	TEXT("r.PathTracing.VisibleLights"),
 	0,
-	TEXT("Should light sources be visible to camera rays? (default = 0 (off))\n")
+	TEXT("Should light sources be visible to camera rays?\n")
 	TEXT("0: Hide lights from camera rays (default)\n")
 	TEXT("1: Make all lights visible to camera\n")
 	TEXT("2: Make skydome only visible to camera\n"),
@@ -267,14 +269,14 @@ TAutoConsoleVariable<int32> CVarPathTracingVisibleLights(
 TAutoConsoleVariable<int32> CVarPathTracingMaxSSSBounces(
 	TEXT("r.PathTracing.MaxSSSBounces"),
 	256,
-	TEXT("Sets the maximum number of bounces inside subsurface materials. Lowering this value can make subsurface scattering render too dim, while setting it too high can cause long render times.  (default = 256)"),
+	TEXT("Sets the maximum number of bounces inside subsurface materials. Lowering this value can make subsurface scattering render too dim, while setting it too high can cause long render times."),
 	ECVF_RenderThreadSafe
 );
 
 TAutoConsoleVariable<float> CVarPathTracingSSSGuidingRatio(
 	TEXT("r.PathTracing.SSSGuidingRatio"),
 	0.5f,
-	TEXT("Sets the ratio between classical random walks and walks guided towards the surface. A value of 0.0 corresponds to a purely classical random walk, while a value of 1.0 is fully guided towards the surface (at the expense of fireflies in non-flat regions of the model. (default = 0.5)"),
+	TEXT("Sets the ratio between classical random walks and walks guided towards the surface. A value of 0.0 corresponds to a purely classical random walk, while a value of 1.0 is fully guided towards the surface (at the expense of fireflies in non-flat regions of the model."),
 	ECVF_RenderThreadSafe
 );
 
@@ -289,7 +291,7 @@ TAutoConsoleVariable<float> CVarPathTracingMaxPathIntensity(
 TAutoConsoleVariable<int32> CVarPathTracingApproximateCaustics(
 	TEXT("r.PathTracing.ApproximateCaustics"),
 	1,
-	TEXT("When non-zero, the path tracer will approximate caustic paths to reduce noise. This reduces speckles and noise from low-roughness glass and metals. (default = 1 (enabled))"),
+	TEXT("When non-zero, the path tracer will approximate caustic paths to reduce noise. This reduces speckles and noise from low-roughness glass and metals."),
 	ECVF_RenderThreadSafe
 );
 
@@ -303,7 +305,7 @@ TAutoConsoleVariable<int32> CVarPathTracingEnableEmissive(
 TAutoConsoleVariable<int32> CVarPathTracingEnableCameraBackfaceCulling(
 	TEXT("r.PathTracing.EnableCameraBackfaceCulling"),
 	1,
-	TEXT("When non-zero, the path tracer will skip over backfacing triangles when tracing primary rays from the camera. (default = 1 (enabled))"),
+	TEXT("When non-zero, the path tracer will skip over backfacing triangles when tracing primary rays from the camera."),
 	ECVF_RenderThreadSafe
 );
 
@@ -340,14 +342,14 @@ TAutoConsoleVariable<int32> CVarPathTracingEnableReferenceClouds(
 TAutoConsoleVariable<int32> CVarPathTracingAtmosphereOpticalDepthLutResolution(
 	TEXT("r.PathTracing.AtmosphereOpticalDepthLUTResolution"),
 	512,
-	TEXT("Size of the square lookup texture used for transmittance calculations by the path tracer in reference atmosphere mode.  (default = 512)"),
+	TEXT("Size of the square lookup texture used for transmittance calculations by the path tracer in reference atmosphere mode."),
 	ECVF_RenderThreadSafe
 );
 
 TAutoConsoleVariable<int32> CVarPathTracingAtmosphereOpticalDepthLutNumSamples(
 	TEXT("r.PathTracing.AtmosphereOpticalDepthLUTNumSamples"),
 	16384,
-	TEXT("Number of ray marching samples used when building the transmittance lookup texture used for transmittance calculations by the path tracer in reference atmosphere mode.  (default = 16384)"),
+	TEXT("Number of ray marching samples used when building the transmittance lookup texture used for transmittance calculations by the path tracer in reference atmosphere mode."),
 	ECVF_RenderThreadSafe
 );
 
@@ -362,21 +364,21 @@ TAutoConsoleVariable<bool> CVarPathTracingEnableAtmosphereGround(
 TAutoConsoleVariable<int32> CVarPathTracingCloudAccelerationMapResolution(
 	TEXT("r.PathTracing.CloudAccelerationMap.Resolution"),
 	512,
-	TEXT("Size of the square texture used to accelerate cloud ray marching for the path tracer in reference atmosphere mode.  (default = 512)"),
+	TEXT("Size of the square texture used to accelerate cloud ray marching for the path tracer in reference atmosphere mode."),
 	ECVF_RenderThreadSafe
 );
 
 TAutoConsoleVariable<int32> CVarPathTracingCloudAccelerationMapNumSamples(
 	TEXT("r.PathTracing.CloudAccelerationMap.NumSamples"),
 	64,
-	TEXT("Number of ray marching samples used when building the cloud acceleration map for the path tracer in reference atmosphere mode.  (default = 64)"),
+	TEXT("Number of ray marching samples used when building the cloud acceleration map for the path tracer in reference atmosphere mode."),
 	ECVF_RenderThreadSafe
 );
 
 TAutoConsoleVariable<bool> CVarPathTracingCloudAccelerationMapVisualize(
 	TEXT("r.PathTracing.CloudAccelerationMap.Visualize"),
 	false,
-	TEXT("If true, replace clouds with a visualization of the acceleration map to help visualize it and fine tune its resolution (default = false)"),
+	TEXT("If true, replace clouds with a visualization of the acceleration map to help visualize it and fine tune its resolution."),
 	ECVF_RenderThreadSafe
 );
 
@@ -393,7 +395,7 @@ TAutoConsoleVariable<int> CVarPathTracingCloudMultipleScatterMode(
 TAutoConsoleVariable<float> CVarPathTracingCloudRoughnessCutoff(
 	TEXT("r.PathTracing.CloudRoughnessCutoff"),
 	0.05f,
-	TEXT("Do not evaluate volumetric clouds beyond this roughness level to improve performance (default=0.05)"),
+	TEXT("Do not evaluate volumetric clouds beyond this roughness level to improve performance."),
 	ECVF_RenderThreadSafe
 );
 
@@ -428,15 +430,15 @@ TAutoConsoleVariable<int32> CVarPathTracingSamplerType(
 TAutoConsoleVariable<int32> CVarPathTracingMultiGPU(
 	TEXT("r.PathTracing.MultiGPU"),
 	0,
-	TEXT("Run the path tracer using all available GPUs when enabled (default = 0)\n")
-	TEXT("Using this functionality in the editor requires -MaxGPUCount=N setting on the command line"),
+	TEXT("Run the path tracer using all available GPUs when enabled\n")
+	TEXT("Using this functionality in the editor requires -MaxGPUCount=N setting on the command line."),
 	ECVF_RenderThreadSafe
 );
 
 TAutoConsoleVariable<bool> CVarPathTracingAdjustMultiGPUPasses(
 	TEXT("r.PathTracing.AdjustMultiGPUPasses"),
 	true,
-	TEXT("Run extra passes per frame when multiple GPUs are active, to improve perf scaling as GPUs are added (default = true)\n"),
+	TEXT("Run extra passes per frame when multiple GPUs are active, to improve perf scaling as GPUs are added."),
 	ECVF_RenderThreadSafe
 );
 #endif  // WITH_MGPU
@@ -444,7 +446,7 @@ TAutoConsoleVariable<bool> CVarPathTracingAdjustMultiGPUPasses(
 TAutoConsoleVariable<bool> CVarPathTracingWiperMode(
 	TEXT("r.PathTracing.WiperMode"),
 	false,
-	TEXT("Enables wiper mode to render using the path tracer only in a region of the screen for debugging purposes (default = false, wiper mode disabled)"),
+	TEXT("Enables wiper mode to render using the path tracer only in a region of the screen for debugging purposes."),
 	ECVF_RenderThreadSafe 
 );
 
@@ -460,56 +462,56 @@ TAutoConsoleVariable<bool> CVarPathTracingProgressDisplay(
 TAutoConsoleVariable<int32> CVarPathTracingLightGridResolution(
 	TEXT("r.PathTracing.LightGridResolution"),
 	256,
-	TEXT("Controls the resolution of the 2D light grid used to cull irrelevant lights from lighting calculations (default = 256)\n"),
+	TEXT("Controls the resolution of the 2D light grid used to cull irrelevant lights from lighting calculations."),
 	ECVF_RenderThreadSafe
 );
 
 TAutoConsoleVariable<int32> CVarPathTracingLightGridMaxCount(
 	TEXT("r.PathTracing.LightGridMaxCount"),
 	128,
-	TEXT("Controls the maximum number of lights per cell in the 2D light grid. The minimum of this value and the number of lights in the scene is used. (default = 128)\n"),
+	TEXT("Controls the maximum number of lights per cell in the 2D light grid. The minimum of this value and the number of lights in the scene is used."),
 	ECVF_RenderThreadSafe
 );
 
 TAutoConsoleVariable<int32> CVarPathTracingLightGridAxis(
 	TEXT("r.PathTracing.LightGridAxis"),
 	-1,
-	TEXT("Choose the coordinate axis along which to project the light grid (default = -1, automatic)\n"),
+	TEXT("Choose the coordinate axis along which to project the light grid (default = -1, automatic)"),
 	ECVF_RenderThreadSafe
 );
 
 TAutoConsoleVariable<bool> CVarPathTracingUseDBuffer(
 	TEXT("r.PathTracing.UseDBuffer"),
 	true,
-	TEXT("Whether to support DBuffer functionality (default=true)"),
+	TEXT("Whether to support DBuffer functionality."),
 	ECVF_RenderThreadSafe
 );
 
 TAutoConsoleVariable<float> CVarPathTracingDecalRoughnessCutoff(
 	TEXT("r.PathTracing.DecalRoughnessCutoff"),
 	0.15f,
-	TEXT("Do not evaluate decals beyond this roughness level to improve performance (default=0.15)"),
+	TEXT("Do not evaluate decals beyond this roughness level to improve performance."),
 	ECVF_RenderThreadSafe
 );
 
 TAutoConsoleVariable<float> CVarPathTracingMeshDecalRoughnessCutoff(
 	TEXT("r.PathTracing.MeshDecalRoughnessCutoff"),
 	0.15f,
-	TEXT("Do not evaluate mesh decals beyond this roughness level to improve performance (default=0.15)"),
+	TEXT("Do not evaluate mesh decals beyond this roughness level to improve performance."),
 	ECVF_RenderThreadSafe
 );
 
 TAutoConsoleVariable<float> CVarPathTracingMeshDecalBias(
 	TEXT("r.PathTracing.MeshDecalBias"),
 	1.0f,
-	TEXT("Bias applied to mesh decal rays to avoid intersection with geometry (default = 1.0f)"),
+	TEXT("Bias applied to mesh decal rays to avoid intersection with geometry."),
 	ECVF_RenderThreadSafe
 );
 
 TAutoConsoleVariable<bool> CVarPathTracingLightFunctionColor(
 	TEXT("r.PathTracing.LightFunctionColor"),
 	true,
-	TEXT("Enables colored light function output (default = true)\n")
+	TEXT("Enables colored light function output\n")
 	TEXT("0: off (light function material output is converted to grayscale)\n")
 	TEXT("1: on (light function material output is used directly)\n"),
 	ECVF_RenderThreadSafe
@@ -518,7 +520,7 @@ TAutoConsoleVariable<bool> CVarPathTracingLightFunctionColor(
 TAutoConsoleVariable<bool> CVarPathTracingHeterogeneousVolumesRebuildEveryFrame(
 	TEXT("r.PathTracing.HeterogeneousVolumes.RebuildEveryFrame"),
 	true,
-	TEXT("Rebuilds volumetric acceleration structures every frame (default = true)\n"),
+	TEXT("Rebuilds volumetric acceleration structures every frame."),
 	ECVF_RenderThreadSafe
 );
 
@@ -581,7 +583,7 @@ TAutoConsoleVariable<int32> CVarPathTracingAdaptiveSampling(
 TAutoConsoleVariable<float> CVarPathTracingAdaptiveSamplingErrorThreshold(
 	TEXT("r.PathTracing.AdaptiveSampling.ErrorThreshold"),
 	0.001f,
-	TEXT("This is the target perceptual error threshold. Once a pixel's error falls below this value, it will not be sampled again (default: 0.001)\n"),
+	TEXT("This is the target perceptual error threshold. Once a pixel's error falls below this value, it will not be sampled again.\n"),
 	ECVF_RenderThreadSafe
 );
 
@@ -1283,22 +1285,27 @@ class FPathTracingRG : public FGlobalShader
 		FPermutationDomain PermutationVector(Parameters.PermutationId);
 		if (bUseExperimental == false)
 		{
-			
+			// default case: experimental features and rarely used modes are not compiled
 			if (PermutationVector.Get<FAdaptiveSampling>())
 			{
-				// adaptive sampling is experimental
+				return false;
+			}
+			if (!PermutationVector.Get<FTraceOpaqueFirst>())
+			{
+				return false;
+			}
+			if (!PermutationVector.Get<FCompactionType>())
+			{
 				return false;
 			}
 		}
-		else
+
+		if (PermutationVector.Get<FCloudShader>())
 		{
-			if (PermutationVector.Get<FCloudShader>())
+			// the cloud shader version can only be supported if the platform supports callable shaders
+			if (!ShouldCompileRayTracingCallableShadersForProject(Parameters.Platform))
 			{
-				// the cloud shader version can only be supported if the platform supports callable shaders
-				if (!ShouldCompileRayTracingCallableShadersForProject(Parameters.Platform))
-				{
-					return false;
-				}
+				return false;
 			}
 		}
 
@@ -3265,7 +3272,7 @@ static FPathTracingRG::FPermutationDomain GetPathTracingRGPermutation(const FVie
 	const bool bUseCloudShader = uint32(View.PathTracingVolumetricCloudCallableShaderIndex) < Scene.RayTracingSBT.NumCallableShaderSlots;
 	const bool bHasComplexSpecialRenderPath = Substrate::IsSubstrateEnabled() && Scene.SubstrateSceneData.bUsesComplexSpecialRenderPath;
 	const bool bUseSER = GRHIGlobals.SupportsShaderExecutionReordering && CVarPathTracingShaderExecutionReordering.GetValueOnRenderThread();
-	const bool bUseTraceOpaqueFirst = CVarPathTracingTraceOpaqueFirst.GetValueOnRenderThread();
+	const bool bUseTraceOpaqueFirst = !bUseExperimental || CVarPathTracingTraceOpaqueFirst.GetValueOnRenderThread();
 
 	FPathTracingRG::FPermutationDomain Out;
 	Out.Set<FPathTracingRG::FCompactionType>(bUseCompaction);
@@ -4032,7 +4039,7 @@ void FDeferredShadingSceneRenderer::RenderPathTracing(
 			// We are writing to the texture, we'll need to extract it...
 			bNeedsTextureExtract = true;
 
-			int32 CompactionDepth = CVarPathTracingCompactionDepth.GetValueOnRenderThread();
+			int32 CompactionDepth = bUseExperimental ? CVarPathTracingCompactionDepth.GetValueOnRenderThread() : -1;
 			if (CompactionDepth < 0 || CompactionDepth > int32(Config.PathTracingData.MaxBounces))
 			{
 				CompactionDepth = Config.PathTracingData.MaxBounces;
