Added the same sorting logic for masked materials in MobileBasePass as in BasePassRendering. This fixes the issue where masked materials rendered with overlay were not drawing in the correct order, causing visual problems.

#okforversepublic
#tests Console(Fonrt end with shoes cosmetic)

diff --git a/Engine/Source/Runtime/Renderer/Private/MobileBasePass.cpp b/Engine/Source/Runtime/Renderer/Private/MobileBasePass.cpp
--- a/Engine/Source/Runtime/Renderer/Private/MobileBasePass.cpp
+++ b/Engine/Source/Runtime/Renderer/Private/MobileBasePass.cpp
@@ -671,13 +671,20 @@ void MobileBasePass::SetTranslucentRenderState(FMeshPassProcessorRenderState& Dr
 	}
 }
 
-static FMeshDrawCommandSortKey GetBasePassStaticSortKey(const bool bIsMasked, bool bBackground, const FMeshMaterialShader* VertexShader, const FMeshMaterialShader* PixelShader)
+static FMeshDrawCommandSortKey GetBasePassStaticSortKey(EDepthDrawingMode EarlyZPassMode, const bool bIsMasked, bool bBackground, const FMeshMaterialShader* VertexShader, const FMeshMaterialShader* PixelShader)
 {
 	FMeshDrawCommandSortKey SortKey;
-	SortKey.BasePass.Masked = (bIsMasked ? 1 : 0);
 	SortKey.BasePass.Background = (bBackground ? 1 : 0); // background flag in second bit
 	SortKey.BasePass.VertexShaderHash = (VertexShader ? VertexShader->GetSortKey() : 0) & 0xFFFF;
 	SortKey.BasePass.PixelShaderHash = PixelShader ? PixelShader->GetSortKey() : 0;
+	if (EarlyZPassMode != DDM_None)
+	{
+		SortKey.BasePass.Masked = bIsMasked ? 0 : 1;
+	}
+	else
+	{
+		SortKey.BasePass.Masked = bIsMasked ? 1 : 0;
+	}
 	return SortKey;
 }
 
@@ -750,6 +757,7 @@ FMobileBasePassMeshProcessor::FMobileBasePassMeshProcessor(
 	, bTranslucentBasePass(InTranslucencyPassType != ETranslucencyPass::TPT_MAX)
 	, bDeferredShading(IsMobileDeferredShadingEnabled(GetFeatureLevelShaderPlatform(ERHIFeatureLevel::ES3_1)))
 	, bPassUsesDeferredShading(bDeferredShading && !bTranslucentBasePass)
+	, EarlyZPassMode(Scene ? Scene->EarlyZPassMode : DDM_None)
 {
 }
 
@@ -907,7 +915,7 @@ bool FMobileBasePassMeshProcessor::Process(
 		bool bBackground = PrimitiveSceneProxy ? PrimitiveSceneProxy->TreatAsBackgroundForOcclusion() : false;
 		// Default static sort key separates masked and non-masked geometry, generic mesh sorting will also sort by PSO
 		// if platform wants front to back sorting, this key will be recomputed in InitViews
-		SortKey = GetBasePassStaticSortKey(bIsMasked, bBackground, BasePassShaders.VertexShader.GetShader(), BasePassShaders.PixelShader.GetShader());
+		SortKey = GetBasePassStaticSortKey(EarlyZPassMode, bIsMasked, bBackground, BasePassShaders.VertexShader.GetShader(), BasePassShaders.PixelShader.GetShader());
 	}
 	
 	const FMeshDrawingPolicyOverrideSettings OverrideSettings = ComputeMeshOverrideSettings(MeshBatch);
diff --git a/Engine/Source/Runtime/Renderer/Private/MobileBasePassRendering.h b/Engine/Source/Runtime/Renderer/Private/MobileBasePassRendering.h
--- a/Engine/Source/Runtime/Renderer/Private/MobileBasePassRendering.h
+++ b/Engine/Source/Runtime/Renderer/Private/MobileBasePassRendering.h
@@ -510,6 +510,7 @@ private:
 	const bool bDeferredShading; 
 	// Whether this pass uses deferred shading
 	const bool bPassUsesDeferredShading; 
+	EDepthDrawingMode EarlyZPassMode;
 };
 
 ENUM_CLASS_FLAGS(FMobileBasePassMeshProcessor::EFlags);
