Ensure third-party temporal upscalers respect the TemporalAA show flag.
This prevents a crash when using SceneCapture components alongside third-party temporal upscalers.

diff --git a/Engine/Source/Runtime/Renderer/Private/SceneRendering.cpp b/Engine/Source/Runtime/Renderer/Private/SceneRendering.cpp
--- a/Engine/Source/Runtime/Renderer/Private/SceneRendering.cpp
+++ b/Engine/Source/Runtime/Renderer/Private/SceneRendering.cpp
@@ -2843,13 +2843,16 @@ FSceneRenderer::FSceneRenderer(const FSceneViewFamily* InViewFamily, FHitProxyCo
 			ViewFamily.TemporalUpscalerInterface = nullptr;
 			ViewFamily.SetTemporalUpscalerInterface(InViewFamily->TemporalUpscalerInterface->Fork_GameThread(ViewFamily));
 
-			for (FViewInfo& View : Views)
+			if (ViewFamily.EngineShowFlags.TemporalAA)
 			{
-				View.AntiAliasingMethod = AAM_TemporalAA;
-				View.PrimaryScreenPercentageMethod = EPrimaryScreenPercentageMethod::TemporalUpscale;
+				for (FViewInfo& View : Views)
+				{
+					View.AntiAliasingMethod = AAM_TemporalAA;
+					View.PrimaryScreenPercentageMethod = EPrimaryScreenPercentageMethod::TemporalUpscale;
 #if DO_CHECK || USING_CODE_ANALYSIS
-				check(View.VerifyMembersChecks());
+					check(View.VerifyMembersChecks());
 #endif
+				}
 			}
 		}
 
