[VSM] Check light proxy pointer before retrieving light name for NPF debug view

diff --git a/Engine/Source/Runtime/Renderer/Private/VirtualShadowMaps/VirtualShadowMapCacheManager.cpp b/Engine/Source/Runtime/Renderer/Private/VirtualShadowMaps/VirtualShadowMapCacheManager.cpp
--- a/Engine/Source/Runtime/Renderer/Private/VirtualShadowMaps/VirtualShadowMapCacheManager.cpp
+++ b/Engine/Source/Runtime/Renderer/Private/VirtualShadowMaps/VirtualShadowMapCacheManager.cpp
@@ -904,7 +904,11 @@ void FVirtualShadowMapArrayCacheManager::InitExtension(FScene& InScene)
 				FString LightNameWithLevel;
 				if (LightId <= uint64(Scene.Lights.GetMaxIndex()))
 				{
-					FSceneRenderer::GetLightNameForDrawEvent(Scene.Lights[LightId].LightSceneInfo->Proxy, LightNameWithLevel);
+					const FLightSceneInfo* LightSceneInfo = Scene.Lights[LightId].LightSceneInfo;
+					if (LightSceneInfo && LightSceneInfo->Proxy)
+					{
+						FSceneRenderer::GetLightNameForDrawEvent(LightSceneInfo->Proxy, LightNameWithLevel);
+					}
 				}
 				NPFDiagnosticMessages.Add(FString::Printf(TEXT("%8d | %48s | %12d"), VSMId, *LightNameWithLevel, ClusterCount));
 				NPFDiagnosticTimer.FindOrAdd(VSMId) = 180;
