Ensure ViewFamily.SupportsScreenPercentage() is used rather than ViewFamily.EngineShowFlags.ScreenPercentage in PrepareViewRectsForRendering()

diff --git a/Engine/Source/Runtime/Renderer/Private/SceneRendering.cpp b/Engine/Source/Runtime/Renderer/Private/SceneRendering.cpp
--- a/Engine/Source/Runtime/Renderer/Private/SceneRendering.cpp
+++ b/Engine/Source/Runtime/Renderer/Private/SceneRendering.cpp
@@ -3282,7 +3282,7 @@ void FSceneRenderer::PrepareViewRectsForRendering()
 	}
 
 	// If not supporting screen percentage, bypass all computation.
-	if (!ViewFamily.EngineShowFlags.ScreenPercentage)
+	if (!ViewFamily.SupportsScreenPercentage())
 	{
 		DynamicResolutionFractions[GDynamicPrimaryResolutionFraction] = 1.0f;
 
