Path Tracing: Improvements to PrimaryRays visualize mode

Launch seperate kernels instead of accumulating inside the shader. This makes the cost more linear and reduces register pressure. This should also make it easier to clarify the read/write relationship between passes for guiding data.

Factor in SSS contribution as part of the diffuse for now (until we port over the random walk code).

Track diffuse/specular albedo through mirror reflections. This improves the stability in mirrors, though isn't perfect yet. Additional denoiser signals will likely be needed.

#rnx

diff --git a/Engine/Shaders/Private/PathTracing/PathTracingDebug.usf b/Engine/Shaders/Private/PathTracing/PathTracingDebug.usf
--- a/Engine/Shaders/Private/PathTracing/PathTracingDebug.usf
+++ b/Engine/Shaders/Private/PathTracing/PathTracingDebug.usf
@@ -50,6 +50,8 @@ uint UsePathGuiding;
 
 int DebugMode;
 int SamplesPerPixel;
+int SampleIndex;
+float SampleBlendFactor;
 int MaxBounces;
 int EnableEmissive;
 
@@ -300,20 +302,8 @@ FPathTracingResult PathTracingDebugCore(FRayDesc Ray, uint2 LaunchIndex, int Sam
 			const float3 ViewZ = View.ViewToTranslatedWorld[2].xyz;
 			const float LinearZ = dot(TranslatedWorldPos, ViewZ);
 			Result.Z = ConvertToDeviceZ(LinearZ);
-			Result.Radiance.a = 0; // opaque alpha (we aren't handling translucency right now)
-			// G-Buffer like data for denoisers
-			Result.WorldNormal = PackedPayload.UnpackWorldNormal();
-#if PATHTRACING_SUBSTRATE_PAYLOAD
-			float4 RoughnessData = PackedPayload.UnpackRoughnessAniso();
-			Result.Roughness = lerp(RoughnessData.x, RoughnessData.y, RoughnessData.z); // mix both lobes to get an average value
-			float3 DiffuseColor = PackedPayload.UnpackDiffuseColor();
-			float3 SpecularColor = PackedPayload.UnpackSpecularColor();
-#else
-			Result.Roughness = PackedPayload.UnpackRoughness();
-			Result.DiffuseColor = PackedPayload.UnpackBaseColor() * (1 - PackedPayload.UnpackMetallic());
-			Result.SpecularColor = lerp(0.08 * PackedPayload.UnpackSpecular(), PackedPayload.UnpackBaseColor(), PackedPayload.UnpackMetallic());
-#endif
 			Result.LinearDepth = LinearZ;
+			Result.Radiance.a = 0; // opaque alpha (we aren't handling translucency right now)
 		}
 
 		// Add radiance (emissive surfaces)
@@ -326,6 +316,34 @@ FPathTracingResult PathTracingDebugCore(FRayDesc Ray, uint2 LaunchIndex, int Sam
 		// Compute direct lighting
 		FPathTracingPayload Payload = UnpackPathTracingPayload(PackedPayload, Ray);
 		AdjustPayloadAfterUnpack(Payload, true);
+		FSSSRandomWalkInfo SSS = GetMaterialSSSInfo(Payload, Ray.Direction);
+		if (any(SSS.Weight > 0))
+		{
+			// If the material has a valid SSS lobe, just merge it into diffuse for now
+			Payload.DiffuseColor += SSS.Color;
+		}
+		float DenoiserContrib = DenoiserAOVWeight(PathRoughness);
+		if (DenoiserContrib > 0)
+		{
+			float3 DenoiserThroughput = PathThroughput * DenoiserContrib;
+
+			// G-Buffer like data for denoisers
+#if PATHTRACING_SUBSTRATE_PAYLOAD
+			float4 RoughnessData = PackedPayload.UnpackRoughnessAniso();
+			float Roughness = lerp(RoughnessData.x, RoughnessData.y, RoughnessData.z); // mix both lobes to get an average value
+			DenoiserThroughput *= Payload.WeightV * Payload.BSDFOpacity * lerp(1.0, Payload.TransmittanceN, Payload.CoverageAboveAlongN);
+#else
+			float Roughness = Payload.Roughness;
+#endif
+			float NormThroughput = Luminance(DenoiserThroughput);
+			float3 DiffuseColor = Payload.DiffuseColor;
+			float3 SpecularColor = Payload.SpecularColor;
+
+			Result.DiffuseColor += DenoiserThroughput * Payload.DiffuseColor;
+			Result.SpecularColor += DenoiserThroughput * Payload.SpecularColor * DenoiserRoughnessWeight(Roughness, PathRoughness);
+			Result.Roughness += NormThroughput * Roughness;
+			Result.WorldNormal += NormThroughput* PackedPayload.UnpackWorldNormal();
+		}
 
 		// Generate main random number for direct+indirect sampling
 		float4 RandSample = RandSequence.Get4D();
@@ -801,32 +819,21 @@ RAY_TRACING_ENTRY_RAYGEN(PathTracingDebugRG)
 	const FRayDesc CameraRay = CreatePrimaryRay(ViewportUV);
 
 #if PATH_TRACER_PRIMARY_RAYS
-	FPathTracingResult Result = (FPathTracingResult) 0;
-	
-	// TODO: launch separate kernels and accumulate instead?
-	LOOP
-	for (int SampleIndex = 0; SampleIndex < SamplesPerPixel; SampleIndex++)
+	FPathTracingResult Result = PathTracingDebugCore(CameraRay, LaunchIndex, SampleIndex);
+	if (SampleIndex > 0)
 	{
-		FPathTracingResult SampleResult = PathTracingDebugCore(CameraRay, LaunchIndex, SampleIndex);
-
-		Result.Radiance += SampleResult.Radiance;
-		Result.Z = max(SampleResult.Z, Result.Z); // inverted Z so larger numbers are closer
-
+		// Blend with previous value
+		Result.Radiance = lerp(RWSceneColor[DispatchIdx], Result.Radiance, SampleBlendFactor);
+		Result.Z = max(RWSceneDepth[DispatchIdx], Result.Z); // Device Z is inverted so larger numbers are closer
+	
 		// G-Buffer like data for denoisers
-		Result.LinearDepth += SampleResult.LinearDepth;
-		Result.DiffuseColor += SampleResult.DiffuseColor;
-		Result.SpecularColor += SampleResult.SpecularColor;
-		Result.WorldNormal += SampleResult.WorldNormal;
-		Result.Roughness += SampleResult.Roughness;
+		// NOTE: averaging here is not strictly necessary, but will be useful for stochastic situations
+		Result.LinearDepth   = lerp(RWSceneLinearDepth[DispatchIdx]			, Result.LinearDepth	, SampleBlendFactor); // TODO: would it be better to use min? only makes a difference in stochastic cases
+		Result.DiffuseColor  = lerp(RWSceneDiffuseColor[DispatchIdx].xyz	, Result.DiffuseColor	, SampleBlendFactor);
+		Result.SpecularColor = lerp(RWSceneSpecularColor[DispatchIdx].xyz	, Result.SpecularColor	, SampleBlendFactor);
+		Result.WorldNormal   = lerp(RWSceneNormal[DispatchIdx].xyz			, Result.WorldNormal	, SampleBlendFactor); // TODO: should we normalize? only makes a difference in stochastic cases
+		Result.Roughness     = lerp(RWSceneRoughness[DispatchIdx]			, Result.Roughness		, SampleBlendFactor);
 	}
-	const float NormalizationFactor = 1.0 / SamplesPerPixel;
-	Result.Radiance *= NormalizationFactor;
-	Result.LinearDepth *= NormalizationFactor;
-	Result.DiffuseColor *= NormalizationFactor;
-	Result.SpecularColor *= NormalizationFactor;
-	Result.WorldNormal = normalize(Result.WorldNormal); // averaged normal better?
-	Result.Roughness *= NormalizationFactor;
-
 #else
 	FPathTracingResult Result = PathTracingDebugCore(CameraRay, LaunchIndex);
 #endif
diff --git a/Engine/Source/Runtime/Renderer/Private/PathTracing.cpp b/Engine/Source/Runtime/Renderer/Private/PathTracing.cpp
--- a/Engine/Source/Runtime/Renderer/Private/PathTracing.cpp
+++ b/Engine/Source/Runtime/Renderer/Private/PathTracing.cpp
@@ -1479,6 +1479,8 @@ class FPathTracingDebugRG : public FGlobalShader
 		SHADER_PARAMETER_RDG_TEXTURE(Texture2D<uint>, RandomSequenceSpaceFillingCurve)
 		SHADER_PARAMETER(int32, DebugMode)
 		SHADER_PARAMETER(int32, SamplesPerPixel)
+		SHADER_PARAMETER(int32, SampleIndex)
+		SHADER_PARAMETER(float, SampleBlendFactor)
 		SHADER_PARAMETER(int32, MaxBounces)
 		SHADER_PARAMETER(int32, EnableEmissive)
 		SHADER_PARAMETER(uint32, MarkovChainStateGridSize)
@@ -3539,68 +3541,89 @@ void FDeferredShadingSceneRenderer::RenderPathTracing(
 			AddClearUAVPass(GraphBuilder, GraphBuilder.CreateUAV(MarkovChainStateGrid), 0);
 		}
 
-		FPathTracingDebugRG::FParameters* PassParameters = GraphBuilder.AllocParameters<FPathTracingDebugRG::FParameters>();
-		PassParameters->RWSceneColor = GraphBuilder.CreateUAV(SceneColorOutputTexture);
-		PassParameters->RWSceneDepth = GraphBuilder.CreateUAV(DepthTexture);
-		PassParameters->TLAS = Scene->RayTracingScene.GetLayerView(ERayTracingSceneLayer::Base);
-		PassParameters->UsePathGuiding = CVarPathTracingVisualizeIndirectGuiding.GetValueOnRenderThread();
-		PassParameters->MarkovChainStateGrid = GraphBuilder.CreateUAV(MarkovChainStateGrid);
-		PassParameters->MarkovChainStateGridSize = MarkovChainStateGridHashSize;
-		PassParameters->ViewUniformBuffer = View.ViewUniformBuffer;
-		PassParameters->RandomSequenceSpaceFillingCurve = GSystemTextures.GetSpaceFillingCurveTexture(GraphBuilder);
-		PassParameters->DebugMode = DebugMode;
+		const bool bUsePrimaryRays = DebugMode == PATH_TRACER_DEBUG_VIZ_PRIMARY_RAYS;
+		const int SamplesPerPixel = bUsePrimaryRays ? FMath::Max(CVarPathTracingVisualizeSamplesPerPixel.GetValueOnRenderThread(), 1) : 1;
 
-		PassParameters->SamplesPerPixel = CVarPathTracingVisualizeSamplesPerPixel.GetValueOnRenderThread();
-		PassParameters->MaxBounces = CVarPathTracingVisualizeMaxBounces.GetValueOnRenderThread();
-		PassParameters->EnableEmissive = CVarPathTracingVisualizeEnableEmissive.GetValueOnRenderThread();
-		SetLightParameters(
-			GraphBuilder,
-			Scene,
-			View,
-			true,
-			false,
-			&PassParameters->SkylightParameters,
-			&PassParameters->LightGridParameters,
-			&PassParameters->SceneVisibleLightCount,
-			&PassParameters->SceneLightCount,
-			&PassParameters->SceneLights
-		);				
-
-		PassParameters->RWSceneDiffuseColor = GraphBuilder.CreateUAV(DiffuseColorTexture);
-		PassParameters->RWSceneSpecularColor = GraphBuilder.CreateUAV(SpecularColorTexture);
-		PassParameters->RWSceneNormal = GraphBuilder.CreateUAV(NormalTexture);
-		PassParameters->RWSceneRoughness = GraphBuilder.CreateUAV(RoughnessTexture);
-		PassParameters->RWSceneLinearDepth = GraphBuilder.CreateUAV(LinearDepthTexture);
-
-		PassParameters->DecalParameters = View.RayTracingDecalUniformBuffer;
-		PassParameters->Scene = View.GetSceneUniforms().GetBuffer(GraphBuilder);
-		PassParameters->NaniteRayTracing = Nanite::GRayTracingManager.GetUniformBuffer();
+		FPathTracingDebugRG::FParameters* FirstPassParameters = nullptr;
+		for (int SampleIndex = 0; SampleIndex < SamplesPerPixel; SampleIndex++)
+		{
+			// TODO: can we do better than making a new set of parameters per sample?
+			FPathTracingDebugRG::FParameters* PassParameters = GraphBuilder.AllocParameters<FPathTracingDebugRG::FParameters>();
+			if (FirstPassParameters == nullptr)
+			{
+				PassParameters->RWSceneColor = GraphBuilder.CreateUAV(SceneColorOutputTexture);
+				PassParameters->RWSceneDepth = GraphBuilder.CreateUAV(DepthTexture);
+				PassParameters->TLAS = Scene->RayTracingScene.GetLayerView(ERayTracingSceneLayer::Base);
+				PassParameters->UsePathGuiding = CVarPathTracingVisualizeIndirectGuiding.GetValueOnRenderThread();
+				PassParameters->MarkovChainStateGrid = GraphBuilder.CreateUAV(MarkovChainStateGrid);
+				PassParameters->MarkovChainStateGridSize = MarkovChainStateGridHashSize;
+				PassParameters->ViewUniformBuffer = View.ViewUniformBuffer;
+				PassParameters->RandomSequenceSpaceFillingCurve = GSystemTextures.GetSpaceFillingCurveTexture(GraphBuilder);
+				PassParameters->DebugMode = DebugMode;
+
+				PassParameters->SamplesPerPixel = SamplesPerPixel;
+				PassParameters->MaxBounces = CVarPathTracingVisualizeMaxBounces.GetValueOnRenderThread();
+				PassParameters->EnableEmissive = CVarPathTracingVisualizeEnableEmissive.GetValueOnRenderThread();
+				{
+					SetLightParameters(
+						GraphBuilder,
+						Scene,
+						View,
+						true,
+						false,
+						&PassParameters->SkylightParameters,
+						&PassParameters->LightGridParameters,
+						&PassParameters->SceneVisibleLightCount,
+						&PassParameters->SceneLightCount,
+						&PassParameters->SceneLights
+					);
+				}
 
-		const bool bUsePrimaryRays = DebugMode == PATH_TRACER_DEBUG_VIZ_PRIMARY_RAYS;
-		const bool bUseSER = GRHIGlobals.SupportsShaderExecutionReordering && CVarPathTracingShaderExecutionReordering.GetValueOnRenderThread();
-		FPathTracingDebugRG::FPermutationDomain PermutationVector;
-		PermutationVector.Set<FPathTracingDebugRG::FUseSER>(bUseSER && bUsePrimaryRays);
-		PermutationVector.Set<FPathTracingDebugRG::FUsePrimaryRays>(bUsePrimaryRays);
-		TShaderMapRef<FPathTracingDebugRG> RayGenShader(View.ShaderMap, PermutationVector);
-		ClearUnusedGraphResources(RayGenShader, PassParameters);
-		GraphBuilder.AddPass(
-			RDG_EVENT_NAME("Path Tracing Visualize"),
-		  	PassParameters,
-		  	ERDGPassFlags::Compute,
-			[PassParameters, RayGenShader, &View, DispatchResX, DispatchResY](FRHICommandList& RHICmdList)
+				PassParameters->RWSceneDiffuseColor = GraphBuilder.CreateUAV(DiffuseColorTexture);
+				PassParameters->RWSceneSpecularColor = GraphBuilder.CreateUAV(SpecularColorTexture);
+				PassParameters->RWSceneNormal = GraphBuilder.CreateUAV(NormalTexture);
+				PassParameters->RWSceneRoughness = GraphBuilder.CreateUAV(RoughnessTexture);
+				PassParameters->RWSceneLinearDepth = GraphBuilder.CreateUAV(LinearDepthTexture);
+
+				PassParameters->DecalParameters = View.RayTracingDecalUniformBuffer;
+				PassParameters->Scene = View.GetSceneUniforms().GetBuffer(GraphBuilder);
+				PassParameters->NaniteRayTracing = Nanite::GRayTracingManager.GetUniformBuffer();
+
+				FirstPassParameters = PassParameters;
+			}
+			else
 			{
-				FRHIBatchedShaderParameters& GlobalResources = RHICmdList.GetScratchShaderParameters();
-				SetShaderParameters(GlobalResources, RayGenShader, *PassParameters);
+				// Already made one set of parameters, re-use it
+				*PassParameters = *FirstPassParameters;
+			}
+			PassParameters->SampleIndex = SampleIndex;
+			PassParameters->SampleBlendFactor = 1.0f / (SampleIndex + 1);
+			const bool bUseSER = GRHIGlobals.SupportsShaderExecutionReordering && CVarPathTracingShaderExecutionReordering.GetValueOnRenderThread();
+			FPathTracingDebugRG::FPermutationDomain PermutationVector;
+			PermutationVector.Set<FPathTracingDebugRG::FUseSER>(bUseSER && bUsePrimaryRays);
+			PermutationVector.Set<FPathTracingDebugRG::FUsePrimaryRays>(bUsePrimaryRays);
+			TShaderMapRef<FPathTracingDebugRG> RayGenShader(View.ShaderMap, PermutationVector);
+			ClearUnusedGraphResources(RayGenShader, PassParameters);
+			GraphBuilder.AddPass(
+				RDG_EVENT_NAME("Path Tracing Visualize [%d/%d]", SampleIndex + 1, SamplesPerPixel),
+				PassParameters,
+				ERDGPassFlags::Compute,
+				[PassParameters, RayGenShader, &View, DispatchResX, DispatchResY](FRHICommandList& RHICmdList)
+				{
+					FRHIBatchedShaderParameters& GlobalResources = RHICmdList.GetScratchShaderParameters();
+					SetShaderParameters(GlobalResources, RayGenShader, *PassParameters);
 				
-				TOptional<FScopedUniformBufferStaticBindings> StaticUniformBufferScope = RayTracing::BindStaticUniformBufferBindings(View, PassParameters->Scene->GetRHI(), PassParameters->NaniteRayTracing->GetRHI(), RHICmdList);
-
-				RHICmdList.RayTraceDispatch(
-					View.MaterialRayTracingData.PipelineState,
-					RayGenShader.GetRayTracingShader(),
-					View.MaterialRayTracingData.ShaderBindingTable, GlobalResources,
-					DispatchResX, DispatchResY
-				);
-			});
+					TOptional<FScopedUniformBufferStaticBindings> StaticUniformBufferScope = RayTracing::BindStaticUniformBufferBindings(View, PassParameters->Scene->GetRHI(), PassParameters->NaniteRayTracing->GetRHI(), RHICmdList);
+
+					RHICmdList.RayTraceDispatch(
+						View.MaterialRayTracingData.PipelineState,
+						RayGenShader.GetRayTracingShader(),
+						View.MaterialRayTracingData.ShaderBindingTable, GlobalResources,
+						DispatchResX, DispatchResY
+					);
+				}
+			);
+		}
 
 		{
 			FPathTracingCopyDepthPS::FParameters* DisplayParameters = GraphBuilder.AllocParameters<FPathTracingCopyDepthPS::FParameters>();
