Get the ShaderPlatform from View and not FShaderTarget

diff --git a/Engine/Source/Runtime/Renderer/Private/SceneOcclusion.h b/Engine/Source/Runtime/Renderer/Private/SceneOcclusion.h
--- a/Engine/Source/Runtime/Renderer/Private/SceneOcclusion.h
+++ b/Engine/Source/Runtime/Renderer/Private/SceneOcclusion.h
@@ -104,7 +104,7 @@ private:
 	void SetParametersInternal(FRHIBatchedShaderParameters& BatchedParameters, const FViewInfo& View, const FVector4f StencilingSpherePosAndScale[2])
 	{
 		// Mobile renderer sets View UB through static binding
-		if (!IsMobilePlatform(GetShaderPlatform()))
+		if (!IsMobilePlatform(View.GetShaderPlatform()))
 		{
 			FGlobalShader::SetParameters<FViewUniformShaderParameters>(BatchedParameters, View.ViewUniformBuffer);
 		}
