[SingleLayerWater] Added missing stencil clear

diff --git a/Engine/Source/Runtime/Renderer/Private/DeferredShadingRenderer.cpp b/Engine/Source/Runtime/Renderer/Private/DeferredShadingRenderer.cpp
--- a/Engine/Source/Runtime/Renderer/Private/DeferredShadingRenderer.cpp
+++ b/Engine/Source/Runtime/Renderer/Private/DeferredShadingRenderer.cpp
@@ -3361,9 +3361,7 @@ void FDeferredShadingSceneRenderer::Render(FRDGBuilder& GraphBuilder, const FSce
 
 			RenderSingleLayerWater(GraphBuilder, Views, SceneTextures, SingleLayerWaterPrePassResult, bShouldRenderVolumetricCloud, SceneWithoutWaterTextures, LumenFrameTemporaries, bIsCameraUnderWater);
 
-			// Replace main depth texture with the output of the SLW depth prepass which contains the scene + water.
-			// Note: Stencil now has all water bits marked with 1. As long as no other passes after this point want to read the depth buffer,
-			// a stencil clear should not be necessary here.
+			// Replace main depth texture with the output of the SLW depth prepass which contains the scene + water. Stencil is cleared to 0.
 			if (SingleLayerWaterPrePassResult)
 			{
 				SceneTextures.Depth = SingleLayerWaterPrePassResult->DepthPrepassTexture;
diff --git a/Engine/Source/Runtime/Renderer/Private/SingleLayerWaterRendering.cpp b/Engine/Source/Runtime/Renderer/Private/SingleLayerWaterRendering.cpp
--- a/Engine/Source/Runtime/Renderer/Private/SingleLayerWaterRendering.cpp
+++ b/Engine/Source/Runtime/Renderer/Private/SingleLayerWaterRendering.cpp
@@ -1080,9 +1080,9 @@ FSingleLayerWaterPrePassResult* FDeferredShadingSceneRenderer::RenderSingleLayer
 
 			// Set depth test to always pass and stencil test to replace all pixels with zero, essentially also clearing stencil while doing the depth copy.
 			FRHIDepthStencilState* DepthStencilState = TStaticDepthStencilState<
-				true, CF_Always,										// depth
-				true, CF_Always, SO_Replace, SO_Replace, SO_Replace,	// frontface stencil
-				true, CF_Always, SO_Replace, SO_Replace, SO_Replace		// backface stencil
+				true, CF_Always,								// depth
+				true, CF_Always, SO_Zero, SO_Zero, SO_Zero,		// frontface stencil
+				true, CF_Always, SO_Zero, SO_Zero, SO_Zero		// backface stencil
 			>::GetRHI();
 
 			FPixelShaderUtils::AddFullscreenPass(
@@ -1094,8 +1094,7 @@ FSingleLayerWaterPrePassResult* FDeferredShadingSceneRenderer::RenderSingleLayer
 				Viewport,
 				nullptr, /*BlendState*/
 				nullptr, /*RasterizerState*/
-				DepthStencilState,
-				0 /*StencilRef*/);
+				DepthStencilState);
 
 			// The above copy technique loses HTILE data during the copy, so until AddCopyTexturePass() supports depth buffer copies on all platforms,
 			// this is the best we can do.
@@ -1219,7 +1218,7 @@ FSingleLayerWaterPrePassResult* FDeferredShadingSceneRenderer::RenderSingleLayer
 				nullptr /*BlendState*/,
 				nullptr /*RasterizerState*/,
 				TStaticDepthStencilState<false, CF_Always, true, CF_Equal>::GetRHI(),
-				1 /*StencilRef*/);
+				STENCIL_SANDBOX_MASK);
 		}
 
 		// Write near plane depth for all water pixels where we want to skip all shading for the scene behind them
@@ -1233,7 +1232,7 @@ FSingleLayerWaterPrePassResult* FDeferredShadingSceneRenderer::RenderSingleLayer
 			PassParameters->View = View.ViewUniformBuffer;
 			PassParameters->WaterRefractionCullingTexture = Result->RefractionMaskTexture;
 			PassParameters->WaterDepthTexture = Result->DepthPrepassTexture.Resolve;
-			PassParameters->RenderTargets.DepthStencil = FDepthStencilBinding(SceneTextures.Depth.Target, ERenderTargetLoadAction::ELoad, ERenderTargetLoadAction::ELoad, FExclusiveDepthStencil::DepthWrite_StencilNop);
+			PassParameters->RenderTargets.DepthStencil = FDepthStencilBinding(SceneTextures.Depth.Target, ERenderTargetLoadAction::ELoad, ERenderTargetLoadAction::ELoad, FExclusiveDepthStencil::DepthWrite_StencilRead);
 			if (bNaniteShadingMaskExport)
 			{
 				PassParameters->RenderTargets[0] = FRenderTargetBinding(NaniteRasterResults[ViewIndex].ShadingMask, ERenderTargetLoadAction::ELoad);
@@ -1253,7 +1252,7 @@ FSingleLayerWaterPrePassResult* FDeferredShadingSceneRenderer::RenderSingleLayer
 				nullptr /*BlendState*/,
 				nullptr /*RasterizerState*/,
 				TStaticDepthStencilState<true, CF_Always>::GetRHI(),
-				1 /*StencilRef*/);
+				STENCIL_SANDBOX_MASK);
 		}
 	}
 
@@ -1808,7 +1807,7 @@ void FDeferredShadingSceneRenderer::RenderSingleLayerWaterInner(
 	FDepthStencilBinding DepthStencilBinding;
 	if (bHasDepthPrepass)
 	{
-		DepthStencilBinding = FDepthStencilBinding(SingleLayerWaterPrePassResult->DepthPrepassTexture.Target, ERenderTargetLoadAction::ELoad, ERenderTargetLoadAction::ELoad, FExclusiveDepthStencil::DepthRead_StencilRead);
+		DepthStencilBinding = FDepthStencilBinding(SingleLayerWaterPrePassResult->DepthPrepassTexture.Target, ERenderTargetLoadAction::ELoad, ERenderTargetLoadAction::ELoad, FExclusiveDepthStencil::DepthRead_StencilWrite);
 	}
 	else
 	{
@@ -2059,7 +2058,7 @@ void FSingleLayerWaterPassMeshProcessor::CollectPSOInitializers(const FSceneText
 		}
 
 		const bool bHasDepthPrepass = IsSingleLayerWaterDepthPrepassEnabled(GetFeatureLevelShaderPlatform(FeatureLevel), FeatureLevel);
-		RenderTargetsInfo.DepthStencilAccess = bHasDepthPrepass ? FExclusiveDepthStencil::DepthRead_StencilRead : FExclusiveDepthStencil::DepthRead_StencilNop;
+		RenderTargetsInfo.DepthStencilAccess = bHasDepthPrepass ? FExclusiveDepthStencil::DepthRead_StencilWrite : FExclusiveDepthStencil::DepthWrite_StencilNop;
 
 		FBasePassMeshProcessor::AddBasePassGraphicsPipelineStateInitializer(
 			FeatureLevel,
@@ -2085,18 +2084,18 @@ FMeshPassProcessor* CreateSingleLayerWaterPassProcessor(ERHIFeatureLevel::Type F
 	FMeshPassProcessorRenderState DrawRenderState;
 
 	// Make sure depth write is enabled if no prepass is used.
-	FExclusiveDepthStencil::Type BasePassDepthStencilAccess_DepthWrite = FExclusiveDepthStencil::Type(bHasDepthPrepass ? FExclusiveDepthStencil::DepthRead : SceneBasePassDepthStencilAccess | FExclusiveDepthStencil::DepthWrite);
+	FExclusiveDepthStencil::Type BasePassDepthStencilAccess_DepthWrite = FExclusiveDepthStencil::Type(bHasDepthPrepass ? FExclusiveDepthStencil::DepthRead_StencilWrite : SceneBasePassDepthStencilAccess | FExclusiveDepthStencil::DepthWrite_StencilNop);
 	SetupBasePassState(BasePassDepthStencilAccess_DepthWrite, false, DrawRenderState);
 	if (bHasDepthPrepass)
 	{
-		// Set depth stencil test to only pass if depth and stencil are equal to the values written by the prepass
+		// Set depth stencil test to only pass if depth and stencil are equal to the values written by the prepass and then clear the stencil bit
 		DrawRenderState.SetDepthStencilState(TStaticDepthStencilState<
 			false, CF_Equal,							// Depth test
-			true, CF_Equal, SO_Keep, SO_Keep, SO_Keep,	// Front face stencil 
-			true, CF_Equal, SO_Keep, SO_Keep, SO_Keep,	// Back face stencil
-			0xFF, 0x0									// Stencil read/write masks
+			true, CF_Equal, SO_Keep, SO_Keep, SO_Zero,	// Front face stencil 
+			true, CF_Equal, SO_Keep, SO_Keep, SO_Zero,	// Back face stencil
+			STENCIL_SANDBOX_MASK, STENCIL_SANDBOX_MASK 	// Stencil read/write masks
 		>::GetRHI());
-		DrawRenderState.SetStencilRef(1);
+		DrawRenderState.SetStencilRef(STENCIL_SANDBOX_MASK);
 	}
 
 	return new FSingleLayerWaterPassMeshProcessor(Scene, FeatureLevel, InViewIfDynamicMeshCommand, DrawRenderState, InDrawListContext);
@@ -2366,7 +2365,7 @@ FMeshPassProcessor* CreateSingleLayerWaterDepthPrepassProcessor(ERHIFeatureLevel
 	FMeshPassProcessorRenderState DrawRenderState;
 
 	// Make sure depth write is enabled.
-	FExclusiveDepthStencil::Type BasePassDepthStencilAccess_DepthWrite = FExclusiveDepthStencil::Type(SceneBasePassDepthStencilAccess | FExclusiveDepthStencil::DepthWrite);
+	FExclusiveDepthStencil::Type BasePassDepthStencilAccess_DepthWrite = FExclusiveDepthStencil::Type(SceneBasePassDepthStencilAccess | FExclusiveDepthStencil::DepthWrite_StencilWrite);
 	
 	// Disable color writes, enable depth tests and writes.
 	DrawRenderState.SetBlendState(TStaticBlendState<CW_NONE>::GetRHI());
@@ -2376,7 +2375,7 @@ FMeshPassProcessor* CreateSingleLayerWaterDepthPrepassProcessor(ERHIFeatureLevel
 		true, CF_Always, SO_Keep, SO_Keep, SO_Replace	// Back face stencil
 	>::GetRHI());
 	DrawRenderState.SetDepthStencilAccess(BasePassDepthStencilAccess_DepthWrite);
-	DrawRenderState.SetStencilRef(1);
+	DrawRenderState.SetStencilRef(STENCIL_SANDBOX_MASK);
 
 	return new FSingleLayerWaterDepthPrepassMeshProcessor(Scene, FeatureLevel, InViewIfDynamicMeshCommand, DrawRenderState, InDrawListContext);
 }
