[SingleLayerWater] Implemented full SLW in the deferred SM5 mobile renderer. This path draws water in a separate pass, allowing for water to have colored transmittance, refraction and to draw depth. It uses the recently introduced SceneColorCopy texture for colored transmittance and adds a depth buffer copy to support the SceneDepthWithoutWater node. This path is only active in the mobile renderer if it's set to deferred shading and r.Mobile.SupportsSM5MaterialNodes is true.

diff --git a/Engine/Shaders/Private/MaterialTemplate.ush b/Engine/Shaders/Private/MaterialTemplate.ush
--- a/Engine/Shaders/Private/MaterialTemplate.ush
+++ b/Engine/Shaders/Private/MaterialTemplate.ush
@@ -2436,10 +2436,18 @@ float2 SceneWithoutSingleLayerWaterInvTextureSize;
 float MaterialExpressionSceneDepthWithoutWater(float2 ViewportUV, float FallbackDepth)
 {
 #if MATERIAL_SHADINGMODEL_SINGLELAYERWATER && !HIT_PROXY_SHADER && !SCENE_TEXTURES_DISABLED && (SINGLE_LAYER_WATER_SHADING_QUALITY != SINGLE_LAYER_WATER_SHADING_QUALITY_MOBILE_WITH_DEPTH_BUFFER)
+
+#if MOBILE_SUPPORTS_SM5_MATERIAL_NODES
+	const float4 SLWBufferBilinearUVMinMax = View.BufferSizeAndInvSize.zwzw * (ResolvedView.ViewRectMinAndSize.xyxy + float4(0.55f, 0.55f, ResolvedView.ViewRectMinAndSize.zw - 0.55f));
+	const float2 ClampedViewportUV = clamp(ViewportUV, SLWBufferBilinearUVMinMax.xy, SLWBufferBilinearUVMinMax.zw);
+	float WaterDeviceZ = WaterSampleSceneDepthWithoutWater(MobileBasePass.SceneDepthCopyTexture, MobileBasePass.SceneDepthCopySampler, ClampedViewportUV, View.BufferSizeAndInvSize.xy, View.BufferSizeAndInvSize.zw);
+#else
 	const float2 ClampedViewportUV = clamp(ViewportUV, SingleLayerWater.SceneWithoutSingleLayerWaterMinMaxUV.xy, SingleLayerWater.SceneWithoutSingleLayerWaterMinMaxUV.zw);
 	float WaterDeviceZ = WaterSampleSceneDepthWithoutWater(SingleLayerWater.SceneDepthWithoutSingleLayerWaterTexture, SingleLayerWaterSceneDepthSampler, ClampedViewportUV, 
 		SingleLayerWater.SceneWithoutSingleLayerWaterTextureSize, SingleLayerWater.SceneWithoutSingleLayerWaterInvTextureSize);
+#endif
 	return ConvertFromDeviceZ(WaterDeviceZ);
+
 #elif POST_PROCESS_MATERIAL
 	if (bSceneDepthWithoutWaterTextureAvailable)
 	{
diff --git a/Engine/Shaders/Private/MobileBasePassPixelShader.usf b/Engine/Shaders/Private/MobileBasePassPixelShader.usf
--- a/Engine/Shaders/Private/MobileBasePassPixelShader.usf
+++ b/Engine/Shaders/Private/MobileBasePassPixelShader.usf
@@ -56,8 +56,18 @@
 	#ifdef SINGLE_LAYER_WATER_SHADING_QUALITY
 	#undef SINGLE_LAYER_WATER_SHADING_QUALITY
 	#endif
-	// Value must match SINGLE_LAYER_WATER_SHADING_QUALITY_MOBILE_WITH_DEPTH_BUFFER in SingleLayerWaterCommon.ush!
-	#define SINGLE_LAYER_WATER_SHADING_QUALITY 2
+
+	#if MOBILE_SUPPORTS_SM5_MATERIAL_NODES
+		// Value must match SINGLE_LAYER_WATER_SHADING_QUALITY_FULL in SingleLayerWaterCommon.ush!
+		#define SINGLE_LAYER_WATER_SHADING_QUALITY 0
+		// SINGLE_LAYER_WATER_SHADING_QUALITY_FULL normally uses a refraction mask texture which we don't have on mobile
+		#define SINGLE_LAYER_WATER_USE_REFRACTION_MASK 0
+		// The SceneColorCopyTexture uses an encoding. SLW needs to know about this to undo it.		
+		#define SINGLE_LAYER_WATER_ENCODED_SCENE_COLOR 1
+	#else
+		// Value must match SINGLE_LAYER_WATER_SHADING_QUALITY_MOBILE_WITH_DEPTH_BUFFER in SingleLayerWaterCommon.ush!
+		#define SINGLE_LAYER_WATER_SHADING_QUALITY 2
+	#endif
 #endif
 
 // Enable Substrate. This define & include need to be defined before certains includes (i.e., DBufferDecalShared which uses them internally)
@@ -1191,22 +1201,38 @@ void Main(
 		WaterDiffuseIndirectIlluminance = lerp(WaterDiffuseIndirectIlluminance, PI, View.UnlitViewmodeMask);
 #endif
 
+#if MOBILE_SUPPORTS_SM5_MATERIAL_NODES
+		const float4 SLWBufferBilinearUVMinMax = View.BufferSizeAndInvSize.zwzw * (ResolvedView.ViewRectMinAndSize.xyxy + float4(0.55f, 0.55f, ResolvedView.ViewRectMinAndSize.zw - 0.55f));
+#endif
+
 		half3 N = GBuffer.WorldNormal;
 		half3 V = MaterialParameters.CameraVector;
 		half NoV = saturate(abs(dot(N, V)) + 1e-5);
 		const float3 EnvBrdf = GetEnvBRDF(GBuffer.SpecularColor, GBuffer.Roughness, NoV);
 
-		const float4 NullDistortionParams = 1.0f;
 		WaterVolumeLightingOutput WaterLighting = EvaluateWaterVolumeLighting(
 			MaterialParameters, PixelMaterialInputs, ResolvedView,
-			DirectionalLightShadow, GBuffer.Specular, NullDistortionParams,
+			DirectionalLightShadow, 
+#if MOBILE_SUPPORTS_SM5_MATERIAL_NODES
+			MobileBasePass.SceneDepthCopyTexture, MobileBasePass.SceneDepthCopySampler,
+			ResolvedView.BufferSizeAndInvSize.xy,
+			ResolvedView.BufferSizeAndInvSize.zw,
+			SceneColorCopyTexture, SceneColorCopySampler,
+			SLWBufferBilinearUVMinMax.xy,
+			SLWBufferBilinearUVMinMax.zw,
+#endif
+			GBuffer.Specular, MobileBasePass.DistortionParams,
 			SunIlluminance, WaterDiffuseIndirectIlluminance, EnvBrdf,
 			CameraIsUnderWater, WaterVisibility, EyeIndex,
 			bSeparateMainDirLightLuminance, SeparatedWaterMainDirLightLuminance);
 
 		// Accumulate luminance and occlude the background according to transmittance to view and mean transmittance to lights.
 		LightAccumulator_AddSplit(DirectLighting, WaterLighting.Luminance, 0.0f, WaterLighting.Luminance, 1.0f, false);
+#if MOBILE_SUPPORTS_SM5_MATERIAL_NODES
+		Opacity = 1.0f;
+#else
 		Opacity = 1.0 - ((1.0 - Opacity) * dot(WaterLighting.WaterToSceneToLightTransmittance, float3(1.0 / 3.0, 1.0 / 3.0, 1.0 / 3.0)));
+#endif
 #if MOBILE_EMULATION
 		DiffuseColorForIndirect += WaterLighting.Luminance;
 #endif
diff --git a/Engine/Shaders/Private/SingleLayerWaterShading.ush b/Engine/Shaders/Private/SingleLayerWaterShading.ush
--- a/Engine/Shaders/Private/SingleLayerWaterShading.ush
+++ b/Engine/Shaders/Private/SingleLayerWaterShading.ush
@@ -8,6 +8,15 @@
 #include "DistortionCommon.ush"
 #include "SingleLayerWaterCommon.ush"
 
+#ifndef SINGLE_LAYER_WATER_USE_REFRACTION_MASK
+#define SINGLE_LAYER_WATER_USE_REFRACTION_MASK 1
+#endif
+
+// The deferred mobile renderer with support for SM5 creates a scene color copy with encoding to fit the data into RGB8
+#ifndef SINGLE_LAYER_WATER_ENCODED_SCENE_COLOR
+#define SINGLE_LAYER_WATER_ENCODED_SCENE_COLOR 0
+#endif
+
 #if SINGLE_LAYER_WATER_SHADING_QUALITY != SINGLE_LAYER_WATER_SHADING_QUALITY_FULL && SINGLE_LAYER_WATER_SHADING_QUALITY != SINGLE_LAYER_WATER_SHADING_QUALITY_MOBILE_WITH_DEPTH_TEXTURE && SINGLE_LAYER_WATER_SHADING_QUALITY != SINGLE_LAYER_WATER_SHADING_QUALITY_MOBILE_WITH_DEPTH_BUFFER
 #error SINGLE_LAYER_WATER_SHADING_QUALITY has an invalid value!
 #endif
@@ -41,7 +50,9 @@ WaterVolumeLightingOutput EvaluateWaterVolumeLighting(
 	Texture2D SceneColorWithoutSingleLayerWaterTexture, SamplerState SceneColorWithoutSingleLayerWaterSampler,
 	float2 SceneWithoutSingleLayerWaterMinUV,
 	float2 SceneWithoutSingleLayerWaterMaxUV,
+	#if SINGLE_LAYER_WATER_USE_REFRACTION_MASK
 	Texture2D RefractionMaskTexture,
+	#endif
 #endif
 	float Specular, const float4 DistortionParams,
 	float3 SunIlluminance, float3 AmbiantIlluminance, float3 EnvBrdf,
@@ -105,7 +116,7 @@ WaterVolumeLightingOutput EvaluateWaterVolumeLighting(
 		float2 DistortedUV = clamp(ViewportUV + BufferUVDistortion, SceneWithoutSingleLayerWaterMinUV, SceneWithoutSingleLayerWaterMaxUV);
 
 		float RefractionMask = 1.0f;
-#if SINGLE_LAYER_WATER_SHADING_QUALITY == SINGLE_LAYER_WATER_SHADING_QUALITY_FULL
+#if (SINGLE_LAYER_WATER_SHADING_QUALITY == SINGLE_LAYER_WATER_SHADING_QUALITY_FULL) && SINGLE_LAYER_WATER_USE_REFRACTION_MASK
 		{
 			// Gather the four refraction mask samples at the distorted UV to see if they are valid for refraction.
 			const float4 PackedRefractionMask4 = RefractionMaskTexture.Gather(SceneColorWithoutSingleLayerWaterSampler, DistortedUV);
@@ -221,6 +232,10 @@ WaterVolumeLightingOutput EvaluateWaterVolumeLighting(
 #else
 		// We use the pixel SvPosition instead of the scene one pre refraction/distortion to avoid those extra ALUs.
 		float3 BehindWaterSceneLuminance = SceneColorWithoutSingleLayerWaterTexture.SampleLevel(SceneColorWithoutSingleLayerWaterSampler, ViewportUV, 0).rgb;
+		#if SINGLE_LAYER_WATER_ENCODED_SCENE_COLOR
+		// Undo the function in EncodeSceneColorForMaterialNode
+		BehindWaterSceneLuminance = pow(BehindWaterSceneLuminance.rgb, 4) * 10;
+		#endif
 		BehindWaterSceneLuminance = MeanTransmittanceToLightSources * ResolvedView.OneOverPreExposure * BehindWaterSceneLuminance;
 		BehindWaterSceneLuminance *= RefractionMask; // Multiply by refraction mask to essentially turn invalid refraction data black.
 #endif
diff --git a/Engine/Shaders/Public/ShaderVersion.ush b/Engine/Shaders/Public/ShaderVersion.ush
--- a/Engine/Shaders/Public/ShaderVersion.ush
+++ b/Engine/Shaders/Public/ShaderVersion.ush
@@ -3,4 +3,4 @@
 // Each time the console command is executed it generates a new GUID. As this file is included
 // in Platform.ush (which should be included in any shader) it allows to invalidate the shader DDC.
 // 
-#pragma message("UESHADERMETADATA_VERSION 8402F08A-02D1-44EF-96B0-B0C00BFC4CE1")
+#pragma message("UESHADERMETADATA_VERSION 93A9D3F3-46BF-6FE6-256A-4B9988CFE825")
\ No newline at end of file
diff --git a/Engine/Source/Runtime/Engine/Private/Materials/MaterialInterface.cpp b/Engine/Source/Runtime/Engine/Private/Materials/MaterialInterface.cpp
--- a/Engine/Source/Runtime/Engine/Private/Materials/MaterialInterface.cpp
+++ b/Engine/Source/Runtime/Engine/Private/Materials/MaterialInterface.cpp
@@ -475,7 +475,7 @@ FMaterialRelevance UMaterialInterface::GetRelevance_Internal(const UMaterial* Ma
 
 	const bool bIsMobile = IsMobilePlatform(InShaderPlatform);
 	const bool bUsesSingleLayerWaterMaterial = MaterialResource->GetShadingModels().HasShadingModel(MSM_SingleLayerWater);
-	const bool bIsSinglePassWaterTranslucent = bIsMobile && bUsesSingleLayerWaterMaterial;
+	const bool bIsSinglePassWaterTranslucent = bIsMobile && bUsesSingleLayerWaterMaterial && !MobileSupportsSM5MaterialNodes(InShaderPlatform);
 
 	// Note that even though XX_GameThread() api is called, this function can be called on non game thread via 
 	// GetRelevance_Concurrent()
diff --git a/Engine/Source/Runtime/Engine/Private/ShaderCompiler/ShaderCompiler.cpp b/Engine/Source/Runtime/Engine/Private/ShaderCompiler/ShaderCompiler.cpp
--- a/Engine/Source/Runtime/Engine/Private/ShaderCompiler/ShaderCompiler.cpp
+++ b/Engine/Source/Runtime/Engine/Private/ShaderCompiler/ShaderCompiler.cpp
@@ -3335,6 +3335,7 @@ void GlobalBeginCompileShader(
 			const bool bGLESDeferredShading = (Target.Platform == SP_OPENGL_ES3_1_ANDROID && bAllowFramebufferFetch);
 			SET_SHADER_DEFINE_AND_COMPILE_ARGUMENT(Input.Environment, USE_GLES_FBF_DEFERRED, bGLESDeferredShading);
 			SET_SHADER_DEFINE(Input.Environment, MOBILE_EXTENDED_GBUFFER, MobileUsesExtenedGBuffer((EShaderPlatform)Target.Platform) ? 1 : 0);
+			SET_SHADER_DEFINE(Input.Environment, MOBILE_SUPPORTS_SM5_MATERIAL_NODES, MobileSupportsSM5MaterialNodes((EShaderPlatform)Target.Platform) ? 1 : 0);
 		}
 		else
 		{
diff --git a/Engine/Source/Runtime/Renderer/Private/DepthCopy.h b/Engine/Source/Runtime/Renderer/Private/DepthCopy.h
--- a/Engine/Source/Runtime/Renderer/Private/DepthCopy.h
+++ b/Engine/Source/Runtime/Renderer/Private/DepthCopy.h
@@ -37,7 +37,7 @@ public:
 
 		static bool ShouldCompilePermutation(const FGlobalShaderPermutationParameters& Parameters)
 	{
-		return IsFeatureLevelSupported(Parameters.Platform, ERHIFeatureLevel::SM5);
+		return IsFeatureLevelSupported(Parameters.Platform, ERHIFeatureLevel::SM5) || MobileSupportsSM5MaterialNodes(Parameters.Platform);
 	}
 };
 
diff --git a/Engine/Source/Runtime/Renderer/Private/MobileBasePass.cpp b/Engine/Source/Runtime/Renderer/Private/MobileBasePass.cpp
--- a/Engine/Source/Runtime/Renderer/Private/MobileBasePass.cpp
+++ b/Engine/Source/Runtime/Renderer/Private/MobileBasePass.cpp
@@ -404,9 +404,7 @@ ELightMapPolicyType MobileBasePass::SelectMeshLightmapPolicy(
 	return SelectedLightmapPolicy;
 }
 
-typedef TArray<ELightMapPolicyType, TInlineAllocator<4>> FMobileLightMapPolicyTypeList;
-
-static FMobileLightMapPolicyTypeList GetUniformLightMapPolicyTypeForPSOCollection(bool bLitMaterial, bool bTranslucent/*not used*/, bool bUsesDeferredShading, bool bCanReceiveCSM, bool bMovable)
+FMobileLightMapPolicyTypeList MobileBasePass::GetUniformLightMapPolicyTypeForPSOCollection(bool bLitMaterial, bool bTranslucent/*not used*/, bool bUsesDeferredShading, bool bCanReceiveCSM, bool bMovable)
 {
 	FMobileLightMapPolicyTypeList Result;
 	
@@ -646,7 +644,7 @@ void MobileBasePass::SetTranslucentRenderState(FMeshPassProcessorRenderState& Dr
 															CW_NONE, BO_Add, BF_One, BF_Zero, BO_Add, BF_One, BF_Zero>::GetRHI());
 			break;
 		default:
-			if (ShadingModels.HasShadingModel(MSM_SingleLayerWater))
+			if (ShadingModels.HasShadingModel(MSM_SingleLayerWater) && !MobileSupportsSM5MaterialNodes(ShaderPlatform))
 			{
 				// Single layer water is an opaque marerial rendered as translucent on Mobile. We force pre-multiplied alpha to achieve water depth based transmittance.
 				DrawRenderState.SetBlendState(TStaticBlendState<CW_RGBA, BO_Add, BF_One, BF_InverseSourceAlpha, BO_Add, BF_Zero, BF_InverseSourceAlpha,
@@ -764,7 +762,9 @@ FMobileBasePassMeshProcessor::FMobileBasePassMeshProcessor(
 bool FMobileBasePassMeshProcessor::ShouldDraw(const FMaterial& Material) const
 {
 	const FMaterialShadingModelField ShadingModels = Material.GetShadingModels();
-	const bool bIsTranslucent = IsTranslucentBlendMode(Material.GetBlendMode()) || ShadingModels.HasShadingModel(MSM_SingleLayerWater); // Water goes into the translucent pass;
+	const bool bIsSLW = ShadingModels.HasShadingModel(MSM_SingleLayerWater);
+	const bool bIsTranslucentSLW = bIsSLW && !MobileSupportsSM5MaterialNodes(GetFeatureLevelShaderPlatform(ERHIFeatureLevel::ES3_1));
+	const bool bIsTranslucent = IsTranslucentBlendMode(Material.GetBlendMode()) || bIsTranslucentSLW; // Water goes into the translucent pass;
 	if (bTranslucentBasePass)
 	{
 		// Skipping TPT_TranslucencyAfterDOFModulate. That pass is only needed for Dual Blending, which is not supported on Mobile.
@@ -778,7 +778,7 @@ bool FMobileBasePassMeshProcessor::ShouldDraw(const FMaterial& Material) const
 	else
 	{
 		// opaque materials.
-		return !bIsTranslucent;
+		return !bIsTranslucent && !bIsSLW;
 	}
 }
 static bool MaterialCanReceiveCSM(const FMaterial& Material)
@@ -1100,7 +1100,7 @@ void FMobileBasePassMeshProcessor::CollectPSOInitializers(const FSceneTexturesCo
 
 	const bool bCanReceiveCSM = ((Flags & EFlags::CanReceiveCSM) == EFlags::CanReceiveCSM) && MaterialCanReceiveCSM(Material);
 
-	FMobileLightMapPolicyTypeList UniformLightMapPolicyTypes = GetUniformLightMapPolicyTypeForPSOCollection(bLitMaterial, bTranslucentBasePass, bPassUsesDeferredShading, bCanReceiveCSM, bMovable);
+	FMobileLightMapPolicyTypeList UniformLightMapPolicyTypes = MobileBasePass::GetUniformLightMapPolicyTypeForPSOCollection(bLitMaterial, bTranslucentBasePass, bPassUsesDeferredShading, bCanReceiveCSM, bMovable);
 	
 	for (ELightMapPolicyType LightMapPolicyType : UniformLightMapPolicyTypes)
 	{
diff --git a/Engine/Source/Runtime/Renderer/Private/MobileBasePassRendering.cpp b/Engine/Source/Runtime/Renderer/Private/MobileBasePassRendering.cpp
--- a/Engine/Source/Runtime/Renderer/Private/MobileBasePassRendering.cpp
+++ b/Engine/Source/Runtime/Renderer/Private/MobileBasePassRendering.cpp
@@ -23,6 +23,7 @@
 #include "ReflectionEnvironment.h"
 #include "RenderCore.h"
 #include "LocalFogVolumeRendering.h"
+#include "DistortionRendering.h"
 
 // Changing this causes a full shader recompile
 static TAutoConsoleVariable<int32> CVarMobileDisableVertexFog(
@@ -377,7 +378,10 @@ void SetupMobileBasePassUniformParameters(
 	SetupMobileSSRParameters(GraphBuilder, View, BasePassParameters.SSRParams);
 
 	BasePassParameters.SceneColorCopyTexture = SystemTextures.White;
-	BasePassParameters.SceneColorCopySampler = TStaticSamplerState<SF_Bilinear, AM_Clamp, AM_Clamp, AM_Clamp>::GetRHI();;
+	BasePassParameters.SceneColorCopySampler = TStaticSamplerState<SF_Bilinear, AM_Clamp, AM_Clamp, AM_Clamp>::GetRHI();
+	BasePassParameters.SceneDepthCopyTexture = SystemTextures.DepthDummy;
+	BasePassParameters.SceneDepthCopySampler = TStaticSamplerState<SF_Point, AM_Clamp, AM_Clamp, AM_Clamp>::GetRHI();
+	SetupDistortionParams(BasePassParameters.DistortionParams, View);
 
 #if WITH_DEBUG_VIEW_MODES
 	if (View.Family->UseDebugViewPS())
diff --git a/Engine/Source/Runtime/Renderer/Private/MobileBasePassRendering.h b/Engine/Source/Runtime/Renderer/Private/MobileBasePassRendering.h
--- a/Engine/Source/Runtime/Renderer/Private/MobileBasePassRendering.h
+++ b/Engine/Source/Runtime/Renderer/Private/MobileBasePassRendering.h
@@ -67,6 +67,9 @@ BEGIN_GLOBAL_SHADER_PARAMETER_STRUCT(FMobileBasePassUniformParameters, )
 	SHADER_PARAMETER_STRUCT(FMobileScreenSpaceReflectionParams, SSRParams)
 	SHADER_PARAMETER_RDG_TEXTURE(Texture2D, SceneColorCopyTexture)
 	SHADER_PARAMETER_SAMPLER(SamplerState, SceneColorCopySampler)
+	SHADER_PARAMETER_RDG_TEXTURE(Texture2D, SceneDepthCopyTexture)
+	SHADER_PARAMETER_SAMPLER(SamplerState, SceneDepthCopySampler)
+	SHADER_PARAMETER(FVector4f, DistortionParams)
 END_GLOBAL_SHADER_PARAMETER_STRUCT()
 
 enum class EMobileBasePass
@@ -319,6 +322,7 @@ public:
 	TMobileBasePassPSBaseType() {}
 };
 
+typedef TArray<ELightMapPolicyType, TInlineAllocator<4>> FMobileLightMapPolicyTypeList;
 
 namespace MobileBasePass
 {
@@ -331,6 +335,8 @@ namespace MobileBasePass
 		bool bIsLitMaterial,
 		bool bIsTranslucent);
 
+	FMobileLightMapPolicyTypeList GetUniformLightMapPolicyTypeForPSOCollection(bool bLitMaterial, bool bTranslucent/*not used*/, bool bUsesDeferredShading, bool bCanReceiveCSM, bool bMovable);
+
 	bool GetShaders(
 		ELightMapPolicyType LightMapPolicyType,
 		EMobileLocalLightSetting LocalLightSetting,
@@ -358,8 +364,9 @@ public:
 		// We compile the point light shader combinations based on the project settings
 		const bool bIsLit = Parameters.MaterialParameters.ShadingModels.IsLit();
 		const bool bDeferredShadingEnabled = IsMobileDeferredShadingEnabled(Parameters.Platform);
-		const bool bIsTranslucent = IsTranslucentBlendMode(Parameters.MaterialParameters) || Parameters.MaterialParameters.ShadingModels.HasShadingModel(MSM_SingleLayerWater);
-		const bool bMaterialUsesForwardShading = bIsLit && bIsTranslucent;
+		const bool bIsTranslucent = IsTranslucentBlendMode(Parameters.MaterialParameters);
+		const bool bIsSLW = Parameters.MaterialParameters.ShadingModels.HasShadingModel(MSM_SingleLayerWater);
+		const bool bMaterialUsesForwardShading = bIsLit && (bIsTranslucent || bIsSLW);
 		// Translucent materials always support clustered shading on mobile deferred
 		const bool bForwardShading = !bDeferredShadingEnabled || bMaterialUsesForwardShading;
 
@@ -385,9 +392,10 @@ public:
 		const bool bMobileUsesShadowMaskTexture = MobileUsesShadowMaskTexture(Parameters.Platform);
 		const bool bEnableClusteredReflections = MobileForwardEnableClusteredReflections(Parameters.Platform);
 		const bool bDeferredShadingEnabled = IsMobileDeferredShadingEnabled(Parameters.Platform);
-		const bool bTranslucentMaterial = IsTranslucentBlendMode(Parameters.MaterialParameters) || Parameters.MaterialParameters.ShadingModels.HasShadingModel(MSM_SingleLayerWater);
+		const bool bTranslucentMaterial = IsTranslucentBlendMode(Parameters.MaterialParameters);
+		const bool bIsSLW = Parameters.MaterialParameters.ShadingModels.HasShadingModel(MSM_SingleLayerWater);
 		const bool bIsLit = Parameters.MaterialParameters.ShadingModels.IsLit();
-		const bool bMaterialUsesForwardShading = bIsLit && bTranslucentMaterial;
+		const bool bMaterialUsesForwardShading = bIsLit && (bTranslucentMaterial || bIsSLW);
 		// Translucent materials always support clustered shading on mobile deferred
 		const bool bForwardShading = !bDeferredShadingEnabled || bMaterialUsesForwardShading;
 				
diff --git a/Engine/Source/Runtime/Renderer/Private/MobileShadingRenderer.cpp b/Engine/Source/Runtime/Renderer/Private/MobileShadingRenderer.cpp
--- a/Engine/Source/Runtime/Renderer/Private/MobileShadingRenderer.cpp
+++ b/Engine/Source/Runtime/Renderer/Private/MobileShadingRenderer.cpp
@@ -80,6 +80,8 @@
 #include "ViewData.h"
 #include "DistanceFieldAmbientOcclusion.h"
 #include "DistanceFieldLightingShared.h"
+#include "DepthCopy.h"
+#include "SingleLayerWaterRendering.h"
 
 uint32 GetShadowQuality();
 
@@ -2421,6 +2423,24 @@ void FMobileSceneRenderer::RenderDeferredMultiPass(FRDGBuilder& GraphBuilder, FS
 		RenderDistanceFieldLighting(GraphBuilder, SceneTextures, Parameters, DynamicBentNormalAOTextures, false, false);
 	}
 
+	const bool bSupportsSM5Nodes = MobileSupportsSM5MaterialNodes(ShaderPlatform);
+	const bool bShouldRenderFullSLW = bSupportsSM5Nodes && ShouldRenderSingleLayerWater(Views);
+	const bool bNeedsDepthCopy = bShouldRenderFullSLW;
+
+	FRDGTextureRef SceneDepthCopyTexture = nullptr;
+	if (bNeedsDepthCopy)
+	{
+		AddResolveSceneDepthPass(GraphBuilder, Views, SceneTextures.Depth);
+
+		SceneDepthCopyTexture = GraphBuilder.CreateTexture(FRDGTextureDesc::Create2D(SceneTextures.Config.Extent, PF_DepthStencil, FClearValueBinding::DepthFar, TexCreate_DepthStencilTargetable | TexCreate_ShaderResource), TEXT("SceneDepthCopy"));
+
+		for (FRenderViewContext& ViewContext : RenderViews)
+		{
+			FViewInfo& View = *ViewContext.ViewInfo;
+			DepthCopy::AddViewDepthCopyPSPass(GraphBuilder, View, SceneTextures.Depth.Resolve, SceneDepthCopyTexture);
+		}
+	}
+
 	// Lighting and translucency
 	uint32 ViewIndex = 0;
 	for (FRenderViewContext& ViewContext : RenderViews)
@@ -2465,37 +2485,70 @@ void FMobileSceneRenderer::RenderDeferredMultiPass(FRDGBuilder& GraphBuilder, FS
 
 		const EMobileSSRQuality MobileSSRQuality = ActiveMobileSSRQuality(View, bShouldRenderVelocities);
 
-		bool bSeparateTranslucyPass = MobileSupportsSM5MaterialNodes(ShaderPlatform) && !View.IsUnderwater();
+		bool bSeparateTranslucencyPass = bSupportsSM5Nodes;
 
 		GraphBuilder.AddPass(
-			bSeparateTranslucyPass ? RDG_EVENT_NAME("LightingAndFog") : RDG_EVENT_NAME("LightingAndTranslucency"),
+			bSeparateTranslucencyPass ? RDG_EVENT_NAME("LightingAndFog") : RDG_EVENT_NAME("LightingAndTranslucency"),
 			PassParameters,
 			ERDGPassFlags::Raster,
-			[this, PassParameters, ViewContext, MobileSSRQuality, bSeparateTranslucyPass, &SortedLightSet](FRHICommandList& RHICmdList)
+			[this, PassParameters, ViewContext, MobileSSRQuality, bSeparateTranslucencyPass, &SortedLightSet](FRHICommandList& RHICmdList)
 			{
 				FViewInfo& View = *ViewContext.ViewInfo;
 
 				MobileDeferredShadingPass(RHICmdList, ViewContext.ViewIndex, Views.Num(), View, *Scene, SortedLightSet, VisibleLightInfos, MobileSSRQuality, PassParameters->BentNormalAOTexture);
 				RenderFog(RHICmdList, View);
 				
-				if (!bSeparateTranslucyPass)
+				if (!bSeparateTranslucencyPass)
 				{
 					RenderTranslucency(RHICmdList, View, Views, StandardTranslucencyPass, StandardTranslucencyMeshPass, &PassParameters->InstanceCullingDrawParams);
 				}
 			});
 
-		if (bSeparateTranslucyPass)
+		if (bSeparateTranslucencyPass)
 		{
 			const bool bShouldRenderTranslucency = ShouldRenderTranslucency(StandardTranslucencyPass, Views) && ViewFamily.EngineShowFlags.Translucency && (View.ParallelMeshDrawCommandPasses[StandardTranslucencyMeshPass] != nullptr);
+			const bool bShouldRenderWater = bShouldRenderFullSLW && View.bHasSingleLayerWaterMaterial;
 
-			if (bShouldRenderTranslucency)
+			if (bShouldRenderTranslucency || bShouldRenderWater)
 			{
-				FRDGTextureRef SceneColorCopyTexture = AddCopySceneColorPass(GraphBuilder, Views, SceneTextures.Color, false);
+				FRDGTextureRef SceneColorCopyTexture = AddCopySceneColorPass(GraphBuilder, Views, SceneTextures.Color, false, false);
 				if (SceneColorCopyTexture != nullptr)
 				{
 					BasePassParameters->SceneColorCopyTexture = SceneColorCopyTexture;
 				}
+				if (SceneDepthCopyTexture != nullptr)
+				{
+					BasePassParameters->SceneDepthCopyTexture = SceneDepthCopyTexture;
+				}
+			}
+			if (bShouldRenderWater)
+			{
+				auto* SLWPassParameters = GraphBuilder.AllocParameters<FMobileRenderPassParameters>();
+				*SLWPassParameters = *PassParameters;
+				// Only SceneColor and Depth
+				SLWPassParameters->RenderTargets.DepthStencil = FDepthStencilBinding(BasePassRenderTargets.DepthStencil.GetTexture(), ERenderTargetLoadAction::ELoad, FExclusiveDepthStencil::DepthWrite_StencilNop);
+
+				if (Scene->GPUScene.IsEnabled())
+				{
+					BuildMeshRenderingCommands(GraphBuilder, EMeshPass::SingleLayerWaterPass, View, Scene->GPUScene, InstanceCullingManager, SLWPassParameters->InstanceCullingDrawParams);
+				}
+
+				GraphBuilder.AddPass(
+					RDG_EVENT_NAME("SingleLayerWater"),
+					SLWPassParameters,
+					ERDGPassFlags::Raster,
+					[this, &View, SLWPassParameters](FRHICommandList& RHICmdList)
+					{
+						RHICmdList.SetViewport(View.ViewRect.Min.X, View.ViewRect.Min.Y, 0.0f, View.ViewRect.Max.X, View.ViewRect.Max.Y, 1.0f);
 
+						if (auto* Pass = View.ParallelMeshDrawCommandPasses[EMeshPass::SingleLayerWaterPass])
+						{
+							Pass->Draw(RHICmdList, &SLWPassParameters->InstanceCullingDrawParams);
+						}
+					});
+			}
+			if (bShouldRenderTranslucency)
+			{
 				GraphBuilder.AddPass(
 					RDG_EVENT_NAME("Translucency"),
 					PassParameters,
diff --git a/Engine/Source/Runtime/Renderer/Private/MobileSingleLayerWaterRendering.cpp b/Engine/Source/Runtime/Renderer/Private/MobileSingleLayerWaterRendering.cpp
new file mode 100644
--- /dev/null
+++ b/Engine/Source/Runtime/Renderer/Private/MobileSingleLayerWaterRendering.cpp
@@ -0,0 +1,317 @@
+// Copyright Epic Games, Inc. All Rights Reserved.
+
+#include "MobileBasePassRendering.h"
+#include "TranslucentRendering.h"
+#include "RenderCore.h"
+#include "SceneRendering.h"
+#include "ScenePrivate.h"
+#include "LightMapRendering.h"
+#include "MeshPassProcessor.inl"
+#include "RenderGraphBuilder.h"
+
+extern bool MobileLocalLightsUseSinglePermutation(EShaderPlatform ShaderPlatform);
+
+static FMeshDrawCommandSortKey GetMobileSingleLayerWaterSortKey(bool bIsMasked, bool bIsBackground, const FMeshMaterialShader* VertexShader, const FMeshMaterialShader* PixelShader)
+{
+	FMeshDrawCommandSortKey SortKey;
+	SortKey.BasePass.Masked = (bIsMasked ? 1 : 0);
+	SortKey.BasePass.Background = (bIsBackground ? 1 : 0); // background flag in second bit
+	SortKey.BasePass.VertexShaderHash = (VertexShader ? VertexShader->GetSortKey() : 0) & 0xFFFF;
+	SortKey.BasePass.PixelShaderHash = PixelShader ? PixelShader->GetSortKey() : 0;
+	return SortKey;
+};
+
+static void SetMobileSingleLayerWaterRenderState(FMeshPassProcessorRenderState& DrawRenderState, const FMaterial& Material)
+{
+	constexpr ERHIFeatureLevel::Type FeatureLevel = ERHIFeatureLevel::ES3_1;
+	EShaderPlatform ShaderPlatform = GetFeatureLevelShaderPlatform(FeatureLevel);
+
+	DrawRenderState.SetDepthStencilAccess(FExclusiveDepthStencil::DepthWrite_StencilNop);
+	DrawRenderState.SetBlendState(TStaticBlendStateWriteMask<CW_RGBA, CW_RGBA, CW_RGBA, CW_RGBA>::GetRHI());
+	DrawRenderState.SetDepthStencilState(TStaticDepthStencilState<true, CF_DepthNearOrEqual>::GetRHI());
+}
+
+class FMobileSingleLayerWaterPassMeshProcessor : public FSceneRenderingAllocatorObject<FMobileSingleLayerWaterPassMeshProcessor>, public FMeshPassProcessor
+{
+public:
+	FMobileSingleLayerWaterPassMeshProcessor(const FScene* Scene, ERHIFeatureLevel::Type FeatureLevel, const FSceneView* InViewIfDynamicMeshCommand, const FMeshPassProcessorRenderState& InPassDrawRenderState, FMeshPassDrawListContext* InDrawListContext);
+
+	virtual void AddMeshBatch(const FMeshBatch& RESTRICT MeshBatch, uint64 BatchElementMask, const FPrimitiveSceneProxy* RESTRICT PrimitiveSceneProxy, int32 StaticMeshId = -1) override final;
+	virtual void CollectPSOInitializers(const FSceneTexturesConfig& SceneTexturesConfig, const FMaterial& Material, const FPSOPrecacheVertexFactoryData& VertexFactoryData, const FPSOPrecacheParams& PreCacheParams, TArray<FPSOPrecacheData>& PSOInitializers) override final;
+
+private:
+	bool TryAddMeshBatch(
+		const FMeshBatch& RESTRICT MeshBatch,
+		uint64 BatchElementMask,
+		const FPrimitiveSceneProxy* RESTRICT PrimitiveSceneProxy,
+		int32 StaticMeshId,
+		const FMaterialRenderProxy& MaterialRenderProxy,
+		const FMaterial& Material);
+
+	bool Process(
+		const FMeshBatch& MeshBatch,
+		uint64 BatchElementMask,
+		const FPrimitiveSceneProxy* RESTRICT PrimitiveSceneProxy,
+		int32 StaticMeshId,
+		const FMaterialRenderProxy& RESTRICT MaterialRenderProxy,
+		const FMaterial& RESTRICT MaterialResource,
+		bool bIsMasked,
+		FMaterialShadingModelField ShadingModels,
+		ELightMapPolicyType LightMapPolicyType,
+		EMobileLocalLightSetting LocalLightSetting,
+		const FUniformLightMapPolicy::ElementDataType& RESTRICT LightMapElementData);
+
+	void CollectPSOInitializersForLMPolicy(
+		const FPSOPrecacheVertexFactoryData& VertexFactoryData,
+		const FMeshPassProcessorRenderState& RESTRICT DrawRenderState,
+		const FGraphicsPipelineRenderTargetsInfo& RESTRICT RenderTargetsInfo,
+		const FMaterial& RESTRICT MaterialResource,
+		EMobileLocalLightSetting LocalLightSetting,
+		const ELightMapPolicyType LightMapPolicyType,
+		ERasterizerFillMode MeshFillMode,
+		ERasterizerCullMode MeshCullMode,
+		EPrimitiveType PrimitiveType,
+		TArray<FPSOPrecacheData>& PSOInitializers);
+
+	FMeshPassProcessorRenderState PassDrawRenderState;
+};
+
+FMobileSingleLayerWaterPassMeshProcessor::FMobileSingleLayerWaterPassMeshProcessor(const FScene* Scene, ERHIFeatureLevel::Type FeatureLevel, const FSceneView* InViewIfDynamicMeshCommand, const FMeshPassProcessorRenderState& InPassDrawRenderState, FMeshPassDrawListContext* InDrawListContext)
+	: FMeshPassProcessor(EMeshPass::SingleLayerWaterPass, Scene, FeatureLevel, InViewIfDynamicMeshCommand, InDrawListContext)
+	, PassDrawRenderState(InPassDrawRenderState)
+{
+}
+
+void FMobileSingleLayerWaterPassMeshProcessor::AddMeshBatch(const FMeshBatch& RESTRICT MeshBatch, uint64 BatchElementMask, const FPrimitiveSceneProxy* RESTRICT PrimitiveSceneProxy, int32 StaticMeshId)
+{
+	// This MeshProcessor is only used on the mobile SM5 path, otherwise all SLW is handled by the translucent path.
+	EShaderPlatform ShaderPlatform = GetFeatureLevelShaderPlatform(FeatureLevel);
+	if (!MobileSupportsSM5MaterialNodes(ShaderPlatform))
+	{
+		return;
+	}
+	const FMaterialRenderProxy* MaterialRenderProxy = MeshBatch.MaterialRenderProxy;
+	while (MaterialRenderProxy)
+	{
+		const FMaterial* Material = MaterialRenderProxy->GetMaterialNoFallback(FeatureLevel);
+		if (Material)
+		{
+			if (TryAddMeshBatch(MeshBatch, BatchElementMask, PrimitiveSceneProxy, StaticMeshId, *MaterialRenderProxy, *Material))
+			{
+				break;
+			}
+		}
+
+		MaterialRenderProxy = MaterialRenderProxy->GetFallback(FeatureLevel);
+	}
+}
+
+bool FMobileSingleLayerWaterPassMeshProcessor::TryAddMeshBatch(
+	const FMeshBatch& RESTRICT MeshBatch,
+	uint64 BatchElementMask,
+	const FPrimitiveSceneProxy* RESTRICT PrimitiveSceneProxy,
+	int32 StaticMeshId,
+	const FMaterialRenderProxy& MaterialRenderProxy,
+	const FMaterial& Material)
+{
+	FMaterialShadingModelField ShadingModels = Material.GetShadingModels();
+	if (ShadingModels.HasShadingModel(MSM_SingleLayerWater) && IsOpaqueOrMaskedBlendMode(Material))
+	{
+		const bool bIsMasked = IsMaskedBlendMode(Material);
+		ELightMapPolicyType LightmapPolicyType = MobileBasePass::SelectMeshLightmapPolicy(Scene, MeshBatch, PrimitiveSceneProxy,
+			true /*bCanReceiveCSM*/, false /*bPassUsesDeferredShading*/, true /*bIsLitMaterial*/, true /*bIsTranslucent*/);
+
+		EMobileLocalLightSetting LocalLightSetting = EMobileLocalLightSetting::LOCAL_LIGHTS_DISABLED;
+		if (Scene && PrimitiveSceneProxy)
+		{
+			// we can choose to use a single permutation regarless of local light state
+			// this is to avoid re-caching MDC on light state changes
+			if ((MobileLocalLightsUseSinglePermutation(Scene->GetShaderPlatform()) || PrimitiveSceneProxy->GetPrimitiveSceneInfo()->NumMobileDynamicLocalLights > 0))
+			{
+				LocalLightSetting = GetMobileForwardLocalLightSetting(Scene->GetShaderPlatform());
+			}
+		}
+
+		return Process(MeshBatch, BatchElementMask, PrimitiveSceneProxy, StaticMeshId, MaterialRenderProxy, Material, bIsMasked, ShadingModels, LightmapPolicyType, LocalLightSetting, MeshBatch.LCI);
+	}
+
+	return true;
+}
+
+bool FMobileSingleLayerWaterPassMeshProcessor::Process(
+	const FMeshBatch& MeshBatch,
+	uint64 BatchElementMask,
+	const FPrimitiveSceneProxy* RESTRICT PrimitiveSceneProxy,
+	int32 StaticMeshId,
+	const FMaterialRenderProxy& RESTRICT MaterialRenderProxy,
+	const FMaterial& RESTRICT MaterialResource,
+	bool bIsMasked,
+	FMaterialShadingModelField ShadingModels,
+	ELightMapPolicyType LightMapPolicyType,
+	EMobileLocalLightSetting LocalLightSetting,
+	const FUniformLightMapPolicy::ElementDataType& RESTRICT LightMapElementData)
+{
+	TMeshProcessorShaders<
+		TMobileBasePassVSPolicyParamType<FUniformLightMapPolicy>,
+		TMobileBasePassPSPolicyParamType<FUniformLightMapPolicy>> BasePassShaders;
+
+	if (!MobileBasePass::GetShaders(
+		LightMapPolicyType,
+		LocalLightSetting,
+		MaterialResource,
+		MeshBatch.VertexFactory->GetType(),
+		BasePassShaders.VertexShader,
+		BasePassShaders.PixelShader))
+	{
+		return false;
+	}
+
+	FMeshPassProcessorRenderState DrawRenderState(PassDrawRenderState);
+	SetMobileSingleLayerWaterRenderState(DrawRenderState, MaterialResource);
+
+	// Background primitives will be rendered last in masked/non-masked buckets
+	bool bIsBackground = PrimitiveSceneProxy ? PrimitiveSceneProxy->TreatAsBackgroundForOcclusion() : false;
+	// Default static sort key separates masked and non-masked geometry, generic mesh sorting will also sort by PSO
+	// if platform wants front to back sorting, this key will be recomputed in InitViews
+	FMeshDrawCommandSortKey SortKey = GetMobileSingleLayerWaterSortKey(bIsMasked, bIsBackground, BasePassShaders.VertexShader.GetShader(), BasePassShaders.PixelShader.GetShader());
+
+	const FMeshDrawingPolicyOverrideSettings OverrideSettings = ComputeMeshOverrideSettings(MeshBatch);
+	ERasterizerFillMode MeshFillMode = ComputeMeshFillMode(MaterialResource, OverrideSettings);
+	ERasterizerCullMode MeshCullMode = ComputeMeshCullMode(MaterialResource, OverrideSettings);
+
+	TMobileBasePassShaderElementData<FUniformLightMapPolicy> ShaderElementData(MeshBatch.LCI, true /*bCanReceiveCSM*/);
+	ShaderElementData.InitializeMeshMaterialData(ViewIfDynamicMeshCommand, PrimitiveSceneProxy, MeshBatch, StaticMeshId, false);
+
+	BuildMeshDrawCommands(
+		MeshBatch,
+		BatchElementMask,
+		PrimitiveSceneProxy,
+		MaterialRenderProxy,
+		MaterialResource,
+		DrawRenderState,
+		BasePassShaders,
+		MeshFillMode,
+		MeshCullMode,
+		SortKey,
+		EMeshPassFeatures::Default,
+		ShaderElementData);
+	return true;
+}
+
+void FMobileSingleLayerWaterPassMeshProcessor::CollectPSOInitializersForLMPolicy(
+	const FPSOPrecacheVertexFactoryData& VertexFactoryData,
+	const FMeshPassProcessorRenderState& RESTRICT DrawRenderState,
+	const FGraphicsPipelineRenderTargetsInfo& RESTRICT RenderTargetsInfo,
+	const FMaterial& RESTRICT MaterialResource,
+	EMobileLocalLightSetting LocalLightSetting,
+	const ELightMapPolicyType LightMapPolicyType,
+	ERasterizerFillMode MeshFillMode,
+	ERasterizerCullMode MeshCullMode,
+	EPrimitiveType PrimitiveType,
+	TArray<FPSOPrecacheData>& PSOInitializers)
+{
+	TMeshProcessorShaders<
+		TMobileBasePassVSPolicyParamType<FUniformLightMapPolicy>,
+		TMobileBasePassPSPolicyParamType<FUniformLightMapPolicy>> BasePassShaders;
+
+	if (!MobileBasePass::GetShaders(
+		LightMapPolicyType,
+		LocalLightSetting,
+		MaterialResource,
+		VertexFactoryData.VertexFactoryType,
+		BasePassShaders.VertexShader,
+		BasePassShaders.PixelShader))
+	{
+		return;
+	}
+
+	// subpass info set during the submission of the draws in mobile deferred renderer.
+	uint8 SubpassIndex = 0;
+	ESubpassHint SubpassHint = GetSubpassHint(GMaxRHIShaderPlatform, false /*bIsUsingGBuffers*/, RenderTargetsInfo.MultiViewCount > 1, RenderTargetsInfo.NumSamples);
+
+	AddGraphicsPipelineStateInitializer(
+		VertexFactoryData,
+		MaterialResource,
+		DrawRenderState,
+		RenderTargetsInfo,
+		BasePassShaders,
+		MeshFillMode,
+		MeshCullMode,
+		PrimitiveType,
+		EMeshPassFeatures::Default,
+		SubpassHint,
+		SubpassIndex,
+		true /*bRequired*/,
+		PSOCollectorIndex,
+		PSOInitializers);
+}
+
+void FMobileSingleLayerWaterPassMeshProcessor::CollectPSOInitializers(const FSceneTexturesConfig& SceneTexturesConfig, const FMaterial& Material, const FPSOPrecacheVertexFactoryData& VertexFactoryData, const FPSOPrecacheParams& PreCacheParams, TArray<FPSOPrecacheData>& PSOInitializers)
+{
+	const FMaterialShadingModelField ShadingModels = Material.GetShadingModels();
+	if (!ShadingModels.HasShadingModel(MSM_SingleLayerWater) || !IsOpaqueOrMaskedBlendMode(Material))
+	{
+		return;
+	}
+
+	// Determine the mesh's material and blend mode.
+	EShaderPlatform ShaderPlatform = GetFeatureLevelShaderPlatform(FeatureLevel);
+	const FMeshDrawingPolicyOverrideSettings OverrideSettings = ComputeMeshOverrideSettings(PreCacheParams);
+	const ERasterizerFillMode MeshFillMode = ComputeMeshFillMode(Material, OverrideSettings);
+	ERasterizerCullMode MeshCullMode = ComputeMeshCullMode(Material, OverrideSettings);
+	const EBlendMode BlendMode = Material.GetBlendMode();
+	const bool bLitMaterial = ShadingModels.IsLit();
+
+	bool bMovable =
+		PreCacheParams.Mobility == EComponentMobility::Movable ||
+		PreCacheParams.Mobility == EComponentMobility::Stationary ||
+		PreCacheParams.bUsesIndirectLightingCache; // ILC uses movable path
+
+	// Setup the draw state
+	FMeshPassProcessorRenderState DrawRenderState(PassDrawRenderState);
+
+	FGraphicsPipelineRenderTargetsInfo RenderTargetsInfo;
+	RenderTargetsInfo.NumSamples = 1;
+	EPixelFormat SceneColorFormat = SceneTexturesConfig.ColorFormat;
+	ETextureCreateFlags SceneColorCreateFlags = SceneTexturesConfig.ColorCreateFlags;
+	AddRenderTargetInfo(SceneColorFormat, SceneColorCreateFlags, RenderTargetsInfo);
+	
+	SetupDepthStencilInfo(PF_DepthStencil, SceneTexturesConfig.DepthCreateFlags, ERenderTargetLoadAction::ELoad,
+		ERenderTargetLoadAction::ELoad, FExclusiveDepthStencil::DepthWrite_StencilNop, RenderTargetsInfo);
+	
+	{
+		const static UE::StereoRenderUtils::FStereoShaderAspects Aspects(GMaxRHIShaderPlatform);
+		// If mobile multiview is enabled we expect it will be used with a native MMV, no pre-caching for fallbacks 
+		RenderTargetsInfo.MultiViewCount = Aspects.IsMobileMultiViewEnabled() ? (GSupportsMobileMultiView ? 2 : 1) : 0;
+		// FIXME: Need to figure out if renderer will use shading rate texture or not
+		RenderTargetsInfo.bHasFragmentDensityAttachment = GVRSImageManager.IsAttachmentVRSEnabled();
+	}
+
+	SetMobileSingleLayerWaterRenderState(DrawRenderState, Material);
+
+	const EMobileLocalLightSetting LocalLightSetting = GetMobileForwardLocalLightSetting(ShaderPlatform);
+	const bool bUseLocalLightPermutation = (LocalLightSetting != EMobileLocalLightSetting::LOCAL_LIGHTS_DISABLED);
+
+	FMobileLightMapPolicyTypeList UniformLightMapPolicyTypes = MobileBasePass::GetUniformLightMapPolicyTypeForPSOCollection(bLitMaterial, true /*bTranslucent*/, false /*bUsesDeferredShading*/, true /*bCanReceiveCSM*/, bMovable);
+
+	for (ELightMapPolicyType LightMapPolicyType : UniformLightMapPolicyTypes)
+	{
+		CollectPSOInitializersForLMPolicy(VertexFactoryData, DrawRenderState, RenderTargetsInfo, Material, EMobileLocalLightSetting::LOCAL_LIGHTS_DISABLED, LightMapPolicyType, MeshFillMode, MeshCullMode, (EPrimitiveType)PreCacheParams.PrimitiveType, PSOInitializers);
+		if (bUseLocalLightPermutation)
+		{
+			CollectPSOInitializersForLMPolicy(VertexFactoryData, DrawRenderState, RenderTargetsInfo, Material, LocalLightSetting, LightMapPolicyType, MeshFillMode, MeshCullMode, (EPrimitiveType)PreCacheParams.PrimitiveType, PSOInitializers);
+		}
+	}
+}
+
+FMeshPassProcessor* CreateMobileSingleLayerWaterPassProcessor(ERHIFeatureLevel::Type FeatureLevel, const FScene* Scene, const FSceneView* InViewIfDynamicMeshCommand, FMeshPassDrawListContext* InDrawListContext)
+{
+	FMeshPassProcessorRenderState DrawRenderState;
+	DrawRenderState.SetDepthStencilAccess(FExclusiveDepthStencil::DepthWrite_StencilNop);
+	DrawRenderState.SetBlendState(TStaticBlendStateWriteMask<CW_RGBA, CW_RGBA, CW_RGBA, CW_RGBA>::GetRHI());
+	DrawRenderState.SetDepthStencilState(TStaticDepthStencilState<true, CF_DepthNearOrEqual>::GetRHI());
+
+	return new FMobileSingleLayerWaterPassMeshProcessor(Scene, FeatureLevel, InViewIfDynamicMeshCommand, DrawRenderState, InDrawListContext);
+}
+
+REGISTER_MESHPASSPROCESSOR_AND_PSOCOLLECTOR(MobileSingleLayerWater, CreateMobileSingleLayerWaterPassProcessor, EShadingPath::Mobile, EMeshPass::SingleLayerWaterPass, EMeshPassFlags::CachedMeshCommands | EMeshPassFlags::MainView);
diff --git a/Engine/Source/Runtime/Renderer/Private/TranslucentRendering.cpp b/Engine/Source/Runtime/Renderer/Private/TranslucentRendering.cpp
--- a/Engine/Source/Runtime/Renderer/Private/TranslucentRendering.cpp
+++ b/Engine/Source/Runtime/Renderer/Private/TranslucentRendering.cpp
@@ -318,7 +318,7 @@ public:
 
 IMPLEMENT_GLOBAL_SHADER(FCopySceneColorPS, "/Engine/Private/TranslucentLightingShaders.usf", "CopySceneColorMain", SF_Pixel);
 
-FRDGTextureRef AddCopySceneColorPass(FRDGBuilder& GraphBuilder, TArrayView<const FViewInfo> Views, FRDGTextureMSAA SceneColor, bool WithAlpha)
+FRDGTextureRef AddCopySceneColorPass(FRDGBuilder& GraphBuilder, TArrayView<const FViewInfo> Views, FRDGTextureMSAA SceneColor, bool WithAlpha, bool bSkipIfUnderwater)
 {
 	FRDGTextureRef SceneColorCopyTexture = nullptr;
 	ERenderTargetLoadAction LoadAction = ERenderTargetLoadAction::ENoAction;
@@ -329,7 +329,7 @@ FRDGTextureRef AddCopySceneColorPass(FRDGBuilder& GraphBuilder, TArrayView<const
 	{
 		const FViewInfo& View = Views[ViewIndex];
 
-		if (View.IsUnderwater())
+		if (View.IsUnderwater() && bSkipIfUnderwater)
 		{
 			continue;
 		}
@@ -343,6 +343,7 @@ FRDGTextureRef AddCopySceneColorPass(FRDGBuilder& GraphBuilder, TArrayView<const
 				break;
 			}
 		}
+		bNeedsResolve = bNeedsResolve || (MobileSupportsSM5MaterialNodes(View.GetShaderPlatform()) && View.bHasSingleLayerWaterMaterial);
 
 		if (bNeedsResolve)
 		{
@@ -392,7 +393,7 @@ FRDGTextureRef AddCopySceneColorPass(FRDGBuilder& GraphBuilder, TArrayView<const
 	return SceneColorCopyTexture;
 }
 
-static void AddCopySceneColorAlphaPass(FRDGBuilder& GraphBuilder, TArrayView<const FViewInfo> Views, FRDGTextureRef SourceTexture, FRDGTextureRef TargetTexture)
+static void AddCopySceneColorAlphaPass(FRDGBuilder& GraphBuilder, TArrayView<const FViewInfo> Views, FRDGTextureRef SourceTexture, FRDGTextureRef TargetTexture, bool bSkipIfUnderwater)
 {
 	ERenderTargetLoadAction LoadAction = ERenderTargetLoadAction::ENoAction;
 
@@ -402,7 +403,7 @@ static void AddCopySceneColorAlphaPass(FRDGBuilder& GraphBuilder, TArrayView<con
 	{
 		const FViewInfo& View = Views[ViewIndex];
 
-		if (View.IsUnderwater())
+		if (View.IsUnderwater() && bSkipIfUnderwater)
 		{
 			continue;
 		}
@@ -1864,7 +1865,7 @@ void FDeferredShadingSceneRenderer::RenderTranslucency(
 
 	if (EnumHasAnyFlags(ViewsToRender, ETranslucencyView::AboveWater))
 	{
-		SceneColorCopyTexture = AddCopySceneColorPass(GraphBuilder, InViews, SceneTextures.Color, /*WithAlpha*/ bIsTranslucentHoldoutEnabled);
+		SceneColorCopyTexture = AddCopySceneColorPass(GraphBuilder, InViews, SceneTextures.Color, /*WithAlpha*/ bIsTranslucentHoldoutEnabled, /*bSkipIfUnderwater*/ true);
 	}
 
 	// Create a shared depth texture at the correct resolution.
diff --git a/Engine/Source/Runtime/Renderer/Private/TranslucentRendering.h b/Engine/Source/Runtime/Renderer/Private/TranslucentRendering.h
--- a/Engine/Source/Runtime/Renderer/Private/TranslucentRendering.h
+++ b/Engine/Source/Runtime/Renderer/Private/TranslucentRendering.h
@@ -36,7 +36,7 @@ struct FSeparateTranslucencyDimensions
 DECLARE_GPU_DRAWCALL_STAT_EXTERN(Translucency);
 
 /** Add Copy SceneColor Pass. */
-FRDGTextureRef AddCopySceneColorPass(FRDGBuilder& GraphBuilder, TArrayView<const FViewInfo> Views, FRDGTextureMSAA SceneColor, bool WithAlpha = false);
+FRDGTextureRef AddCopySceneColorPass(FRDGBuilder& GraphBuilder, TArrayView<const FViewInfo> Views, FRDGTextureMSAA SceneColor, bool WithAlpha = false, bool bSkipIfUnderwater = true);
 
 /** Converts the the translucency pass into the respective mesh pass. */
 EMeshPass::Type TranslucencyPassToMeshPass(ETranslucencyPass::Type TranslucencyPass);
