Fix inset shadows from being incorrectly culled for player 2 in split screen.

diff --git a/Engine/Source/Runtime/Renderer/Private/ShadowSetup.cpp b/Engine/Source/Runtime/Renderer/Private/ShadowSetup.cpp
--- a/Engine/Source/Runtime/Renderer/Private/ShadowSetup.cpp
+++ b/Engine/Source/Runtime/Renderer/Private/ShadowSetup.cpp
@@ -2511,8 +2511,9 @@ void FProjectedShadowInfo::SetupMeshDrawCommandsForShadowDepth(FSceneRenderer& R
 	BaseParams.RasterContextSize = FIntPoint(ResolutionX, ResolutionY);
 
 	// Only enable min screen radius culling for uncached non-VSM
-	if ((!HasVirtualShadowMap() && CacheMode != SDCM_StaticPrimitivesOnly) 
-		|| (VirtualShadowMapPerLightCacheEntry.IsValid() && VirtualShadowMapPerLightCacheEntry->IsUncached()))
+	if (!bPerObjectOpaqueShadow && 
+		((!HasVirtualShadowMap() && CacheMode != SDCM_StaticPrimitivesOnly) 
+		|| (VirtualShadowMapPerLightCacheEntry.IsValid() && VirtualShadowMapPerLightCacheEntry->IsUncached())))
 	{
 		BaseParams.Flags |= NANITE_VIEW_MIN_SCREEN_RADIUS_CULL;
 		BaseParams.MinBoundsRadius = GMinScreenRadiusForShadowCaster;
