Removing outdated checks: Global uniform buffers aren't always bound to slot 0.

diff --git a/Engine/Source/Runtime/Renderer/Private/MeshPassProcessor.cpp b/Engine/Source/Runtime/Renderer/Private/MeshPassProcessor.cpp
--- a/Engine/Source/Runtime/Renderer/Private/MeshPassProcessor.cpp
+++ b/Engine/Source/Runtime/Renderer/Private/MeshPassProcessor.cpp
@@ -90,7 +90,6 @@ inline void SetTextureParameter(FRHIBatchedShaderParameters& BatchedParameters,
 #if PLATFORM_SUPPORTS_BINDLESS_RENDERING
 	if (Parameter.Type == EShaderParameterType::BindlessSRV)
 	{
-		check(Parameter.BufferIndex == 0);
 		BatchedParameters.SetBindlessTexture(Parameter.BaseIndex, TextureRHI);
 	}
 	else
@@ -105,7 +104,6 @@ inline void SetSrvParameter(FRHIBatchedShaderParameters& BatchedParameters, cons
 #if PLATFORM_SUPPORTS_BINDLESS_RENDERING
 	if (Parameter.Type == EShaderParameterType::BindlessSRV)
 	{
-		check(Parameter.BufferIndex == 0);
 		BatchedParameters.SetBindlessResourceView(Parameter.BaseIndex, SrvRHI);
 	}
 	else
