Support dynamic per-object shadows in a mobile multi-pass deferred mode
Fixed occasional clipping of per-object shadows on mobile
Improved shadow target allocation for per-object shadows when atlas is not used
Improved logic around when full depth pre-pass is required on mobile
Skip shadow channel allocation when its not necessary

diff --git a/Engine/Config/Android/AndroidScalability.ini b/Engine/Config/Android/AndroidScalability.ini
--- a/Engine/Config/Android/AndroidScalability.ini
+++ b/Engine/Config/Android/AndroidScalability.ini
@@ -53,6 +53,7 @@ r.Shadow.CSM.MaxCascades=1
 r.Shadow.RadiusThreshold=0.06
 r.VolumetricFog=0
 r.ContactShadows=0
+r.DFShadowQuality=0
 
 [ShadowQuality@1]
 r.ShadowQuality=1
@@ -62,6 +63,7 @@ r.Shadow.CSM.MaxCascades=1
 r.Shadow.RadiusThreshold=0.06
 r.VolumetricFog=0
 r.ContactShadows=0
+r.DFShadowQuality=1
 
 [ShadowQuality@2]
 r.ShadowQuality=2
@@ -71,6 +73,7 @@ r.Shadow.CSM.MaxCascades=2
 r.Shadow.RadiusThreshold=0.06
 r.VolumetricFog=0
 r.ContactShadows=0
+r.DFShadowQuality=1
 
 [ShadowQuality@3]
 r.ShadowQuality=2
@@ -80,6 +83,7 @@ r.Shadow.CSM.MaxCascades=3
 r.Shadow.RadiusThreshold=0.06
 r.VolumetricFog=0
 r.ContactShadows=0
+r.DFShadowQuality=2
 
 ; ------------------------------------------------------------------------------
 ; Postprocess settings
diff --git a/Engine/Config/IOS/IOSScalability.ini b/Engine/Config/IOS/IOSScalability.ini
--- a/Engine/Config/IOS/IOSScalability.ini
+++ b/Engine/Config/IOS/IOSScalability.ini
@@ -53,6 +53,7 @@ r.Shadow.CSM.MaxCascades=1
 r.Shadow.RadiusThreshold=0.06
 r.VolumetricFog=0
 r.ContactShadows=0
+r.DFShadowQuality=0
 
 [ShadowQuality@1]
 r.ShadowQuality=1
@@ -62,6 +63,7 @@ r.Shadow.CSM.MaxCascades=1
 r.Shadow.RadiusThreshold=0.06
 r.VolumetricFog=0
 r.ContactShadows=0
+r.DFShadowQuality=1
 
 [ShadowQuality@2]
 r.ShadowQuality=2
@@ -71,6 +73,7 @@ r.Shadow.CSM.MaxCascades=2
 r.Shadow.RadiusThreshold=0.06
 r.VolumetricFog=0
 r.ContactShadows=0
+r.DFShadowQuality=1
 
 [ShadowQuality@3]
 r.ShadowQuality=2
@@ -80,6 +83,7 @@ r.Shadow.CSM.MaxCascades=3
 r.Shadow.RadiusThreshold=0.06
 r.VolumetricFog=0
 r.ContactShadows=0
+r.DFShadowQuality=2
 
 ; ------------------------------------------------------------------------------
 ; Postprocess settings
diff --git a/Engine/Config/VulkanPC/DataDrivenPlatformInfo.ini b/Engine/Config/VulkanPC/DataDrivenPlatformInfo.ini
--- a/Engine/Config/VulkanPC/DataDrivenPlatformInfo.ini
+++ b/Engine/Config/VulkanPC/DataDrivenPlatformInfo.ini
@@ -266,6 +266,8 @@ MaxFeatureLevel=ES3_1
 ShaderFormat=SF_VULKAN_ES31
 bIsPC=true
 bIsMobile=true
+; Enables RDG pass merging(like on Android), helps with debugging
+bTargetsTiledGPU=true
 bSupportsMobileMultiView=true
 bSupportsIndexBufferUAVs=true
 bSupportsManualVertexFetch=false
diff --git a/Engine/Shaders/Private/MobileBasePassPixelShader.usf b/Engine/Shaders/Private/MobileBasePassPixelShader.usf
--- a/Engine/Shaders/Private/MobileBasePassPixelShader.usf
+++ b/Engine/Shaders/Private/MobileBasePassPixelShader.usf
@@ -1179,7 +1179,7 @@ void Main(
 
 		half4 LocalLightDynamicShadowFactors = 1.0f;
 
-		#if USE_SHADOWMASKTEXTURE
+		#if ENABLE_SHADOWMASKTEXTURE
 			LocalLightDynamicShadowFactors = DynamicShadowFactors;
 		#endif
 		
diff --git a/Engine/Shaders/Private/MobileLightingCommon.ush b/Engine/Shaders/Private/MobileLightingCommon.ush
--- a/Engine/Shaders/Private/MobileLightingCommon.ush
+++ b/Engine/Shaders/Private/MobileLightingCommon.ush
@@ -269,20 +269,21 @@ FLightAccumulator GetTranslucencySHLighting(FGBufferData GBuffer, FDeferredLight
 
 FDeferredLightData GetDirectionalLightData(float4 ScreenPosition, float4 SvPosition, inout half4 OutDynamicShadowFactors, inout half OutDynamicShadowing, uint EyeIndex)
 {
-#if USE_SHADOWMASKTEXTURE
+#if ENABLE_SHADOWMASKTEXTURE 
 	half4 PreviewShadowMapChannelMask = half4(1.0f, 0.0f, 0.0f, 0.0f);
-#if MOBILE_MULTI_VIEW
-	OutDynamicShadowFactors = Texture2DArraySample(MobileBasePass.ScreenSpaceShadowMaskTextureArray, MobileBasePass.ScreenSpaceShadowMaskSampler, float3(SvPositionToBufferUV(SvPosition), EyeIndex));
-#else
-	OutDynamicShadowFactors = Texture2DSample(MobileBasePass.ScreenSpaceShadowMaskTexture, MobileBasePass.ScreenSpaceShadowMaskSampler, SvPositionToBufferUV(SvPosition));
-#endif
-	OutDynamicShadowFactors = DecodeLightAttenuation(OutDynamicShadowFactors);
-		
-	PreviewShadowMapChannelMask = UnpackShadowMapChannelMask(MobileDirectionalLight.DirectionalLightShadowMapChannelMask >> 4);
-
-	OutDynamicShadowing = dot(PreviewShadowMapChannelMask, OutDynamicShadowFactors);
+	#if MOBILE_MULTI_VIEW
+		OutDynamicShadowFactors = Texture2DArraySample(MobileBasePass.ScreenSpaceShadowMaskTextureArray, MobileBasePass.ScreenSpaceShadowMaskSampler, float3(SvPositionToBufferUV(SvPosition), EyeIndex));
+	#else
+		OutDynamicShadowFactors = Texture2DSample(MobileBasePass.ScreenSpaceShadowMaskTexture, MobileBasePass.ScreenSpaceShadowMaskSampler, SvPositionToBufferUV(SvPosition));
+	#endif
+		OutDynamicShadowFactors = DecodeLightAttenuation(OutDynamicShadowFactors);
+	#if ALLOW_STATIC_LIGHTING
+		PreviewShadowMapChannelMask = UnpackShadowMapChannelMask(MobileDirectionalLight.DirectionalLightShadowMapChannelMask >> 4);
+		OutDynamicShadowing = dot(PreviewShadowMapChannelMask, OutDynamicShadowFactors);
+	#else
+		OutDynamicShadowing = OutDynamicShadowFactors.r;
+	#endif
 #else
-
 	float  ShadowPositionZ = 0;
 	half ShadowMap = MobileDirectionalLightCSM(ScreenPosition.xy, ScreenPosition.w, ShadowPositionZ);
 	OutDynamicShadowing = 1.0f;
@@ -613,7 +614,7 @@ void AccumulateLightGridLocalLighting(const FCulledLightsGridHeader InLightGridH
 		uint LightType = UnpackLightType(LightSceneInfoExtraDataPacked);
 
 		half DynamicShadowing = 1.0f;
-#if USE_SHADOWMASKTEXTURE
+#if ENABLE_SHADOWMASKTEXTURE
 		if (LightType == LIGHT_TYPE_SPOT || LightType == LIGHT_TYPE_POINT)
 		{
 			half4 PreviewShadowMapChannelMask = UnpackShadowMapChannelMask(LightSceneInfoExtraDataPacked >> 4);
diff --git a/Engine/Source/Runtime/RenderCore/Private/RenderGraphUtils.cpp b/Engine/Source/Runtime/RenderCore/Private/RenderGraphUtils.cpp
--- a/Engine/Source/Runtime/RenderCore/Private/RenderGraphUtils.cpp
+++ b/Engine/Source/Runtime/RenderCore/Private/RenderGraphUtils.cpp
@@ -675,7 +675,15 @@ void AddClearRenderTargetPass(FRDGBuilder& GraphBuilder, FRDGTextureRef Texture,
 				uint8 CurrentMipIndex = IntCastChecked<uint8>(TextureClearInfo.FirstMipIndex + MipIndex);
 
 				FRenderTargetParameters* Parameters = GraphBuilder.AllocParameters<FRenderTargetParameters>();
-				Parameters->RenderTargets[0] = FRenderTargetBinding(Texture, ERenderTargetLoadAction::EClear, CurrentMipIndex, CurrentSliceIndex);
+				if (TextureClearInfo.NumSlices == 1 && CurrentSliceIndex == 0)
+				{
+					// Use default value for a Slice, since it can be used for a pass merging
+					Parameters->RenderTargets[0] = FRenderTargetBinding(Texture, ERenderTargetLoadAction::EClear, CurrentMipIndex);
+				}
+				else
+				{
+					Parameters->RenderTargets[0] = FRenderTargetBinding(Texture, ERenderTargetLoadAction::EClear, CurrentMipIndex, CurrentSliceIndex);
+				}
 
 				GraphBuilder.AddPass(
 					RDG_EVENT_NAME("ClearRenderTarget(%s, slice %d, mip %d) %dx%d ClearAction", Texture->Name, CurrentSliceIndex, CurrentMipIndex, Texture->Desc.Extent.X, Texture->Desc.Extent.Y),
diff --git a/Engine/Source/Runtime/RenderCore/Private/RenderUtils.cpp b/Engine/Source/Runtime/RenderCore/Private/RenderUtils.cpp
--- a/Engine/Source/Runtime/RenderCore/Private/RenderUtils.cpp
+++ b/Engine/Source/Runtime/RenderCore/Private/RenderUtils.cpp
@@ -470,11 +470,12 @@ RENDERCORE_API bool MobileAllowFramebufferFetch(const FStaticShaderPlatform Plat
 		const static IConsoleVariable* CVarAllowFramebufferFetchOpenGL = IConsoleManager::Get().FindConsoleVariable(TEXT("r.Mobile.AllowFramebufferFetchOpenGL"));
 		return CVarAllowFramebufferFetchOpenGL->GetBool();
 	}
-	else
+	else if (IsVulkanMobilePlatform(Platform) || IsMetalMobilePlatform(Platform))
 	{
 		static FShaderPlatformCachedIniValue<bool> MobileAllowFramebufferFetchIniValue(TEXT("r.Mobile.AllowFramebufferFetch"));
 		return MobileAllowFramebufferFetchIniValue.Get(Platform);
 	}
+	return false;
 }
 
 RENDERCORE_API bool MobileRequiresSceneDepthAux(const FStaticShaderPlatform Platform)
@@ -595,9 +596,10 @@ RENDERCORE_API bool MobileUsesShadowMaskTexture(const FStaticShaderPlatform Plat
 {
 	const bool bMoveablePointOrSpotLightShadowsEnabled = (IsMobileMovableSpotlightShadowsEnabled(Platform) || FReadOnlyCVARCache::EnablePointLightShadows(Platform)) && MobileForwardEnableLocalLights(Platform);
 	const bool bMobileCapsuleShadowsEnabled = IsMobileCapsuleDirectShadowsEnabled(Platform);
+	const bool bForcedDepthPrepass = FReadOnlyCVARCache::MobileEarlyZPass(Platform) == 1; // In case forward has a forced depth-pre pass use use shadowmask
 	// Only distance field shadow needs to render shadow mask texture on mobile deferred, normal shadows need to be rendered separately because of handling lighting channels.
 	// Besides distance field shadow, with clustered lighting and shadow of local light enabled, shadows will render to shadow mask texture on mobile forward, lighting channels are handled in base pass shader.
-	return IsMobileDistanceFieldEnabled(Platform) || (!IsMobileDeferredShadingEnabled(Platform) && (bMobileCapsuleShadowsEnabled || bMoveablePointOrSpotLightShadowsEnabled));
+	return IsMobileDistanceFieldEnabled(Platform) || (!IsMobileDeferredShadingEnabled(Platform) && (bForcedDepthPrepass || bMobileCapsuleShadowsEnabled || bMoveablePointOrSpotLightShadowsEnabled));
 }
 
 // Whether to support more than 4 color attachments for GBuffer 
@@ -634,7 +636,14 @@ RENDERCORE_API bool MobileBasePassAlwaysUsesCSM(const FStaticShaderPlatform Plat
 
 RENDERCORE_API bool MobileUsesFullDepthPrepass(const FStaticShaderPlatform Platform)
 {
-	return IsStencilForLODDitherEnabled(Platform) || MobileUsesShadowMaskTexture(Platform) || IsMobileAmbientOcclusionEnabled(Platform) || IsUsingDBuffers(Platform) || FReadOnlyCVARCache::MobileEarlyZPass(Platform) == 1;
+	const bool bForcedDepthPrepass = (FReadOnlyCVARCache::MobileEarlyZPass(Platform) == 1 || IsUsingDBuffers(Platform) || IsStencilForLODDitherEnabled(Platform));
+	// Forward or single-pass deferred
+	if (!bForcedDepthPrepass && (!IsMobileDeferredShadingEnabled(Platform) || MobileAllowFramebufferFetch(Platform))) 
+	{
+		// Features that require full depth before lighting
+		return (MobileUsesShadowMaskTexture(Platform) || IsMobileAmbientOcclusionEnabled(Platform));
+	}
+	return bForcedDepthPrepass;
 }
 
 RENDERCORE_API bool MobileSupportsSM5MaterialNodes(const FStaticShaderPlatform Platform)
diff --git a/Engine/Source/Runtime/Renderer/Private/MobileBasePassRendering.cpp b/Engine/Source/Runtime/Renderer/Private/MobileBasePassRendering.cpp
--- a/Engine/Source/Runtime/Renderer/Private/MobileBasePassRendering.cpp
+++ b/Engine/Source/Runtime/Renderer/Private/MobileBasePassRendering.cpp
@@ -322,9 +322,7 @@ void SetupMobileBasePassUniformParameters(
 	const bool bIsMobileMultiViewEnabled = View.bIsMobileMultiViewEnabled || View.Aspects.IsMobileMultiViewEnabled();
 	BasePassParameters.DBuffer = GetDBufferParameters(GraphBuilder, MobileBasePassTextures.DBufferTextures, View.GetShaderPlatform(), bIsMobileMultiViewEnabled);
 
-	const bool bMobileUsesShadowMaskTexture = MobileUsesShadowMaskTexture(View.GetShaderPlatform());
-	
-	if (bMobileUsesShadowMaskTexture && GScreenSpaceShadowMaskTextureMobileOutputs.ScreenSpaceShadowMaskTextureMobile.IsValid())
+	if (GScreenSpaceShadowMaskTextureMobileOutputs.ScreenSpaceShadowMaskTextureMobile.IsValid())
 	{
 		FRDGTextureRef ScreenShadowMaskTexture = GraphBuilder.RegisterExternalTexture(GScreenSpaceShadowMaskTextureMobileOutputs.ScreenSpaceShadowMaskTextureMobile, TEXT("ScreenSpaceShadowMaskTextureMobile"));
 		BasePassParameters.ScreenSpaceShadowMaskTexture = ScreenShadowMaskTexture;
diff --git a/Engine/Source/Runtime/Renderer/Private/MobileBasePassRendering.h b/Engine/Source/Runtime/Renderer/Private/MobileBasePassRendering.h
--- a/Engine/Source/Runtime/Renderer/Private/MobileBasePassRendering.h
+++ b/Engine/Source/Runtime/Renderer/Private/MobileBasePassRendering.h
@@ -423,7 +423,7 @@ public:
 		}
 		OutEnvironment.SetDefine(TEXT("MERGED_LOCAL_LIGHTS_MOBILE"), MergedLocalLights);
 		OutEnvironment.SetDefine(TEXT("ENABLE_CLUSTERED_REFLECTION"), bEnableClusteredReflections ? 1u : 0u);
-		OutEnvironment.SetDefine(TEXT("USE_SHADOWMASKTEXTURE"), bMobileUsesShadowMaskTexture && !bTranslucentMaterial ? 1u : 0u);
+		OutEnvironment.SetDefine(TEXT("ENABLE_SHADOWMASKTEXTURE"), bMobileUsesShadowMaskTexture && !bTranslucentMaterial ? 1u : 0u);
 		OutEnvironment.SetDefine(TEXT("ENABLE_DBUFFER_TEXTURES"), Parameters.MaterialParameters.MaterialDomain == MD_Surface ? 1u : 0u);
 		EMobileTranslucentColorTransmittanceMode TranslucentColorTransmittanceMode = EMobileTranslucentColorTransmittanceMode::DEFAULT;
 		if (MaterialRequiresColorTransmittanceBlending(Parameters.MaterialParameters))
diff --git a/Engine/Source/Runtime/Renderer/Private/MobileDeferredShadingPass.cpp b/Engine/Source/Runtime/Renderer/Private/MobileDeferredShadingPass.cpp
--- a/Engine/Source/Runtime/Renderer/Private/MobileDeferredShadingPass.cpp
+++ b/Engine/Source/Runtime/Renderer/Private/MobileDeferredShadingPass.cpp
@@ -66,6 +66,7 @@ class FMobileDirectionalLightFunctionPS : public FMaterialShader
 	class FEnableClustredReflection		: SHADER_PERMUTATION_BOOL("ENABLE_CLUSTERED_REFLECTION");
 	class FEnablePlanarReflection		: SHADER_PERMUTATION_BOOL("ENABLE_PLANAR_REFLECTION");
 	class FEnableSkyLight				: SHADER_PERMUTATION_BOOL("ENABLE_SKY_LIGHT");
+	class FEnableScreenSpaceShadowMask  : SHADER_PERMUTATION_BOOL("ENABLE_SHADOWMASKTEXTURE");
 	class FEnableCSM					: SHADER_PERMUTATION_BOOL("ENABLE_MOBILE_CSM");
 	class FShadowQuality				: SHADER_PERMUTATION_RANGE_INT("MOBILE_SHADOW_QUALITY", 1, 3); // not using Quality=0
 	class FMobileSSRQuality 			: SHADER_PERMUTATION_ENUM_CLASS("MOBILE_SSR_QUALITY", EMobileSSRQuality);
@@ -77,6 +78,7 @@ class FMobileDirectionalLightFunctionPS : public FMaterialShader
 		FEnableClustredReflection, 
 		FEnablePlanarReflection,
 		FEnableSkyLight,
+		FEnableScreenSpaceShadowMask,
 		FEnableCSM, 
 		FShadowQuality,
 		FMobileSSRQuality,
@@ -89,8 +91,6 @@ class FMobileDirectionalLightFunctionPS : public FMaterialShader
 		SHADER_PARAMETER(FVector4f, LightFunctionParameters)
 		SHADER_PARAMETER(FVector2f, LightFunctionParameters2)
 		SHADER_PARAMETER(FVector3f, CameraRelativeLightPosition)
-		SHADER_PARAMETER_TEXTURE(Texture2D, ScreenSpaceShadowMaskTexture)
-		SHADER_PARAMETER_SAMPLER(SamplerState, ScreenSpaceShadowMaskSampler)
 		SHADER_PARAMETER_STRUCT_INCLUDE(FDFAOUpsampleParameters, DFAOUpsampleParameters)
 		SHADER_PARAMETER_STRUCT_INCLUDE(FSkyDiffuseLightingParameters, SkyDiffuseLighting)
 	END_SHADER_PARAMETER_STRUCT()
@@ -101,7 +101,6 @@ class FMobileDirectionalLightFunctionPS : public FMaterialShader
 		FMaterialShader::ModifyCompilationEnvironment(Parameters, OutEnvironment);
 		FForwardLightingParameters::ModifyCompilationEnvironment(Parameters.Platform, OutEnvironment);
 		OutEnvironment.SetDefine(TEXT("USE_LIGHT_FUNCTION"), Parameters.MaterialParameters.bIsDefaultMaterial ? 0 : 1);
-		OutEnvironment.SetDefine(TEXT("USE_SHADOWMASKTEXTURE"), MobileUsesShadowMaskTexture(Parameters.Platform) ? 1u : 0u);
 		OutEnvironment.SetDefine(TEXT("ENABLE_AMBIENT_OCCLUSION"), IsMobileAmbientOcclusionEnabled(Parameters.Platform) ? 1u : 0u);
 		OutEnvironment.SetDefine(TEXT("MATERIAL_SHADER"), 1);
 		OutEnvironment.SetDefine(TEXT("IS_MOBILE_DEFERREDSHADING_SUBPASS"), 1u);
@@ -114,7 +113,7 @@ class FMobileDirectionalLightFunctionPS : public FMaterialShader
 
 	static FPermutationDomain RemapPermutationVector(FPermutationDomain PermutationVector, EShaderPlatform Platform)
 	{
-		if (MobileUsesShadowMaskTexture(Platform))
+		if (PermutationVector.Get<FEnableScreenSpaceShadowMask>() == true)
 		{
 			PermutationVector.Set<FEnableCSM>(false);
 		}
@@ -167,22 +166,23 @@ class FMobileDirectionalLightFunctionPS : public FMaterialShader
 		bool bShadingModelSupport, 
 		bool bDynamicShadows, 
 		bool bSkyLight, 
-		bool bPlanarReflection, 
+		bool bPlanarReflection,
+		bool bScreenSpaceShadowMask,
 		EMobileSSRQuality MobileSSRQuality,
 		bool bApplySkyShadowing)
 	{
 		bool bUseClusteredLights = UseClusteredDeferredShading(View.GetShaderPlatform());
 		bool bClustredReflection = bInlineReflectionAndSky && (View.NumBoxReflectionCaptures + View.NumSphereReflectionCaptures) > 0;
 		bool bEnableSkyLight = bInlineReflectionAndSky && bSkyLight;
-		const bool bMobileUsesShadowMaskTexture = MobileUsesShadowMaskTexture(View.GetShaderPlatform());
-		int32 ShadowQuality = bDynamicShadows && !bMobileUsesShadowMaskTexture ? (int32)GetShadowQuality() : 0;
-				
+		int32 ShadowQuality = bDynamicShadows && !bScreenSpaceShadowMask ? (int32)GetShadowQuality() : 0;
+			
 		FPermutationDomain PermutationVector;
 		PermutationVector.Set<FMobileDirectionalLightFunctionPS::FEnableShadingModelSupport>(bShadingModelSupport);
 		PermutationVector.Set<FMobileDirectionalLightFunctionPS::FEnableClustredLights>(bUseClusteredLights);
 		PermutationVector.Set<FMobileDirectionalLightFunctionPS::FEnableClustredReflection>(bClustredReflection);
 		PermutationVector.Set<FMobileDirectionalLightFunctionPS::FEnablePlanarReflection>(bPlanarReflection);
 		PermutationVector.Set<FMobileDirectionalLightFunctionPS::FEnableSkyLight>(bEnableSkyLight);
+		PermutationVector.Set<FMobileDirectionalLightFunctionPS::FEnableScreenSpaceShadowMask>(bScreenSpaceShadowMask);
 		PermutationVector.Set<FMobileDirectionalLightFunctionPS::FEnableCSM>(ShadowQuality > 0);
 		PermutationVector.Set<FMobileDirectionalLightFunctionPS::FShadowQuality>(FMath::Clamp(ShadowQuality, 1, 3));
 		PermutationVector.Set<FMobileDirectionalLightFunctionPS::FMobileSSRQuality>(bInlineReflectionAndSky ? MobileSSRQuality : EMobileSSRQuality::Disabled);
@@ -281,7 +281,7 @@ public:
 		FMaterialShader::ModifyCompilationEnvironment(Parameters, OutEnvironment);
 		OutEnvironment.SetDefine(TEXT("USE_LIGHT_FUNCTION"), Parameters.MaterialParameters.bIsDefaultMaterial ? 0 : 1);
 		OutEnvironment.SetDefine(TEXT("MATERIAL_SHADER"), 1);
-		OutEnvironment.SetDefine(TEXT("USE_SHADOWMASKTEXTURE"), 0);
+		OutEnvironment.SetDefine(TEXT("ENABLE_SHADOWMASKTEXTURE"), 0);
 		OutEnvironment.SetDefine(TEXT("ENABLE_CLUSTERED_LIGHTS"), 0);
 		OutEnvironment.SetDefine(TEXT("IS_MOBILE_DEFERREDSHADING_SUBPASS"), 1u);
 
@@ -645,19 +645,6 @@ static void RenderDirectionalLight(
 	PassParameters.LightFunctionParameters = FVector4f(1.0f, 1.0f, 0.0f, 0.0f);
 	PassParameters.CameraRelativeLightPosition = GetCamRelativeLightPosition(View.ViewMatrices, DirectionalLight);
 
-	const bool bMobileUsesShadowMaskTexture = MobileUsesShadowMaskTexture(View.GetShaderPlatform());
-
-	if (bMobileUsesShadowMaskTexture && GScreenSpaceShadowMaskTextureMobileOutputs.ScreenSpaceShadowMaskTextureMobile.IsValid())
-	{
-		PassParameters.ScreenSpaceShadowMaskTexture = GScreenSpaceShadowMaskTextureMobileOutputs.ScreenSpaceShadowMaskTextureMobile->GetRHI();
-		PassParameters.ScreenSpaceShadowMaskSampler = TStaticSamplerState<SF_Point, AM_Clamp, AM_Clamp, AM_Clamp>::GetRHI();
-	}
-	else
-	{
-		PassParameters.ScreenSpaceShadowMaskTexture = GSystemTextures.WhiteDummy->GetRHI();
-		PassParameters.ScreenSpaceShadowMaskSampler = TStaticSamplerState<SF_Point, AM_Clamp, AM_Clamp, AM_Clamp>::GetRHI();
-	}
-
 	{
 		PassParameters.LightFunctionParameters2 = FVector2f(DirectionalLight.Proxy->GetLightFunctionFadeDistance(), DirectionalLight.Proxy->GetLightFunctionDisabledBrightness());
 		const FVector Scale = DirectionalLight.Proxy->GetLightFunctionScale();
@@ -717,6 +704,10 @@ static void RenderDirectionalLight(
 		ShadingModelStencilRef[PassIndex] = STENCIL_MOBILE_DEFAULTLIT_MASK | STENCIL_MOBILE_REFLECTIVE_MASK;
 	}
 
+	static auto CVarMobileSupportInsetShadows = IConsoleManager::Get().FindConsoleVariable(TEXT("r.Mobile.SupportInsetShadows"));
+	const bool bInsetShadows = (CVarMobileSupportInsetShadows != nullptr && CVarMobileSupportInsetShadows->GetInt() != 0);
+	const bool bMobileUsesShadowMaskTexture = (MobileUsesShadowMaskTexture(View.GetShaderPlatform()) || bInsetShadows);
+
 	const uint8 LightingChannelStencilValue = GetLightingChannelStencilValue(LightingChannel);
 	
 	for (int32 PassIndex = 0; PassIndex < NumPasses; ++PassIndex)
@@ -730,6 +721,7 @@ static void RenderDirectionalLight(
 			bDynamicShadows,
 			bEnableSkyLight,
 			bPlanarReflection, 
+			bMobileUsesShadowMaskTexture,
 			(PassEnableSSR & (1<<PassIndex)) ? MobileSSRQuality : EMobileSSRQuality::Disabled,
 			DynamicBentNormalAOTexture && UseDistanceFieldAO() && IsMobileDistanceFieldAOEnabled(View.GetShaderPlatform())
 		);
diff --git a/Engine/Source/Runtime/Renderer/Private/MobileShadingRenderer.cpp b/Engine/Source/Runtime/Renderer/Private/MobileShadingRenderer.cpp
--- a/Engine/Source/Runtime/Renderer/Private/MobileShadingRenderer.cpp
+++ b/Engine/Source/Runtime/Renderer/Private/MobileShadingRenderer.cpp
@@ -584,7 +584,10 @@ void FMobileSceneRenderer::InitViews(
 
 	bShouldRenderVelocities = ShouldRenderVelocities();
 
-	bRequiresShadowProjections = MobileUsesShadowMaskTexture(ShaderPlatform)
+	static auto CVarMobileSupportInsetShadows = IConsoleManager::Get().FindConsoleVariable(TEXT("r.Mobile.SupportInsetShadows"));
+	const bool bInsetShadows = (CVarMobileSupportInsetShadows != nullptr && CVarMobileSupportInsetShadows->GetInt() != 0);
+
+	bRequiresShadowProjections = (MobileUsesShadowMaskTexture(ShaderPlatform) || bInsetShadows)
 		&& ViewFamily.EngineShowFlags.Lighting
 		&& !Views[0].bIsReflectionCapture
 		&& !Views[0].bIsPlanarReflection
@@ -2360,6 +2363,9 @@ void FMobileSceneRenderer::RenderDeferredMultiPass(FRDGBuilder& GraphBuilder, FS
 			});
 	}
 
+	SceneTextures.MobileSetupMode = EMobileSceneTextureSetupMode::SceneDepth | EMobileSceneTextureSetupMode::GBuffers;
+	SceneTextures.MobileUniformBuffer = CreateMobileSceneTextureUniformBuffer(GraphBuilder, &SceneTextures, SceneTextures.MobileSetupMode);
+
 	if (bShouldRenderHZB && !bIsFullDepthPrepassEnabled)
 	{
 		RenderHZB(GraphBuilder, SceneTextures.Depth.Target);
@@ -2413,16 +2419,20 @@ void FMobileSceneRenderer::RenderDeferredMultiPass(FRDGBuilder& GraphBuilder, FS
 	}
 
 	TArray<FRDGTextureRef> DynamicBentNormalAOTextures;
-	if(bEnableDistanceFieldAO)
+	if (bEnableDistanceFieldAO)
 	{
-		SceneTextures.MobileSetupMode = 
-			EMobileSceneTextureSetupMode::SceneDepth |
-			EMobileSceneTextureSetupMode::GBuffers;
-		SceneTextures.MobileUniformBuffer = CreateMobileSceneTextureUniformBuffer(GraphBuilder, &SceneTextures, SceneTextures.MobileSetupMode);
 		FDistanceFieldAOParameters Parameters(Scene->SkyLight->OcclusionMaxDistance, Scene->SkyLight->Contrast);
 		RenderDistanceFieldLighting(GraphBuilder, SceneTextures, Parameters, DynamicBentNormalAOTextures, false, false);
 	}
 
+	if (bRequiresShadowProjections && !bIsFullDepthPrepassEnabled)
+	{
+		RDG_CSV_STAT_EXCLUSIVE_SCOPE(GraphBuilder, RenderMobileShadowProjections);
+		RDG_EVENT_SCOPE_STAT(GraphBuilder, ShadowProjection, "ShadowProjection");
+		RDG_GPU_STAT_SCOPE(GraphBuilder, ShadowProjection);
+		RenderMobileShadowProjections(GraphBuilder);
+	}
+
 	const bool bSupportsSM5Nodes = MobileSupportsSM5MaterialNodes(ShaderPlatform);
 	const bool bShouldRenderFullSLW = bSupportsSM5Nodes && ShouldRenderSingleLayerWater(Views);
 	const bool bNeedsDepthCopy = bShouldRenderFullSLW;
diff --git a/Engine/Source/Runtime/Renderer/Private/RendererScene.cpp b/Engine/Source/Runtime/Renderer/Private/RendererScene.cpp
--- a/Engine/Source/Runtime/Renderer/Private/RendererScene.cpp
+++ b/Engine/Source/Runtime/Renderer/Private/RendererScene.cpp
@@ -2262,12 +2262,27 @@ void FScene::AddLightSceneInfo_RenderThread(FLightSceneInfo* LightSceneInfo)
 	{
 		LightSceneInfo->Proxy->IESAtlasId = IESAtlas::AddTexture(IESTexture);
 	}
-
-	const EShaderPlatform ShaderPlatform = GetShaderPlatform();
-	const bool bAssignShadowMapChannel = IsForwardShadingEnabled(ShaderPlatform) ||  (IsMobilePlatform(ShaderPlatform) && MobileUsesShadowMaskTexture(ShaderPlatform));
-	if (bAssignShadowMapChannel && (LightSceneInfo->Proxy->CastsDynamicShadow() || LightSceneInfo->Proxy->GetLightFunctionMaterial()))
+		
+	if (LightSceneInfo->Proxy->CastsDynamicShadow())
 	{
-		AssignAvailableShadowMapChannelForLight(LightSceneInfo);
+		const EShaderPlatform ShaderPlatform = GetShaderPlatform();
+		bool bAssignShadowMapChannel = false;
+		if (IsMobilePlatform(ShaderPlatform))
+		{
+			if (IsStaticLightingAllowed() && MobileUsesShadowMaskTexture(ShaderPlatform))
+			{
+				AssignAvailableShadowMapChannelForLight(LightSceneInfo);
+			}
+			else
+			{
+				// Deferred and forward use single channel for shadows
+				LightSceneInfo->SetDynamicShadowMapChannel(0);
+			}
+		}
+		else if (IsForwardShadingEnabled(ShaderPlatform) || LightSceneInfo->Proxy->GetLightFunctionMaterial())
+		{
+			AssignAvailableShadowMapChannelForLight(LightSceneInfo);
+		}
 	}
 
 	ProcessAtmosphereLightAddition_RenderThread(LightSceneInfo);
diff --git a/Engine/Source/Runtime/Renderer/Private/ShadowDepthRendering.cpp b/Engine/Source/Runtime/Renderer/Private/ShadowDepthRendering.cpp
--- a/Engine/Source/Runtime/Renderer/Private/ShadowDepthRendering.cpp
+++ b/Engine/Source/Runtime/Renderer/Private/ShadowDepthRendering.cpp
@@ -735,8 +735,8 @@ void FProjectedShadowInfo::SetStateForView(FRHICommandList& RHICmdList, bool bRe
 			true,
 			X + BorderSize,
 			Y + BorderSize,
-			X + ResolutionX,
-			Y + ResolutionY
+			X + ResolutionX + BorderSize,
+			Y + ResolutionY + BorderSize
 		);
 	}
 }
diff --git a/Engine/Source/Runtime/Renderer/Private/ShadowRendering.h b/Engine/Source/Runtime/Renderer/Private/ShadowRendering.h
--- a/Engine/Source/Runtime/Renderer/Private/ShadowRendering.h
+++ b/Engine/Source/Runtime/Renderer/Private/ShadowRendering.h
@@ -1260,7 +1260,7 @@ public:
 	static bool ShouldCompilePermutation(const FGlobalShaderPermutationParameters& Parameters)
 	{
 		return IsFeatureLevelSupported(Parameters.Platform, ERHIFeatureLevel::SM5)
-			|| (bUseTransmission == 0 && SubPixelShadow == 0 && MobileUsesShadowMaskTexture(Parameters.Platform));
+			|| (bUseTransmission == 0 && SubPixelShadow == 0); // Mobile uses only those
 	}
 
 	/**
diff --git a/Engine/Source/Runtime/Renderer/Private/ShadowSetup.cpp b/Engine/Source/Runtime/Renderer/Private/ShadowSetup.cpp
--- a/Engine/Source/Runtime/Renderer/Private/ShadowSetup.cpp
+++ b/Engine/Source/Runtime/Renderer/Private/ShadowSetup.cpp
@@ -53,8 +53,6 @@ static const int32 NumCubeShadowDepthSurfaces = 5;
 /** Number of surfaces used for translucent shadows. */
 static const int32 NumTranslucencyShadowSurfaces = 2;
 
-static bool bGMobileInsetShadows = false;
-
 float GMinScreenRadiusForShadowCaster = 0.01f;
 static FAutoConsoleVariableRef CVarMinScreenRadiusForShadowCaster(
 	TEXT("r.Shadow.RadiusThreshold"),
@@ -3303,6 +3301,7 @@ void FSceneRenderer::CreatePerObjectProjectedShadow(
 	bool bTranslucentRelevance = false;
 	bool bTranslucentShadowIsVisibleThisFrame = false;
 	int32 NumBufferedFrames = FOcclusionQueryHelpers::GetNumBufferedFrames(FeatureLevel);
+	bool bCanUseShadowDepthAtlas = FProjectedShadowInfo::SupportsShadowStencilCulling(ShaderPlatform);
 
 	for (int32 ViewIndex = 0; ViewIndex < Views.Num(); ViewIndex++)
 	{
@@ -3500,8 +3499,11 @@ void FSceneRenderer::CreatePerObjectProjectedShadow(
 		{
 			// Round down to the nearest power of two so that resolution changes are always doubling or halving the resolution, which increases filtering stability
 			// Use the max resolution if the desired resolution is larger than that
-			const int32 SizeX = MaxDesiredResolution >= MaxShadowResolution ? MaxShadowResolution : (1 << (FMath::CeilLogTwo(MaxDesiredResolution) - 1));
-
+			const int32 SizeX = MaxDesiredResolution >= MaxShadowResolution ? MaxShadowResolution : 
+				bCanUseShadowDepthAtlas ? 
+				(1 << (FMath::CeilLogTwo(MaxDesiredResolution) - 1)) : 
+				(FMath::RoundUpToPowerOfTwo(MaxDesiredResolution) - SHADOW_BORDER * 2); // see AllocateAtlasedShadowDepthTargets, we round up target size there
+			
 			if (bOpaque && bCreateOpaqueObjectShadow && (bOpaqueShadowIsVisibleThisFrame || bShadowIsPotentiallyVisibleNextFrame))
 			{
 				// Create a projected shadow for this interaction's shadow.
@@ -6137,6 +6139,7 @@ void FSceneRenderer::CreateDynamicShadows(FDynamicShadowsTaskData& TaskData)
 
 	const bool bProjectEnablePointLightShadows = FReadOnlyCVARCache::EnablePointLightShadows(ShaderPlatform);
 	const bool bProjectEnableMovableSpotLightShadows = !bMobile || IsMobileMovableSpotlightShadowsEnabled(ShaderPlatform);
+	const bool bMobileInsetShadows = CVarMobileInsetShadows.GetValueOnRenderThread();
 
 	uint32 NumPointShadowCachesUpdatedThisFrame = 0;
 	uint32 NumSpotShadowCachesUpdatedThisFrame = 0;
@@ -6292,20 +6295,8 @@ void FSceneRenderer::CreateDynamicShadows(FDynamicShadowsTaskData& TaskData)
 							{
 								AddViewDependentWholeSceneShadowsForView(ViewDependentWholeSceneShadows, ViewDependentWholeSceneShadowsThatNeedCulling, VisibleLightInfo, *LightSceneInfo, GetCachedShadowMapsSize(), NumCSMCachesUpdatedThisFrame);
 							}
-
-#if !(UE_BUILD_SHIPPING || UE_BUILD_TEST)
-							bool bMobileInsetShadowsCVarValueChanged = bGMobileInsetShadows != CVarMobileInsetShadows.GetValueOnRenderThread();
-							if (bMobile && CVarMobileInsetShadows.GetValueOnRenderThread() && bMobileInsetShadowsCVarValueChanged)
-							{
-								if (!MobileUsesShadowMaskTexture(ShaderPlatform))
-								{
-									UE_LOG(LogRenderer, Warning, TEXT("r.Mobile.SupportInsetShadows is enabled on a platform not using ShadowMaskTexture, so the inset shadows projections won't be rendered."));
-								}
-							}
-#endif
-							bGMobileInsetShadows = CVarMobileInsetShadows.GetValueOnRenderThread();
-
-							if (!bMobile || bGMobileInsetShadows || IsMobileCapsuleShadowsEnabled(ShaderPlatform) || (LightSceneInfo->Proxy->CastsModulatedShadows() && !LightSceneInfo->Proxy->UseCSMForDynamicObjects() && LightSceneInfo->Proxy->HasStaticShadowing()))
+							
+							if (!bMobile || bMobileInsetShadows || IsMobileCapsuleShadowsEnabled(ShaderPlatform) || (LightSceneInfo->Proxy->CastsModulatedShadows() && !LightSceneInfo->Proxy->UseCSMForDynamicObjects() && LightSceneInfo->Proxy->HasStaticShadowing()))
 							{
 								const TArray<FLightPrimitiveInteraction*>* InteractionShadowPrimitives = LightSceneInfo->GetInteractionShadowPrimitives();
 
