Port adaptive virtual texture page table update to RDG.
Update now happens at the same time as the main page table updates and transitions are correctly handled by the graph builder.
#rnx

diff --git a/Engine/Source/Runtime/Renderer/Private/VT/AdaptiveVirtualTexture.cpp b/Engine/Source/Runtime/Renderer/Private/VT/AdaptiveVirtualTexture.cpp
--- a/Engine/Source/Runtime/Renderer/Private/VT/AdaptiveVirtualTexture.cpp
+++ b/Engine/Source/Runtime/Renderer/Private/VT/AdaptiveVirtualTexture.cpp
@@ -2,6 +2,8 @@
 
 #include "AdaptiveVirtualTexture.h"
 
+#include "RenderGraphBuilder.h"
+#include "RenderGraphUtils.h"
 #include "VT/AllocatedVirtualTexture.h"
 #include "VT/VirtualTexturePhysicalSpace.h"
 #include "VT/VirtualTextureScalability.h"
@@ -765,27 +767,8 @@ void FAdaptiveVirtualTexture::UpdateAllocations(FVirtualTextureSystem* InSystem,
 		}
 	}
 
-	// Update indirection texture
-	if (TextureUpdates.Num())
-	{
-		SCOPED_GPU_MASK(RHICmdList, FRHIGPUMask::All());
-
-		FVirtualTextureSpace* Space = InSystem->GetSpace(GetSpaceID());
-		FRHITexture* Texture = Space->GetPageTableIndirectionTexture()->GetReferencedTexture();
-
-		//todo[vt]: If we have more than 1 or 2 updates per frame then add a shader to batch updates.
-		RHICmdList.Transition({ FRHITransitionInfo(Texture, ERHIAccess::SRVMask, ERHIAccess::UAVCompute) }, ERHIPipeline::All, ERHIPipeline::Graphics);
-		for (FIndirectionTextureUpdate& TextureUpdate : TextureUpdates)
-		{
-			const FUpdateTextureRegion2D Region(TextureUpdate.X, TextureUpdate.Y, 0, 0, 1, 1);
-			RHIUpdateTexture2D((FRHITexture*)Texture, 0, Region, 4, (uint8*)&TextureUpdate.Value);
-		}
-		RHICmdList.Transition({ FRHITransitionInfo(Texture, ERHIAccess::UAVCompute, ERHIAccess::SRVMask) }, ERHIPipeline::Graphics, ERHIPipeline::All);
-	}
-
 	// Clear requests
 	RequestsToMap.Reset();
-	TextureUpdates.Reset();
 
 	// Release any producers
 	for (int32 ProducerIndex = ProducersToRelease.Num() - 1; ProducerIndex >= 0; ProducerIndex--)
@@ -796,3 +779,45 @@ void FAdaptiveVirtualTexture::UpdateAllocations(FVirtualTextureSystem* InSystem,
 		}
 	}
 }
+
+/** RDG parameters for update pass. */
+BEGIN_SHADER_PARAMETER_STRUCT(FUpdateAdaptiveIndirectionTextureParameters, )
+	RDG_TEXTURE_ACCESS(IndirectionTexture, ERHIAccess::CopyDest)
+END_SHADER_PARAMETER_STRUCT()
+
+void FAdaptiveVirtualTexture::ApplyPageTableUpdates(FVirtualTextureSystem* InSystem, FRDGBuilder& GraphBuilder, FRDGExternalAccessQueue& ExternalAccessQueue)
+{
+	if (TextureUpdates.Num())
+	{
+		FVirtualTextureSpace* Space = InSystem->GetSpace(GetSpaceID());
+		TRefCountPtr<IPooledRenderTarget> RenderTarget = Space->GetPageTableIndirectionRenderTarget();
+
+		FRDGTextureRef Texture = GraphBuilder.RegisterExternalTexture(RenderTarget, ERDGTextureFlags::ForceImmediateFirstBarrier);
+
+		FUpdateAdaptiveIndirectionTextureParameters* Parameters = GraphBuilder.AllocParameters<FUpdateAdaptiveIndirectionTextureParameters>();
+		Parameters->IndirectionTexture = Texture;
+
+		GraphBuilder.AddPass(
+			RDG_EVENT_NAME("UpdateAdaptiveIndirectionTexture"),
+			Parameters,
+			ERDGPassFlags::Copy | ERDGPassFlags::NeverCull,
+			[Texture, MovedTextureUpdates = MoveTemp(TextureUpdates)](FRDGAsyncTask, FRHICommandList& RHICmdList)
+		{
+			SCOPED_GPU_MASK(RHICmdList, FRHIGPUMask::All());
+
+			// The texture is already in the correct state from RDG, so disable the automatic RHI transitions.
+			FRHICommandListScopedAllowExtraTransitions ScopedExtraTransitions(RHICmdList, false);
+
+			//todo[vt]: If we have more than 1 or 2 updates per frame then add a shader to batch updates.
+			for (FIndirectionTextureUpdate const& TextureUpdate : MovedTextureUpdates)
+			{
+				const FUpdateTextureRegion2D Region(TextureUpdate.X, TextureUpdate.Y, 0, 0, 1, 1);
+				RHICmdList.UpdateTexture2D(Texture->GetRHI(), 0, Region, 4, (uint8*)&TextureUpdate.Value);
+			}
+		});
+		
+		TextureUpdates.Reset();
+
+		ExternalAccessQueue.Add(Texture, ERHIAccess::SRVMask, ERHIPipeline::All);
+	}
+}
diff --git a/Engine/Source/Runtime/Renderer/Private/VT/AdaptiveVirtualTexture.h b/Engine/Source/Runtime/Renderer/Private/VT/AdaptiveVirtualTexture.h
--- a/Engine/Source/Runtime/Renderer/Private/VT/AdaptiveVirtualTexture.h
+++ b/Engine/Source/Runtime/Renderer/Private/VT/AdaptiveVirtualTexture.h
@@ -8,6 +8,8 @@
 #include "VirtualTexturing.h"
 
 class FAllocatedVirtualTexture;
+class FRDGBuilder;
+class FRDGExternalAccessQueue;
 class FRHICommandListImmediate;
 class FVirtualTextureSystem;
 class FRHICommandListBase;
@@ -35,6 +37,8 @@ public:
 	static void QueuePackedAllocationRequests(FVirtualTextureSystem* InSystem, uint32 const* InRequests, uint32 InNumRequests, uint32 InFrame);
 	/** Update any allocations based on recent requests. */
 	void UpdateAllocations(FVirtualTextureSystem* InSystem, FRHICommandListImmediate& RHICmdList, uint32 InFrame);
+	/** Apply any page table updates. */
+	void ApplyPageTableUpdates(FVirtualTextureSystem* InSystem, FRDGBuilder& GraphBuilder, FRDGExternalAccessQueue& ExternalAccessQueue);
 
 	//~ Begin IAdaptiveVirtualTexture Interface.
 	virtual IAllocatedVirtualTexture* GetAllocatedVirtualTexture() override;
diff --git a/Engine/Source/Runtime/Renderer/Private/VT/VirtualTextureSpace.h b/Engine/Source/Runtime/Renderer/Private/VT/VirtualTextureSpace.h
--- a/Engine/Source/Runtime/Renderer/Private/VT/VirtualTextureSpace.h
+++ b/Engine/Source/Runtime/Renderer/Private/VT/VirtualTextureSpace.h
@@ -92,6 +92,11 @@ public:
 		return PageTableIndirection.TextureReferenceRHI.GetReference();
 	}
 
+	TRefCountPtr<IPooledRenderTarget> GetPageTableIndirectionRenderTarget() const
+	{
+		return PageTableIndirection.RenderTarget;
+	}
+
 	void				QueueUpdate( uint8 Layer, uint8 vLogSize, uint32 vAddress, uint8 vLevel, const FPhysicalTileLocation& pTileLocation);
 	void				AllocateTextures(FRDGBuilder& GraphBuilder);
 	void				FinalizeTextures(FRDGBuilder& GraphBuilder, FRDGExternalAccessQueue& ExternalAccessQueue);
diff --git a/Engine/Source/Runtime/Renderer/Private/VT/VirtualTextureSystem.cpp b/Engine/Source/Runtime/Renderer/Private/VT/VirtualTextureSystem.cpp
--- a/Engine/Source/Runtime/Renderer/Private/VT/VirtualTextureSystem.cpp
+++ b/Engine/Source/Runtime/Renderer/Private/VT/VirtualTextureSystem.cpp
@@ -2629,6 +2629,11 @@ void FVirtualTextureSystem::FinalizeRequests(FRDGBuilder& GraphBuilder, ISceneRe
 			if (Spaces[ID])
 			{
 				Spaces[ID]->ApplyUpdates(this, GraphBuilder, ExternalAccessQueue);
+
+				if (FAdaptiveVirtualTexture* AdaptiveVT = AdaptiveVTs[ID])
+				{
+					AdaptiveVT->ApplyPageTableUpdates(this, GraphBuilder, ExternalAccessQueue);
+				}
 			}
 		}
 	}
