Fix a typo in the VRS image debug preview.

diff --git a/Engine/Source/Runtime/Renderer/Private/VariableRateShading/VariableRateShadingImageManager.cpp b/Engine/Source/Runtime/Renderer/Private/VariableRateShading/VariableRateShadingImageManager.cpp
--- a/Engine/Source/Runtime/Renderer/Private/VariableRateShading/VariableRateShadingImageManager.cpp
+++ b/Engine/Source/Runtime/Renderer/Private/VariableRateShading/VariableRateShadingImageManager.cpp
@@ -702,7 +702,7 @@ void FVariableRateShadingImageManager::DrawDebugPreview(FRDGBuilder& GraphBuilde
 				OutputSceneColor);
 
 			FRHIBlendState* BlendState =
-				TStaticBlendState<CW_RGB, BO_Add, BF_SourceAlpha, BF_InverseSource1Alpha>::GetRHI();
+				TStaticBlendState<CW_RGB, BO_Add, BF_SourceAlpha, BF_InverseSourceAlpha>::GetRHI();
 			FRHIDepthStencilState* DepthStencilState = FScreenPassPipelineState::FDefaultDepthStencilState::GetRHI();
 
 			EScreenPassDrawFlags DrawFlags = EScreenPassDrawFlags::AllowHMDHiddenAreaMask;
