Material Cache, removed invalid rdg stat

diff --git a/Engine/Source/Runtime/Renderer/Private/MaterialCache/MaterialCacheRenderer.cpp b/Engine/Source/Runtime/Renderer/Private/MaterialCache/MaterialCacheRenderer.cpp
--- a/Engine/Source/Runtime/Renderer/Private/MaterialCache/MaterialCacheRenderer.cpp
+++ b/Engine/Source/Runtime/Renderer/Private/MaterialCache/MaterialCacheRenderer.cpp
@@ -1811,7 +1811,6 @@ void MaterialCacheRenderPages(FRDGBuilder& GraphBuilder)
 	for (auto&& [Guid, Bucket] : SceneExtension.TagBuckets)
 	{
 		RDG_EVENT_SCOPE(GraphBuilder, "MaterialCache");
-		RDG_GPU_STAT_SCOPE(GraphBuilder, MaterialCacheTag);
 
 		// Create render data on the graph's lifetime
 		FMaterialCacheRenderData& RenderData = *GraphBuilder.AllocObject<FMaterialCacheRenderData>();
