[HWRT] Updated FRayTracingShaderBindingDataOneFrameArray to use SceneRenderingAllocator.
- this used to be the case, but back then caused issues since SceneRenderingAllocator was not thread safe, however, since then, the underlying allocator used by SceneRenderingAllocator was changed, so this is now safe.

diff --git a/Engine/Source/Runtime/Renderer/Private/RayTracing/RayTracing.cpp b/Engine/Source/Runtime/Renderer/Private/RayTracing/RayTracing.cpp
--- a/Engine/Source/Runtime/Renderer/Private/RayTracing/RayTracing.cpp
+++ b/Engine/Source/Runtime/Renderer/Private/RayTracing/RayTracing.cpp
@@ -1462,7 +1462,7 @@ namespace RayTracing
 		const RayTracing::FSceneOptions& SceneOptions,
 		TConstArrayView<FRelevantPrimitive> RelevantStaticPrimitives,
 		FRayTracingScene& RayTracingScene,
-		TArray<FRayTracingShaderBindingData>& VisibleShaderBindingData)
+		FRayTracingShaderBindingDataOneFrameArray& VisibleShaderBindingData)
 	{
 		TRACE_CPUPROFILER_EVENT_SCOPE(RayTracingScene_AddStaticInstances);
 
@@ -1608,7 +1608,7 @@ namespace RayTracing
 		FRayTracingScene& RayTracingScene,
 		FRayTracingShaderBindingTable& RayTracingSBT,
 		TConstArrayView<FRelevantCachedPrimitive> RelevantCachedPrimitives,
-		TArray<FRayTracingShaderBindingData>& VisibleShaderBindingData)
+		FRayTracingShaderBindingDataOneFrameArray& VisibleShaderBindingData)
 	{
 		TRACE_CPUPROFILER_EVENT_SCOPE(RayTracing_AddVisibleCachedInstances);
 
diff --git a/Engine/Source/Runtime/Renderer/Public/RayTracingMeshDrawCommands.h b/Engine/Source/Runtime/Renderer/Public/RayTracingMeshDrawCommands.h
--- a/Engine/Source/Runtime/Renderer/Public/RayTracingMeshDrawCommands.h
+++ b/Engine/Source/Runtime/Renderer/Public/RayTracingMeshDrawCommands.h
@@ -225,7 +225,7 @@ struct TUseBitwiseSwap<FRayTracingShaderBindingData>
 	enum { Value = false };
 };
 
-typedef TArray<FRayTracingShaderBindingData> FRayTracingShaderBindingDataOneFrameArray;
+typedef TArray<FRayTracingShaderBindingData, SceneRenderingAllocator> FRayTracingShaderBindingDataOneFrameArray;
 
 class FRayTracingMeshCommandContext
 {
