Fix Velocity Clipped Depth performance regression. Use shaderpipeline template to remove unused output. It improves performance on some platforms.

#tests perB verifies regression is fixed

diff --git a/Engine/Source/Runtime/Renderer/Private/VelocityRendering.cpp b/Engine/Source/Runtime/Renderer/Private/VelocityRendering.cpp
--- a/Engine/Source/Runtime/Renderer/Private/VelocityRendering.cpp
+++ b/Engine/Source/Runtime/Renderer/Private/VelocityRendering.cpp
@@ -152,15 +152,44 @@ class FVelocityPS : public FMeshMaterialShader
 public:
 	DECLARE_SHADER_TYPE(FVelocityPS, MeshMaterial);
 
-	class FDimensionVelocityClippedDepthPass : SHADER_PERMUTATION_BOOL("VELOCITY_CLIPPED_DEPTH_PASS");
-	using FPermutationDomain = TShaderPermutationDomain<FDimensionVelocityClippedDepthPass>;
+	static bool ShouldCompilePermutation(const FMeshMaterialShaderPermutationParameters& Parameters)
+	{
+		return false;
+	}
+
+	static void ModifyCompilationEnvironment(const FMaterialShaderPermutationParameters& Parameters, FShaderCompilerEnvironment& OutEnvironment)
+	{
+		FMaterialShader::ModifyCompilationEnvironment(Parameters, OutEnvironment);
+		OutEnvironment.SetRenderTargetOutputFormat(0, PF_A16B16G16R16);
+
+		// We support velocity on thin translucent only with masking, and only if the material is only made of thin translucent shading model.
+		OutEnvironment.SetDefine(TEXT("VELOCITY_THIN_TRANSLUCENT_MODE"), Parameters.MaterialParameters.ShadingModels.HasOnlyShadingModel(MSM_ThinTranslucent));
+	}
+
+	FVelocityPS() = default;
+	FVelocityPS(const ShaderMetaType::CompiledShaderInitializerType& Initializer)
+		: FMeshMaterialShader(Initializer)
+	{}
+};
+
+enum EVelocityPassMode
+{
+	Velocity_Standard,
+	Velocity_ClippedDepth
+};
+
+// Use template to allow shader pipeline binding for better performance on some platforms
+template <EVelocityPassMode PassMode>
+class TVelocityPS : public FVelocityPS
+{
+public:
+	DECLARE_SHADER_TYPE(TVelocityPS, MeshMaterial);
 
 	static bool ShouldCompilePermutation(const FMeshMaterialShaderPermutationParameters& Parameters)
 	{
 		if (FVelocityVS::ShouldCompilePermutation(Parameters))
 		{
-			FPermutationDomain PermutationVector(Parameters.PermutationId);
-			const bool bIsVelocityClippedDepthPass = PermutationVector.Get<FDimensionVelocityClippedDepthPass>();
+			const bool bIsVelocityClippedDepthPass = PassMode == EVelocityPassMode::Velocity_ClippedDepth;
 
 			if (bIsVelocityClippedDepthPass)
 			{
@@ -177,22 +206,22 @@ public:
 
 	static void ModifyCompilationEnvironment(const FMaterialShaderPermutationParameters& Parameters, FShaderCompilerEnvironment& OutEnvironment)
 	{
-		FMaterialShader::ModifyCompilationEnvironment(Parameters, OutEnvironment);
-		OutEnvironment.SetRenderTargetOutputFormat(0, PF_A16B16G16R16);
+		FVelocityPS::ModifyCompilationEnvironment(Parameters, OutEnvironment);
+		const bool bIsVelocityClippedDepthPass = (PassMode == EVelocityPassMode::Velocity_ClippedDepth);
+		OutEnvironment.SetDefine(TEXT("VELOCITY_CLIPPED_DEPTH_PASS"), bIsVelocityClippedDepthPass ? 1 : 0);
 
-		// We support velocity on thin translucent only with masking, and only if the material is only made of thin translucent shading model.
-		OutEnvironment.SetDefine(TEXT("VELOCITY_THIN_TRANSLUCENT_MODE"), Parameters.MaterialParameters.ShadingModels.HasOnlyShadingModel(MSM_ThinTranslucent));
-
-		FPermutationDomain PermutationVector(Parameters.PermutationId);
-		if (PermutationVector.Get<FDimensionVelocityClippedDepthPass>())
+		if (bIsVelocityClippedDepthPass)
 		{
 			OutEnvironment.CompilerFlags.Add(CFLAG_AllowTypedUAVLoads);
 		}
 	}
 
-	FVelocityPS() = default;
-	FVelocityPS(const ShaderMetaType::CompiledShaderInitializerType& Initializer)
-		: FMeshMaterialShader(Initializer)
+	TVelocityPS()
+	{
+	}
+
+	TVelocityPS(const ShaderMetaType::CompiledShaderInitializerType& Initializer)
+		: FVelocityPS(Initializer)
 	{}
 };
 
@@ -203,6 +232,12 @@ END_GLOBAL_SHADER_PARAMETER_STRUCT()
 IMPLEMENT_SHADER_TYPE(,FVelocityVS, TEXT("/Engine/Private/VelocityShader.usf"), TEXT("MainVertexShader"), SF_Vertex);
 IMPLEMENT_SHADER_TYPE(,FVelocityPS, TEXT("/Engine/Private/VelocityShader.usf"), TEXT("MainPixelShader"), SF_Pixel);
 
+IMPLEMENT_MATERIAL_SHADER_TYPE(template<>, TVelocityPS<EVelocityPassMode::Velocity_Standard>, TEXT("/Engine/Private/VelocityShader.usf"), TEXT("MainPixelShader"), SF_Pixel);
+IMPLEMENT_MATERIAL_SHADER_TYPE(template<>, TVelocityPS<EVelocityPassMode::Velocity_ClippedDepth>, TEXT("/Engine/Private/VelocityShader.usf"), TEXT("MainPixelShader"), SF_Pixel);
+
+IMPLEMENT_SHADERPIPELINE_TYPE_VSPS(StandardVelocityPipeline, FVelocityVS, TVelocityPS<EVelocityPassMode::Velocity_Standard>, true);
+IMPLEMENT_SHADERPIPELINE_TYPE_VSPS(VelocityClippedDepthPipeline, FVelocityVS, TVelocityPS<EVelocityPassMode::Velocity_ClippedDepth>, true);
+
 EMeshPass::Type GetMeshPassFromVelocityPass(EVelocityPass VelocityPass)
 {
 	switch (VelocityPass)
@@ -877,17 +912,19 @@ bool GetVelocityPassShaders(
 	TShaderRef<FVelocityPS>& PixelShader)
 {
 	FMaterialShaderTypes ShaderTypes;
-
-	ShaderTypes.AddShaderType<FVelocityVS>();
 	
-	int32 PermutationId = 0;
+	ShaderTypes.AddShaderType<FVelocityVS>();
+
 	if (bIsTranslucentClippedDepthPass)
 	{
-		FVelocityPS::FPermutationDomain PermutationVector;
-		PermutationVector.Set<FVelocityPS::FDimensionVelocityClippedDepthPass>(bIsTranslucentClippedDepthPass);
-		PermutationId = PermutationVector.ToDimensionValueId();
+		ShaderTypes.PipelineType = &VelocityClippedDepthPipeline;
+		ShaderTypes.AddShaderType<TVelocityPS<EVelocityPassMode::Velocity_ClippedDepth>>();
+	}
+	else
+	{
+		ShaderTypes.PipelineType = &StandardVelocityPipeline;
+		ShaderTypes.AddShaderType<TVelocityPS<EVelocityPassMode::Velocity_Standard>>();
 	}
-	ShaderTypes.AddShaderType<FVelocityPS>(PermutationId);
 
 	FMaterialShaders Shaders;
 	if (!Material.TryGetShaders(ShaderTypes, VertexFactoryType, Shaders))
