[MeshEdges] Fix issue where view count in family changes during rendering (e.g. stereo), causing out-of-bounds access of view settings.

diff --git a/Engine/Source/Runtime/Renderer/Private/MeshEdgesRendering.cpp b/Engine/Source/Runtime/Renderer/Private/MeshEdgesRendering.cpp
--- a/Engine/Source/Runtime/Renderer/Private/MeshEdgesRendering.cpp
+++ b/Engine/Source/Runtime/Renderer/Private/MeshEdgesRendering.cpp
@@ -179,9 +179,10 @@ const FMeshEdgesViewSettings& GetMeshEdgesViewSettings(const FSceneViewFamily& V
 {
 	FMeshEdgesViewFamilyData* FamilyData = const_cast<FSceneViewFamily&>(ViewFamily).GetOrCreateExtentionData<FMeshEdgesViewFamilyData>();
 
-	if (FamilyData->ViewSettings.IsEmpty())
+	// Init on first access, and handles edge-case where viewcount has changed since last access.
+	if (ViewIndex >= FamilyData->ViewSettings.Num())
 	{
-		FamilyData->ViewSettings.Init(FMeshEdgesViewSettings{}, ViewFamily.Views.Num());
+		FamilyData->ViewSettings.SetNum(ViewFamily.Views.Num());
 	}
 
 	return FamilyData->ViewSettings[ViewIndex];
diff --git a/Engine/Source/Runtime/Renderer/Public/MeshEdges.h b/Engine/Source/Runtime/Renderer/Public/MeshEdges.h
--- a/Engine/Source/Runtime/Renderer/Public/MeshEdges.h
+++ b/Engine/Source/Runtime/Renderer/Public/MeshEdges.h
@@ -14,6 +14,9 @@ struct FMeshEdgesViewSettings
 	float Opacity = 1.0;
 };
 
+// Retrieve settings for how MeshEdges should be rendered for the specified view.
+// The returned object reference is owned by the ViewFamily and may be invalidated
+// when the number of views changes. Do not cache.
 RENDERER_API const FMeshEdgesViewSettings& GetMeshEdgesViewSettings(const FSceneView& View);
 RENDERER_API FMeshEdgesViewSettings& GetMeshEdgesViewSettings(FSceneView& View);
 
