Fix regression where VSM stats rendering would cause redundant rendertarget copy when stats are hidden.

#rnx

diff --git a/Engine/Source/Runtime/Renderer/Private/PostProcess/PostProcessing.cpp b/Engine/Source/Runtime/Renderer/Private/PostProcess/PostProcessing.cpp
--- a/Engine/Source/Runtime/Renderer/Private/PostProcess/PostProcessing.cpp
+++ b/Engine/Source/Runtime/Renderer/Private/PostProcess/PostProcessing.cpp
@@ -576,8 +576,8 @@ void AddPostProcessingPasses(
 #endif
 
 #if WITH_EDITOR || !UE_BUILD_SHIPPING
-	PassSequence.SetEnabled(EPass::VisualizeVirtualShadowMaps_PreEditorPrimitives, VirtualShadowMapArray != nullptr);
-	PassSequence.SetEnabled(EPass::VisualizeVirtualShadowMaps_PostEditorPrimitives, VirtualShadowMapArray != nullptr);
+	PassSequence.SetEnabled(EPass::VisualizeVirtualShadowMaps_PreEditorPrimitives, VirtualShadowMapArray != nullptr && VirtualShadowMapArray->IsVisualizePassEnabled(View, ViewIndex, EVSMVisualizationPostPass::PreEditorPrimitives));
+	PassSequence.SetEnabled(EPass::VisualizeVirtualShadowMaps_PostEditorPrimitives, VirtualShadowMapArray != nullptr && VirtualShadowMapArray->IsVisualizePassEnabled(View, ViewIndex, EVSMVisualizationPostPass::PostEditorPrimitives));
 #else
 	PassSequence.SetEnabled(EPass::VisualizeVirtualShadowMaps_PreEditorPrimitives, false);
 	PassSequence.SetEnabled(EPass::VisualizeVirtualShadowMaps_PostEditorPrimitives, false);
diff --git a/Engine/Source/Runtime/Renderer/Private/VirtualShadowMaps/VirtualShadowMapArray.cpp b/Engine/Source/Runtime/Renderer/Private/VirtualShadowMaps/VirtualShadowMapArray.cpp
--- a/Engine/Source/Runtime/Renderer/Private/VirtualShadowMaps/VirtualShadowMapArray.cpp
+++ b/Engine/Source/Runtime/Renderer/Private/VirtualShadowMaps/VirtualShadowMapArray.cpp
@@ -4580,6 +4580,22 @@ IMPLEMENT_GLOBAL_SHADER(FTonemapProjectionDebugTexturePS, "/Engine/Private/Virtu
 extern int32 GNaniteVisualizeOverdrawScale;
 #endif //!UE_BUILD_SHIPPING
 
+bool FVirtualShadowMapArray::IsVisualizePassEnabled(const FViewInfo& View, int32 ViewIndex, EVSMVisualizationPostPass Pass) const
+{
+#if !UE_BUILD_SHIPPING
+	bool bIsEnabled = View.Family->EngineShowFlags.VisualizeVirtualShadowMap;
+
+	if (CacheManager)
+	{
+		bIsEnabled = bIsEnabled || CacheManager->IsVisualizePassEnabled(View, ViewIndex, Pass);
+	}
+
+	return bIsEnabled;
+#else
+	return false;
+#endif
+}
+
 FScreenPassTexture FVirtualShadowMapArray::AddVisualizePass(FRDGBuilder& GraphBuilder, const FViewInfo& View, int32 ViewIndex, EVSMVisualizationPostPass Pass, FScreenPassTexture& SceneColor, FScreenPassRenderTarget& OverrideOutput)
 {
 	FScreenPassTexture Output = MoveTemp(SceneColor);
diff --git a/Engine/Source/Runtime/Renderer/Private/VirtualShadowMaps/VirtualShadowMapArray.h b/Engine/Source/Runtime/Renderer/Private/VirtualShadowMaps/VirtualShadowMapArray.h
--- a/Engine/Source/Runtime/Renderer/Private/VirtualShadowMaps/VirtualShadowMapArray.h
+++ b/Engine/Source/Runtime/Renderer/Private/VirtualShadowMaps/VirtualShadowMapArray.h
@@ -430,6 +430,7 @@ public:
 		TArray<Nanite::FPackedView,	SceneRenderingAllocator>& OutVirtualShadowViews) const;
 
 	// Add visualization composite pass, if enabled
+	bool IsVisualizePassEnabled(const FViewInfo& View, int32 ViewIndex, EVSMVisualizationPostPass Pass) const;
 	FScreenPassTexture AddVisualizePass(FRDGBuilder& GraphBuilder, const FViewInfo& View, int32 ViewIndex, EVSMVisualizationPostPass Pass, FScreenPassTexture& SceneColor, FScreenPassRenderTarget& Output);
 
 	//
diff --git a/Engine/Source/Runtime/Renderer/Private/VirtualShadowMaps/VirtualShadowMapCacheManager.cpp b/Engine/Source/Runtime/Renderer/Private/VirtualShadowMaps/VirtualShadowMapCacheManager.cpp
--- a/Engine/Source/Runtime/Renderer/Private/VirtualShadowMaps/VirtualShadowMapCacheManager.cpp
+++ b/Engine/Source/Runtime/Renderer/Private/VirtualShadowMaps/VirtualShadowMapCacheManager.cpp
@@ -1019,6 +1019,15 @@ FVirtualShadowMapArrayCacheManager::~FVirtualShadowMapArrayCacheManager()
 #endif
 }
 
+bool FVirtualShadowMapArrayCacheManager::IsVisualizePassEnabled(const FViewInfo& View, int32 ViewIndex, EVSMVisualizationPostPass Pass) const
+{
+#if !UE_BUILD_SHIPPING
+	return Pass == EVSMVisualizationPostPass::PostEditorPrimitives && (FVirtualShadowMap::GetEnabledStatSections() & EVSMStatSection::Clusters);
+#else
+	return false;
+#endif
+}
+
 FScreenPassTexture FVirtualShadowMapArrayCacheManager::AddVisualizePass(FRDGBuilder& GraphBuilder, const FViewInfo& View, int32 ViewIndex, EVSMVisualizationPostPass Pass, FScreenPassTexture& SceneColor, FScreenPassRenderTarget& Output)
 {
 #if !UE_BUILD_SHIPPING
diff --git a/Engine/Source/Runtime/Renderer/Private/VirtualShadowMaps/VirtualShadowMapCacheManager.h b/Engine/Source/Runtime/Renderer/Private/VirtualShadowMaps/VirtualShadowMapCacheManager.h
--- a/Engine/Source/Runtime/Renderer/Private/VirtualShadowMaps/VirtualShadowMapCacheManager.h
+++ b/Engine/Source/Runtime/Renderer/Private/VirtualShadowMaps/VirtualShadowMapCacheManager.h
@@ -344,6 +344,7 @@ public:
 	}
 
 	bool IsAccumulatingStats();
+	bool IsVisualizePassEnabled(const FViewInfo& View, int32 ViewIndex, EVSMVisualizationPostPass Pass) const;
 	FScreenPassTexture AddVisualizePass(FRDGBuilder& GraphBuilder, const FViewInfo& View, int32 ViewIndex, EVSMVisualizationPostPass Pass, FScreenPassTexture& SceneColor, FScreenPassRenderTarget& Output);
 
 	using FInstanceGPULoadBalancer = TInstanceCullingLoadBalancer<SceneRenderingAllocator>;
