Fixed: Overlapping lights warning introduced in CL# 44513467

diff --git a/Engine/Source/Runtime/Renderer/Private/RendererScene.cpp b/Engine/Source/Runtime/Renderer/Private/RendererScene.cpp
--- a/Engine/Source/Runtime/Renderer/Private/RendererScene.cpp
+++ b/Engine/Source/Runtime/Renderer/Private/RendererScene.cpp
@@ -2262,28 +2262,24 @@ void FScene::AddLightSceneInfo_RenderThread(FLightSceneInfo* LightSceneInfo)
 	{
 		LightSceneInfo->Proxy->IESAtlasId = IESAtlas::AddTexture(IESTexture);
 	}
-		
-	if (LightSceneInfo->Proxy->CastsDynamicShadow())
+
+	const EShaderPlatform ShaderPlatform = GetShaderPlatform();
+	if (IsMobilePlatform(ShaderPlatform) && LightSceneInfo->Proxy->CastsDynamicShadow())
 	{
-		const EShaderPlatform ShaderPlatform = GetShaderPlatform();
-		bool bAssignShadowMapChannel = false;
-		if (IsMobilePlatform(ShaderPlatform))
+		if (IsStaticLightingAllowed() && MobileUsesShadowMaskTexture(ShaderPlatform))
 		{
-			if (IsStaticLightingAllowed() && MobileUsesShadowMaskTexture(ShaderPlatform))
-			{
-				AssignAvailableShadowMapChannelForLight(LightSceneInfo);
-			}
-			else
-			{
-				// Deferred and forward use single channel for shadows
-				LightSceneInfo->SetDynamicShadowMapChannel(0);
-			}
+			AssignAvailableShadowMapChannelForLight(LightSceneInfo);
 		}
-		else if (IsForwardShadingEnabled(ShaderPlatform) || LightSceneInfo->Proxy->GetLightFunctionMaterial())
+		else
 		{
-			AssignAvailableShadowMapChannelForLight(LightSceneInfo);
+			// Mobile deferred and forward use single channel for shadows
+			LightSceneInfo->SetDynamicShadowMapChannel(0);
 		}
 	}
+	else if (IsForwardShadingEnabled(ShaderPlatform) && (LightSceneInfo->Proxy->CastsDynamicShadow() || LightSceneInfo->Proxy->GetLightFunctionMaterial()))
+	{
+		AssignAvailableShadowMapChannelForLight(LightSceneInfo);
+	}
 
 	ProcessAtmosphereLightAddition_RenderThread(LightSceneInfo);
 
