Remove frame number from the RDG_GPU_STAT_SCOPE_VERBOSE description since it can cause creation of a new unique FName at each frame.
Also with the new GPU Profiler the frame number is added on top of the breadcrumb stack making it visible in the Insight tags.

diff --git a/Engine/Source/Runtime/Renderer/Private/DeferredShadingRenderer.cpp b/Engine/Source/Runtime/Renderer/Private/DeferredShadingRenderer.cpp
--- a/Engine/Source/Runtime/Renderer/Private/DeferredShadingRenderer.cpp
+++ b/Engine/Source/Runtime/Renderer/Private/DeferredShadingRenderer.cpp
@@ -1863,8 +1863,7 @@ void FDeferredShadingSceneRenderer::Render(FRDGBuilder& GraphBuilder, const FSce
 	// By default, limit our GPU usage to only GPUs specified in the view masks.
 	RDG_GPU_MASK_SCOPE(GraphBuilder, ViewFamily.EngineShowFlags.PathTracing ? FRHIGPUMask::All() : AllViewsGPUMask);
 	RDG_EVENT_SCOPE(GraphBuilder, "Scene");
-	FString FrameNumDescription = FString::Printf(TEXT("%s Frame: %d"), *ViewFamily.ProfileDescription, GFrameCounterRenderThread);
-	RDG_GPU_STAT_SCOPE_VERBOSE(GraphBuilder, Unaccounted, *FrameNumDescription);
+	RDG_GPU_STAT_SCOPE_VERBOSE(GraphBuilder, Unaccounted, *ViewFamily.ProfileDescription);
 	
 	if (RendererOutput == ERendererOutput::FinalSceneColor)
 	{
