Add PSO precaching for mobile deferred directional lights

diff --git a/Engine/Source/Runtime/Engine/Private/Materials/Material.cpp b/Engine/Source/Runtime/Engine/Private/Materials/Material.cpp
--- a/Engine/Source/Runtime/Engine/Private/Materials/Material.cpp
+++ b/Engine/Source/Runtime/Engine/Private/Materials/Material.cpp
@@ -627,12 +627,11 @@ void UMaterialInterface::InitDefaultMaterials()
 
 void UMaterialInterface::PrecacheDefaultMaterialPSOs()
 {
-	if (!GIsRHIInitialized || (!RHISupportsManualVertexFetch(GMaxRHIShaderPlatform) && !IsPSOShaderPreloadingEnabled()))
+	if (!GIsRHIInitialized)
 	{
-		// Skip platforms that do not support MVF, non-MVF path needs mesh information for PSO 
 		return;
 	}
-		
+
 	TArray<FMaterialPSOPrecacheRequestID> MaterialPrecacheRequestIDs;
 
 	FPSOPrecacheParams PrecachePSOParams;
@@ -645,20 +644,33 @@ void UMaterialInterface::PrecacheDefaultMaterialPSOs()
 		AllVertexFactoryTypes.Add(VFData);
 	}
 
-
-	for (int32 Domain = 0; Domain < MD_MAX; ++Domain)
+	// Skip platforms that do not support MVF, non-MVF path needs mesh information for PSO 
+	if (RHISupportsManualVertexFetch(GMaxRHIShaderPlatform) || IsPSOShaderPreloadingEnabled())
 	{
-		if (GDefaultMaterials[Domain])
+		for (int32 Domain = 0; Domain < MD_MAX; ++Domain)
 		{
-			PrecachePSOParams.Mobility = EComponentMobility::Static;
-			if (PipelineStateCache::IsPSOPrecachingEnabled() || IsPSOShaderPreloadingEnabled())
+			if (GDefaultMaterials[Domain])
 			{
-				GDefaultMaterials[Domain]->PrecachePSOs(AllVertexFactoryTypes, PrecachePSOParams, EPSOPrecachePriority::High, MaterialPrecacheRequestIDs);
+				PrecachePSOParams.Mobility = EComponentMobility::Static;
+				if (PipelineStateCache::IsPSOPrecachingEnabled() || IsPSOShaderPreloadingEnabled())
+				{
+					GDefaultMaterials[Domain]->PrecachePSOs(AllVertexFactoryTypes, PrecachePSOParams, EPSOPrecachePriority::High, MaterialPrecacheRequestIDs);
+				}
+				PrecachePSOParams.Mobility = EComponentMobility::Movable;
+				if (PipelineStateCache::IsPSOPrecachingEnabled() || IsPSOShaderPreloadingEnabled())
+				{
+					GDefaultMaterials[Domain]->PrecachePSOs(AllVertexFactoryTypes, PrecachePSOParams, EPSOPrecachePriority::High, MaterialPrecacheRequestIDs);
+				}
 			}
-			PrecachePSOParams.Mobility = EComponentMobility::Movable;
-			if (PipelineStateCache::IsPSOPrecachingEnabled() || IsPSOShaderPreloadingEnabled())
+		}
+	}
+	else
+	{
+		if (GDefaultMaterials[MD_LightFunction])
+		{
+			if (PipelineStateCache::IsPSOPrecachingEnabled())
 			{
-				GDefaultMaterials[Domain]->PrecachePSOs(AllVertexFactoryTypes, PrecachePSOParams, EPSOPrecachePriority::High, MaterialPrecacheRequestIDs);
+				GDefaultMaterials[MD_LightFunction]->PrecachePSOs(AllVertexFactoryTypes, PrecachePSOParams, EPSOPrecachePriority::High, MaterialPrecacheRequestIDs);
 			}
 		}
 	}
@@ -4785,14 +4797,12 @@ void UMaterial::PostLoad()
 			PreloadMaterialShaderMap(GetMaterialResource(GMaxRHIShaderPlatform), Unused);
 		}
 	}
-	else if (IsDeferredDecal() || IsUIMaterial() || (MaterialDomain == MD_LightFunction) || IsPostProcessMaterial())
-	{
+	else if (IsUIMaterial() || (MaterialDomain == MD_LightFunction) ||
 		// TODO: need to pass a correct vertex declaration for non-MVF platforms
-		if (IsUIMaterial() || RHISupportsManualVertexFetch(GMaxRHIShaderPlatform))
-		{
-			FPSOPrecacheParams PSOPrecacheParams;
-			UMaterialInterface::PrecachePSOs(&FLocalVertexFactory::StaticType, PSOPrecacheParams);
-		}
+		((IsDeferredDecal() || IsPostProcessMaterial()) && RHISupportsManualVertexFetch(GMaxRHIShaderPlatform)))
+	{
+		FPSOPrecacheParams PSOPrecacheParams;
+		UMaterialInterface::PrecachePSOs(&FLocalVertexFactory::StaticType, PSOPrecacheParams);
 	}
 #if !WITH_EDITOR
 	CacheTexturesSamplingInfo();
diff --git a/Engine/Source/Runtime/Engine/Private/Materials/MaterialInstance.cpp b/Engine/Source/Runtime/Engine/Private/Materials/MaterialInstance.cpp
--- a/Engine/Source/Runtime/Engine/Private/Materials/MaterialInstance.cpp
+++ b/Engine/Source/Runtime/Engine/Private/Materials/MaterialInstance.cpp
@@ -3648,14 +3648,12 @@ void UMaterialInstance::PostLoad()
 			PreloadMaterialShaderMap(GetMaterialResource(GMaxRHIShaderPlatform), Unused);
 		}
 	}
-	else if (IsDeferredDecal() || IsUIMaterial() || IsPostProcessMaterial())
-	{
+	else if (IsUIMaterial() ||
 		// TODO: need to pass a correct vertex declaration for non-MVF platforms
-		if (IsUIMaterial() || RHISupportsManualVertexFetch(GMaxRHIShaderPlatform))
-		{
-			FPSOPrecacheParams PSOPrecacheParams;
-			UMaterialInterface::PrecachePSOs(&FLocalVertexFactory::StaticType, PSOPrecacheParams);
-		}
+		((IsDeferredDecal() || IsPostProcessMaterial()) && RHISupportsManualVertexFetch(GMaxRHIShaderPlatform)))
+	{
+		FPSOPrecacheParams PSOPrecacheParams;
+		UMaterialInterface::PrecachePSOs(&FLocalVertexFactory::StaticType, PSOPrecacheParams);
 	}
 	//DumpDebugInfo(*GLog);
 
diff --git a/Engine/Source/Runtime/Renderer/Private/MobileDeferredShadingPass.cpp b/Engine/Source/Runtime/Renderer/Private/MobileDeferredShadingPass.cpp
--- a/Engine/Source/Runtime/Renderer/Private/MobileDeferredShadingPass.cpp
+++ b/Engine/Source/Runtime/Renderer/Private/MobileDeferredShadingPass.cpp
@@ -2,6 +2,7 @@
 
 #include "MobileDeferredShadingPass.h"
 #include "BasePassRendering.h"
+#include "PSOPrecacheValidation.h"
 #include "SceneView.h"
 #include "ScenePrivate.h"
 #include "SceneProxies/SkyLightSceneProxy.h"
@@ -20,6 +21,8 @@
 
 DECLARE_GPU_STAT(DeferredShading);
 
+static const TCHAR* DeferredMobileLightMaterialPSOCollectorName = TEXT("DeferredMobileLightMaterialPSOCollector");
+
 int32 GMobileUseClusteredDeferredShading = 0;
 static FAutoConsoleVariableRef CVarMobileUseClusteredDeferredShading(
 	TEXT("r.Mobile.UseClusteredDeferredShading"),
@@ -175,19 +178,18 @@ class FMobileDirectionalLightFunctionPS : public FMaterialShader
 		bool bClustredReflection = bInlineReflectionAndSky && (View.NumBoxReflectionCaptures + View.NumSphereReflectionCaptures) > 0;
 		bool bEnableSkyLight = bInlineReflectionAndSky && bSkyLight;
 		int32 ShadowQuality = bDynamicShadows && !bScreenSpaceShadowMask ? (int32)GetShadowQuality() : 0;
-			
+
 		FPermutationDomain PermutationVector;
-		PermutationVector.Set<FMobileDirectionalLightFunctionPS::FEnableShadingModelSupport>(bShadingModelSupport);
-		PermutationVector.Set<FMobileDirectionalLightFunctionPS::FEnableClustredLights>(bUseClusteredLights);
-		PermutationVector.Set<FMobileDirectionalLightFunctionPS::FEnableClustredReflection>(bClustredReflection);
-		PermutationVector.Set<FMobileDirectionalLightFunctionPS::FEnablePlanarReflection>(bPlanarReflection);
-		PermutationVector.Set<FMobileDirectionalLightFunctionPS::FEnableSkyLight>(bEnableSkyLight);
-		PermutationVector.Set<FMobileDirectionalLightFunctionPS::FEnableScreenSpaceShadowMask>(bScreenSpaceShadowMask);
-		PermutationVector.Set<FMobileDirectionalLightFunctionPS::FEnableCSM>(ShadowQuality > 0);
-		PermutationVector.Set<FMobileDirectionalLightFunctionPS::FShadowQuality>(FMath::Clamp(ShadowQuality, 1, 3));
-		PermutationVector.Set<FMobileDirectionalLightFunctionPS::FMobileSSRQuality>(bInlineReflectionAndSky ? MobileSSRQuality : EMobileSSRQuality::Disabled);
-		extern bool UseDistanceFieldAO();
-		PermutationVector.Set<FMobileDirectionalLightFunctionPS::FSkyShadowing>(bInlineReflectionAndSky && bApplySkyShadowing);
+		PermutationVector.Set<FEnableShadingModelSupport>(bShadingModelSupport);
+		PermutationVector.Set<FEnableClustredLights>(bUseClusteredLights);
+		PermutationVector.Set<FEnableClustredReflection>(bClustredReflection);
+		PermutationVector.Set<FEnablePlanarReflection>(bPlanarReflection);
+		PermutationVector.Set<FEnableSkyLight>(bEnableSkyLight);
+		PermutationVector.Set<FEnableScreenSpaceShadowMask>(bScreenSpaceShadowMask);
+		PermutationVector.Set<FEnableCSM>(ShadowQuality > 0);
+		PermutationVector.Set<FShadowQuality>(FMath::Clamp(ShadowQuality, 1, 3));
+		PermutationVector.Set<FMobileSSRQuality>(bInlineReflectionAndSky ? MobileSSRQuality : EMobileSSRQuality::Disabled);
+		PermutationVector.Set<FSkyShadowing>(bInlineReflectionAndSky && bApplySkyShadowing);
 		return PermutationVector;
 	}
 
@@ -735,6 +737,15 @@ static void RenderDirectionalLight(
 		GraphicsPSOInit.BoundShaderState.VertexShaderRHI = VertexShader.GetVertexShader();
 		GraphicsPSOInit.BoundShaderState.PixelShaderRHI = PixelShader.GetPixelShader();
 		GraphicsPSOInit.PrimitiveType = PT_TriangleList;
+
+#if PSO_PRECACHING_VALIDATE
+		if (PSOCollectorStats::IsFullPrecachingValidationEnabled())
+		{
+			static const int32 MaterialPSOCollectorIndex = FPSOCollectorCreateManager::GetIndex(GetFeatureLevelShadingPath(GMaxRHIFeatureLevel), DeferredMobileLightMaterialPSOCollectorName);
+			PSOCollectorStats::CheckFullPipelineStateInCache(GraphicsPSOInit, EPSOPrecacheResult::Unknown, LightMaterial.Material, nullptr, nullptr, MaterialPSOCollectorIndex);
+		}
+#endif // PSO_PRECACHING_VALIDATE
+
 		SetGraphicsPipelineState(RHICmdList, GraphicsPSOInit, StencilRef);
 
 		SetShaderParametersMixedPS(RHICmdList, PixelShader, PassParameters, View, LightMaterial.MaterialProxy, *LightMaterial.Material);
@@ -1152,3 +1163,152 @@ void MobileDeferredShadingPass(
 		RenderLocalLight(RHICmdList, Scene, View, LightSceneInfo, DefaultMaterial, VisibleLightInfos);
 	}
 }
+
+class FDeferredMobileLightMaterialPSOCollector : public IPSOCollector
+{
+public:
+	FDeferredMobileLightMaterialPSOCollector(ERHIFeatureLevel::Type InFeatureLevel) :
+		IPSOCollector(FPSOCollectorCreateManager::GetIndex(GetFeatureLevelShadingPath(InFeatureLevel), DeferredMobileLightMaterialPSOCollectorName))
+	{
+	}
+
+	virtual void CollectPSOInitializers(
+		const FSceneTexturesConfig& SceneTexturesConfig,
+		const FMaterial& Material,
+		const FPSOPrecacheVertexFactoryData& VertexFactoryData,
+		const FPSOPrecacheParams& PreCacheParams,
+		TArray<FPSOPrecacheData>& PSOInitializers
+	) override final;
+};
+
+void FDeferredMobileLightMaterialPSOCollector::CollectPSOInitializers(
+	const FSceneTexturesConfig& SceneTexturesConfig,
+	const FMaterial& Material,
+	const FPSOPrecacheVertexFactoryData& VertexFactoryData,
+	const FPSOPrecacheParams& PreCacheParams,
+	TArray<FPSOPrecacheData>& PSOInitializers
+)
+{
+	FMaterialShaderTypes ShaderTypesToGetAnyPermutation;
+	ShaderTypesToGetAnyPermutation.AddShaderType<FMobileDirectionalLightFunctionPS>();
+
+	FMaterialShaders ShadersAnyPermutation;
+	if (!Material.TryGetShaders(ShaderTypesToGetAnyPermutation, nullptr, ShadersAnyPermutation))
+	{
+		return;
+	}
+
+	auto AddPSOInitializer =
+		[&](int32 PassIndex, bool bInlineReflectionAndSky, bool bOnlyDefaultLitInView, int32 ShadowQuality, bool bHasBoxSphere)
+		{
+			const uint32 CustomPassIndex = 1;
+			const int PassEnableShadingModelSupport = (!bOnlyDefaultLitInView && MobileUsesGBufferCustomData(GMaxRHIShaderPlatform)) ? (1 << CustomPassIndex) : 0;
+			const bool bDynamicSkyLight = true;
+			const bool bShadingModelSupport = PassEnableShadingModelSupport & (1 << PassIndex);
+			const bool bUseClusteredLights = UseClusteredDeferredShading(GMaxRHIShaderPlatform);
+			const bool bClustredReflection = bInlineReflectionAndSky && bHasBoxSphere;
+			const bool bPlanarReflection = false;
+			const bool bEnableSkyLight = bInlineReflectionAndSky && bDynamicSkyLight;
+			const bool bMobileUsesShadowMaskTexture = MobileUsesShadowMaskTexture(GMaxRHIShaderPlatform);
+			const bool bApplySkyShadowing = false;
+			const int32 LightingChannel = 0;
+			uint32 PassShadingModelStencilMaskAnyDefaultLit = 1;
+
+			FMobileDirectionalLightFunctionPS::FPermutationDomain PermutationVector;
+			PermutationVector.Set<FMobileDirectionalLightFunctionPS::FEnableShadingModelSupport>(bShadingModelSupport);
+			PermutationVector.Set<FMobileDirectionalLightFunctionPS::FEnableClustredLights>(bUseClusteredLights);
+			PermutationVector.Set<FMobileDirectionalLightFunctionPS::FEnableClustredReflection>(bClustredReflection);
+			PermutationVector.Set<FMobileDirectionalLightFunctionPS::FEnablePlanarReflection>(bPlanarReflection);
+			PermutationVector.Set<FMobileDirectionalLightFunctionPS::FEnableSkyLight>(bEnableSkyLight);
+			PermutationVector.Set<FMobileDirectionalLightFunctionPS::FEnableCSM>(ShadowQuality > 0);
+			PermutationVector.Set<FMobileDirectionalLightFunctionPS::FShadowQuality>(FMath::Clamp(ShadowQuality, 1, 3));
+			PermutationVector.Set<FMobileDirectionalLightFunctionPS::FMobileSSRQuality>(EMobileSSRQuality::Disabled);
+			PermutationVector.Set<FMobileDirectionalLightFunctionPS::FSkyShadowing>(bInlineReflectionAndSky && bApplySkyShadowing);
+			FMaterialShaderTypes ShaderTypesToGet;
+			ShaderTypesToGet.AddShaderType<FMobileDirectionalLightFunctionPS>(PermutationVector.ToDimensionValueId());
+
+			FMaterialShaders Shaders;
+			if (!Material.TryGetShaders(ShaderTypesToGet, nullptr, Shaders))
+			{
+				return;
+			}
+
+			TShaderRef<FMobileDirectionalLightFunctionPS> PixelShader;
+			Shaders.TryGetPixelShader(PixelShader);
+			if (!PixelShader.IsValid())
+			{
+				return;
+			}
+
+			const bool bRequired = false;
+			FRHIPixelShader* RHIPixelShader = PixelShader.GetPixelShader(bRequired);
+			if (RHIPixelShader == nullptr)
+			{
+				return;
+			}
+
+			TShaderMapRef<FPostProcessVS> VertexShader(GetGlobalShaderMap(GMaxRHIShaderPlatform));
+			FRHIVertexShader* RHIVertexShader = VertexShader.GetVertexShader(bRequired);
+			if (RHIVertexShader == nullptr)
+			{
+				return;
+			}
+
+			FGraphicsPipelineStateInitializer GraphicsPSOInit;
+			GraphicsPSOInit.RasterizerState = TStaticRasterizerState<>::GetRHI();
+
+			// Add to emissive in SceneColor
+			if (bInlineReflectionAndSky && !bDynamicSkyLight)
+			{
+				// pre-multiply SceneColor with AO
+				GraphicsPSOInit.BlendState = TStaticBlendState<CW_RGB, BO_Add, BF_One, BF_SourceAlpha>::GetRHI();
+			}
+			else
+			{
+				GraphicsPSOInit.BlendState = TStaticBlendState<CW_RGB, BO_Add, BF_One, BF_One>::GetRHI();
+			}
+			SetDirectionalLightDepthStencilState(GraphicsPSOInit, LightingChannel, (PassShadingModelStencilMaskAnyDefaultLit & (1 << PassIndex)) == 0);
+
+			GraphicsPSOInit.BoundShaderState.VertexDeclarationRHI = GFilterVertexDeclaration.VertexDeclarationRHI;
+			GraphicsPSOInit.BoundShaderState.VertexShaderRHI = RHIVertexShader;
+			GraphicsPSOInit.BoundShaderState.PixelShaderRHI = RHIPixelShader;
+			GraphicsPSOInit.PrimitiveType = PT_TriangleList;
+
+			FGraphicsPipelineRenderTargetsInfo RenderTargetsInfo;
+			RenderTargetsInfo.NumSamples = 1;
+			AddRenderTargetInfo(SceneTexturesConfig.ColorFormat, SceneTexturesConfig.ColorCreateFlags, RenderTargetsInfo);
+
+			SetupDepthStencilInfo(PF_DepthStencil, SceneTexturesConfig.DepthCreateFlags, ERenderTargetLoadAction::ELoad,
+				ERenderTargetLoadAction::ELoad, FExclusiveDepthStencil::DepthRead_StencilWrite, RenderTargetsInfo);
+
+			GraphicsPSOInit.StatePrecachePSOHash = RHIComputeStatePrecachePSOHash(GraphicsPSOInit);
+			ApplyTargetsInfo(GraphicsPSOInit, RenderTargetsInfo);
+
+			FPSOPrecacheData PSOPrecacheData;
+			PSOPrecacheData.bRequired = true;
+			PSOPrecacheData.Type = FPSOPrecacheData::EType::Graphics;
+			PSOPrecacheData.GraphicsPSOInitializer = GraphicsPSOInit;
+#if PSO_PRECACHING_VALIDATE
+			PSOPrecacheData.PSOCollectorIndex = PSOCollectorIndex;
+			PSOPrecacheData.VertexFactoryType = nullptr;
+#endif // PSO_PRECACHING_VALIDATE
+
+			PSOInitializers.Add(MoveTemp(PSOPrecacheData));
+		};
+
+	// int32 PassIndex, bool bInlineReflectionAndSky, bool bOnlyDefaultLitInView, int32 ShadowQuality, bool bHasBoxSphere
+	AddPSOInitializer(0, true, false, 3, true);
+	AddPSOInitializer(1, true, false, 3, true);
+	AddPSOInitializer(0, false, true, 0, false);
+	AddPSOInitializer(0, true, false, 0, false);
+	AddPSOInitializer(1, true, false, 0, false);
+	AddPSOInitializer(0, false, false, 0, false);
+	AddPSOInitializer(1, false, false, 0, false);
+	AddPSOInitializer(0, true, true, 0, false);
+}
+
+IPSOCollector* CreateDeferredMobileLightMaterialPSOCollector(ERHIFeatureLevel::Type FeatureLevel)
+{
+	return new FDeferredMobileLightMaterialPSOCollector(FeatureLevel);
+}
+FRegisterPSOCollectorCreateFunction RegisterDeferredMobileLightMaterialPSOCollector(&CreateDeferredMobileLightMaterialPSOCollector, EShadingPath::Mobile, DeferredMobileLightMaterialPSOCollectorName);
