AVSM Cascades for directional light

Three different algorithms for generating cascade shadow maps for the directional light. Auto-fit generates the highest-quality shadow map for a still image. Due to the nature of auto-fitting the contained volumes, the shadow quality may immediately change in fidelity based on nearby volumes popping in and out of the cascade level.

View-based clipmaps generates the most stable shadow map under motion using pixel snapping, but is very conservative in its bounds reporting, and usually generates a lower-quality shadow as a result.

View-based frustums generates a higher quality shadow map that is stable under motion, except in the very near-field, where camera rotation can cause shimmering even with pixel-snapping.

r.HeterogeneousVolumes.Shadows.Cascades
0: Disabled
1: View-based frustums (default)
2. View-based clipmaps
3. Auto-fit

r.HeterogeneousVolumes.Shadows.Cascades.PixelSnapping
Cascade quantification unit in pixel-space (default=16)

#test User testing with shadow-casting directional lights. Modification of cascade controls out to 30,000 units, to coincide with default shadow trace distance.

diff --git a/Engine/Shaders/Private/HeterogeneousVolumes/HeterogeneousVolumesAdaptiveVolumetricShadowMapSampling.ush b/Engine/Shaders/Private/HeterogeneousVolumes/HeterogeneousVolumesAdaptiveVolumetricShadowMapSampling.ush
--- a/Engine/Shaders/Private/HeterogeneousVolumes/HeterogeneousVolumesAdaptiveVolumetricShadowMapSampling.ush
+++ b/Engine/Shaders/Private/HeterogeneousVolumes/HeterogeneousVolumesAdaptiveVolumetricShadowMapSampling.ush
@@ -18,9 +18,10 @@
 struct FAdaptiveVolumetricShadowMap
 {
 	float4x4 TranslatedWorldToShadow[6];
-	float3 TranslatedWorldOrigin;
+	float3 TranslatedWorldOrigin[6];
+	float4 TranslatedWorldPlane[6];
+	float4 SplitDepths[6];
 	float DownsampleFactor;
-	float4 TranslatedWorldPlane;
 
 	int2 Resolution;
 	int NumShadowMatrices;
@@ -53,7 +54,7 @@ bool AVSM_PixelAndDepth(
 			Pixel = clamp(ShadowPosition.xy * ShadowMap.Resolution - 0.5, 0, ShadowMap.Resolution - 1);
 			if (ShadowMap.bIsDirectionalLight)
 			{
-				Depth = dot(ShadowMap.TranslatedWorldPlane.xyz, TranslatedWorldPosition) + ShadowMap.TranslatedWorldPlane.w;
+				Depth = dot(ShadowMap.TranslatedWorldPlane[Face].xyz, TranslatedWorldPosition) + ShadowMap.TranslatedWorldPlane[Face].w;
 			}
 			return true;
 		}
@@ -434,39 +435,85 @@ float AVSM_SampleTransmittance(
 		return 1.0;
 	}
 
-	float Transmittance = 1.0;
+	float Transmittance[] = { 1.0, 1.0 };
+	float Alpha = 1.0;
 
-	int Face = 0;
-	if (AVSM.NumShadowMatrices > 1)
+	float Depth = length(TranslatedWorldPosition - View.TranslatedWorldCameraOrigin);
+
+	int Face[] = {
+		ShadowMap.NumShadowMatrices - 1,
+		ShadowMap.NumShadowMatrices - 1
+	};
+
+#if 0
+	if (View.GeneralPurposeTweak >= 0)
+	{
+		Face[0] = max(View.GeneralPurposeTweak - 1, 0);
+		Face[1] = min(View.GeneralPurposeTweak, ShadowMap.NumShadowMatrices - 1);
+	}
+	else
+#endif
+	if (AVSM.bIsDirectionalLight)
 	{
-		float3 LightDir = TranslatedWorldPosition - ShadowMap.TranslatedWorldOrigin;
+		// Find the highest quality cascade for the shade point
+		for (; Face[1] > 0; --Face[1])
+		{
+			float Far = ShadowMap.SplitDepths[Face[1]].y;
+			if (Depth <= Far)
+			{
+				break;
+			}
+		}
+		Face[0] = max(Face[1] - 1, 0);
+
+		// Apply near-far distance fades
+		float Near = ShadowMap.SplitDepths[Face[1]].x;
+		float NearFadeDist = ShadowMap.SplitDepths[Face[1]].z;
+		if (Depth - Near < NearFadeDist)
+		{
+			Alpha = saturate((Depth - Near) / NearFadeDist);
+			Face[0] = min(Face[1] + 1, ShadowMap.NumShadowMatrices - 1);
+		}
+
+		float Far = ShadowMap.SplitDepths[Face[1]].y;
+		float FarFadeDist = ShadowMap.SplitDepths[Face[1]].w;
+		if (Far - Depth < FarFadeDist)
+		{
+			Alpha = saturate((Far - Depth) / FarFadeDist);
+			Face[0] = max(Face[1] - 1, 0);
+		}
+	}
+	else if (AVSM.NumShadowMatrices > 1)
+	{
+		float3 LightDir = TranslatedWorldPosition - ShadowMap.TranslatedWorldOrigin[0];
 		float3 LightDirAbs = abs(LightDir);
 
-		Face = LightDirAbs[1] > LightDirAbs[0] ? 1 : 0;
-		Face = LightDirAbs[2] > LightDirAbs[Face] ? 2 : Face;
+		Face[0] = LightDirAbs[1] > LightDirAbs[0] ? 1 : 0;
+		Face[0] = LightDirAbs[2] > LightDirAbs[Face[0]] ? 2 : Face[0];
 
-		int FaceOffset = (sign(LightDir[Face]) > 0) ? 1 : 0;
-		Face = Face * 2 + FaceOffset;
+		int FaceOffset = (sign(LightDir[Face[0]]) > 0) ? 1 : 0;
+		Face[0] = Face[1] = Face[0] * 2 + FaceOffset;
 	}
 
 	float4 ShadowPositionAsFloat4 = 0;
 	float3 ShadowPosition = 0;
+	for (int Index = 0; Index < 2; ++Index)
 	{
-		ShadowPositionAsFloat4 = mul(float4(TranslatedWorldPosition, 1), ShadowMap.TranslatedWorldToShadow[Face]);
+		ShadowPositionAsFloat4 = mul(float4(TranslatedWorldPosition, 1), ShadowMap.TranslatedWorldToShadow[Face[Index]]);
 		if (ShadowPositionAsFloat4.w > 0)
 		{
 			ShadowPosition = ShadowPositionAsFloat4.xyz / ShadowPositionAsFloat4.w;
 			if (all(ShadowPosition.xyz > 0) || all(ShadowPosition.xy < 1))
 			{
 				float2 ShadowMapCoord = clamp(ShadowPosition.xy * ShadowMap.Resolution - 0.5, 0, ShadowMap.Resolution - 1);
-				float ShadowZ = length(TranslatedWorldPosition - ShadowMap.TranslatedWorldOrigin);
+				float ShadowZ = length(TranslatedWorldPosition - ShadowMap.TranslatedWorldOrigin[Face[Index]]);
 				if (ShadowMap.bIsDirectionalLight)
 				{
-					ShadowZ = dot(ShadowMap.TranslatedWorldPlane.xyz, TranslatedWorldPosition) + ShadowMap.TranslatedWorldPlane.w;
+					ShadowZ = dot(ShadowMap.TranslatedWorldPlane[Face[Index]].xyz, TranslatedWorldPosition) + ShadowMap.TranslatedWorldPlane[Face[Index]].w;
 				}
 
-			#define BILINEAR_INTERPOLATION 1
-			#if BILINEAR_INTERPOLATION
+#define BILINEAR_INTERPOLATION 1
+#if BILINEAR_INTERPOLATION
 				// Bilinear interpolation of shadow data
 				int2 ShadowMapCoordX[4] =
 				{
@@ -478,24 +525,24 @@ float AVSM_SampleTransmittance(
 
 				float TransmittanceX[4] =
 				{
-					AVSM_Sample(ShadowMap, Face, ShadowMapCoordX[0], ShadowZ),
-					AVSM_Sample(ShadowMap, Face, ShadowMapCoordX[1], ShadowZ),
-					AVSM_Sample(ShadowMap, Face, ShadowMapCoordX[2], ShadowZ),
-					AVSM_Sample(ShadowMap, Face, ShadowMapCoordX[3], ShadowZ)
+					AVSM_Sample(ShadowMap, Face[Index], ShadowMapCoordX[0], ShadowZ),
+					AVSM_Sample(ShadowMap, Face[Index], ShadowMapCoordX[1], ShadowZ),
+					AVSM_Sample(ShadowMap, Face[Index], ShadowMapCoordX[2], ShadowZ),
+					AVSM_Sample(ShadowMap, Face[Index], ShadowMapCoordX[3], ShadowZ)
 				};
 
 				float2 Weight = frac(ShadowMapCoord);
 				float TransmittanceY0 = lerp(TransmittanceX[0], TransmittanceX[1], Weight.x);
 				float TransmittanceY1 = lerp(TransmittanceX[2], TransmittanceX[3], Weight.x);
-				Transmittance = lerp(TransmittanceY0, TransmittanceY1, Weight.y);
-			#else
-				Transmittance = AVSM_Sample(ShadowMap, Face, ShadowMapCoord, ShadowZ);
-			#endif
+				Transmittance[Index] = lerp(TransmittanceY0, TransmittanceY1, Weight.y);
+#else
+				Transmittance[Index] = AVSM_Sample(ShadowMap, Face[Index], ShadowMapCoord, ShadowZ);
+#endif
 			}
 		}
 	}
 
-	return Transmittance;
+	return ((1.0 - Alpha) * Transmittance[0]) + (Alpha * Transmittance[1]);
 }
 
 FAdaptiveVolumetricShadowMap GetAdaptiveVolumetricShadowMap()
@@ -504,9 +551,10 @@ FAdaptiveVolumetricShadowMap GetAdaptiveVolumetricShadowMap()
 	for (int i = 0; i < 6; ++i)
 	{
 		ShadowMap.TranslatedWorldToShadow[i] = AVSM.TranslatedWorldToShadow[i];
+		ShadowMap.TranslatedWorldOrigin[i] = AVSM.TranslatedWorldOrigin[i];
+		ShadowMap.TranslatedWorldPlane[i] = AVSM.TranslatedWorldPlane[i];
+		ShadowMap.SplitDepths[i] = AVSM.SplitDepths[i];
 	}
-	ShadowMap.TranslatedWorldOrigin = AVSM.TranslatedWorldOrigin;
-	ShadowMap.TranslatedWorldPlane = AVSM.TranslatedWorldPlane;
 
 	ShadowMap.Resolution = AVSM.Resolution;
 	ShadowMap.DownsampleFactor = AVSM.DownsampleFactor;
@@ -530,9 +578,10 @@ FAdaptiveVolumetricShadowMap GetAdaptiveVolumetricShadowMapForCamera()
 	for (int i = 0; i < 6; ++i)
 	{
 		ShadowMap.TranslatedWorldToShadow[i] = CameraAVSM.TranslatedWorldToShadow[i];
+		ShadowMap.TranslatedWorldOrigin[i] = CameraAVSM.TranslatedWorldOrigin[i];
+		ShadowMap.TranslatedWorldPlane[i] = CameraAVSM.TranslatedWorldPlane[i];
+		ShadowMap.SplitDepths[i] = CameraAVSM.SplitDepths[i];
 	}
-	ShadowMap.TranslatedWorldOrigin = CameraAVSM.TranslatedWorldOrigin;
-	ShadowMap.TranslatedWorldPlane = CameraAVSM.TranslatedWorldPlane;
 
 	ShadowMap.Resolution = CameraAVSM.Resolution;
 	ShadowMap.DownsampleFactor = CameraAVSM.DownsampleFactor;
diff --git a/Engine/Shaders/Private/HeterogeneousVolumes/HeterogeneousVolumesLiveShadingShadows.usf b/Engine/Shaders/Private/HeterogeneousVolumes/HeterogeneousVolumesLiveShadingShadows.usf
--- a/Engine/Shaders/Private/HeterogeneousVolumes/HeterogeneousVolumesLiveShadingShadows.usf
+++ b/Engine/Shaders/Private/HeterogeneousVolumes/HeterogeneousVolumesLiveShadingShadows.usf
@@ -54,7 +54,7 @@ int StochasticFilteringMode;
 int NumShadowMatrices;
 float4x4 TranslatedWorldToShadow[6];
 float4x4 ShadowToTranslatedWorld[6];
-float3 TranslatedWorldOrigin;
+float4 TranslatedWorldOrigin[6];
 
 int2 ShadowResolution;
 int MaxSampleCount;
@@ -585,9 +585,4 @@ void RenderVolumetricShadowMapForLightWithLiveShadingCS(
 		AVSM_SetLinkedListPrevPtr(RWVolumetricShadowLinkedListBuffer[PixelIndex], AVSM_NULL_PTR);
 		AVSM_SetLinkedListNextPtr(RWVolumetricShadowLinkedListBuffer[PixelIndex], AVSM_NULL_PTR);
 	}
-
-	if (Face == ShadowDebugTweak)
-	{
-		RWBeerShadowMapTexture[PixelCoord] = float4(Transmittance, 1);
-	}
 }
diff --git a/Engine/Shaders/Private/HeterogeneousVolumes/HeterogeneousVolumesVoxelGridShadows.usf b/Engine/Shaders/Private/HeterogeneousVolumes/HeterogeneousVolumesVoxelGridShadows.usf
--- a/Engine/Shaders/Private/HeterogeneousVolumes/HeterogeneousVolumesVoxelGridShadows.usf
+++ b/Engine/Shaders/Private/HeterogeneousVolumes/HeterogeneousVolumesVoxelGridShadows.usf
@@ -299,6 +299,36 @@ float RelativeErrorThreshold;
 RWBuffer<int> RWVolumetricShadowLinkedListAllocatorBuffer;
 RWStructuredBuffer<uint2> RWVolumetricShadowLinkedListBuffer;
 
+uint GetLinearIndex(uint Face, uint2 PixelCoord)
+{
+	uint LinearIndex = Face * ShadowResolution.x * ShadowResolution.y + PixelCoord.y * ShadowResolution.x + PixelCoord.x;
+	return LinearIndex;
+}
+
+[numthreads(THREADGROUP_SIZE_2D, THREADGROUP_SIZE_2D, 1)]
+void ClearVolumetricShadowMapCS(
+	uint2 GroupThreadId : SV_GroupThreadID,
+	uint3 DispatchThreadId : SV_DispatchThreadID
+)
+{
+	// Create screen ray
+	if (any(DispatchThreadId.xy >= ShadowResolution))
+	{
+		return;
+	}
+
+	uint Face = DispatchThreadId.z;
+	uint2 PixelCoord = DispatchThreadId.xy;
+	uint LinearIndex = GetLinearIndex(Face, PixelCoord);
+	int PixelIndex = LinearIndex * MaxSampleCount;
+
+	float Depth = 0.0;
+	float Transmittance = 1.0;
+	AVSM_SetLinkedListSampleData(RWVolumetricShadowLinkedListBuffer[PixelIndex], AVSM_CreateSampleData(Depth, Transmittance));
+	AVSM_SetLinkedListPrevPtr(RWVolumetricShadowLinkedListBuffer[PixelIndex], AVSM_NULL_PTR);
+	AVSM_SetLinkedListNextPtr(RWVolumetricShadowLinkedListBuffer[PixelIndex], AVSM_NULL_PTR);
+}
+
 // Debug
 struct FVolumetricShadowMapDebugData
 {
@@ -311,12 +341,6 @@ struct FVolumetricShadowMapDebugData
 
 RWStructuredBuffer<FVolumetricShadowMapDebugData> RWDebugBuffer;
 
-uint GetLinearIndex(uint Face, uint2 PixelCoord)
-{
-	uint LinearIndex = Face * ShadowResolution.x * ShadowResolution.y + PixelCoord.y * ShadowResolution.x + PixelCoord.x;
-	return LinearIndex;
-}
-
 uint FindSampleIndexWithLowestError(
 	uint PixelIndex
 )
@@ -827,6 +851,16 @@ struct FBeerShadowMapData
 	float Tau[2];
 };
 
+float4 PackBeerShadowMapData(FBeerShadowMapData Data)
+{
+	float4 PackedData;
+	PackedData.r = Data.T[0];
+	PackedData.g = Data.T[1];
+	PackedData.b = Data.Tau[0];
+	PackedData.a = Data.Tau[1];
+	return PackedData;
+}
+
 FBeerShadowMapData UnpackBeerShadowMapData(float4 PackedData)
 {
 	FBeerShadowMapData Data;
@@ -919,3 +953,39 @@ void ConvertBeerLawShadowMapToVolumetricShadowMapCS(
 		RWVolumetricShadowTransmittanceBuffer[TopLevelOffset + TopLevelIndex] = AVSM_PackSampleData(AVSM_CreateSampleData(ShadowData[1].X, ShadowData[1].Tau), IndirectionData);
 	}
 }
+
+[numthreads(THREADGROUP_SIZE_2D, THREADGROUP_SIZE_2D, 1)]
+void ConvertVolumetricShadowMapToBeerLawShadowMapCS(
+	uint3 DispatchThreadId : SV_DispatchThreadID
+)
+{
+	uint Face = DispatchThreadId.z;
+	uint2 PixelCoord = DispatchThreadId.xy;
+	if (any(PixelCoord >= ShadowResolution))
+	{
+		return;
+	}
+
+	uint LinearIndex = GetLinearIndex(Face, PixelCoord);
+	int PixelIndex = LinearIndex * MaxSampleCount;
+
+	FAVSMSampleData SampleData[] = {
+		AVSM_GetLinkedListSampleData(VolumetricShadowLinkedListBuffer0[PixelIndex]),
+		AVSM_CreateSampleData(0.0, 1.0)
+	};
+
+	int IndexPtr = AVSM_GetLinkedListNextPtr(VolumetricShadowLinkedListBuffer0[PixelIndex]);
+	while (IndexPtr != AVSM_NULL_PTR)
+	{
+		SampleData[1] = AVSM_GetLinkedListSampleData(VolumetricShadowLinkedListBuffer0[PixelIndex + IndexPtr]);
+		IndexPtr = AVSM_GetLinkedListNextPtr(VolumetricShadowLinkedListBuffer0[PixelIndex + IndexPtr]);
+	}
+
+	FBeerShadowMapData BeerData;
+	BeerData.T[0] = SampleData[0].X;
+	BeerData.T[1] = SampleData[1].X;
+	BeerData.Tau[0] = SampleData[0].Tau;
+	BeerData.Tau[1] = SampleData[1].Tau;
+	uint2 BeerPixelCoord = uint2(Face * ShadowResolution.x + PixelCoord.x, PixelCoord.y);
+	RWBeerShadowMapTexture[BeerPixelCoord] = PackBeerShadowMapData(BeerData);
+}
diff --git a/Engine/Source/Runtime/Renderer/Private/HeterogeneousVolumes/HeterogeneousVolumes.cpp b/Engine/Source/Runtime/Renderer/Private/HeterogeneousVolumes/HeterogeneousVolumes.cpp
--- a/Engine/Source/Runtime/Renderer/Private/HeterogeneousVolumes/HeterogeneousVolumes.cpp
+++ b/Engine/Source/Runtime/Renderer/Private/HeterogeneousVolumes/HeterogeneousVolumes.cpp
@@ -701,6 +701,11 @@ namespace HeterogeneousVolumes
 		return static_cast<EIndirectLightingMode>(FMath::Clamp(CVarHeterogeneousVolumesIndirectLightingMode.GetValueOnRenderThread(), 0, 2));
 	}
 
+	bool ShouldDepthSort()
+	{
+		return CVarHeterogeneousVolumesDepthSort.GetValueOnAnyThread() != 0;
+	}
+
 	bool ShouldApplyHeightFog()
 	{
 		return CVarHeterogeneousVolumesApplyHeightFog.GetValueOnRenderThread() != 0;
@@ -1160,8 +1165,7 @@ void FDeferredShadingSceneRenderer::RenderHeterogeneousVolumes(
 				}
 
 				// Provide coarse depth-sorting, based on camera-distance to world centroid
-				bool bDepthSort = CVarHeterogeneousVolumesDepthSort.GetValueOnRenderThread() == 1;
-				if (bDepthSort)
+				if (HeterogeneousVolumes::ShouldDepthSort())
 				{
 					struct FDepthCompareHeterogeneousVolumes
 					{
@@ -1301,11 +1305,18 @@ void FDeferredShadingSceneRenderer::RenderHeterogeneousVolumes(
 								FMatrix44f ShadowToTranslatedWorld[] = {
 									TranslatedWorldToShadow[0].Inverse()
 								};
-								FVector3f TranslatedWorldOrigin = ShadowToTranslatedWorld[0].GetOrigin();
+								FVector3f TranslatedWorldOrigin[] = {
+									ShadowToTranslatedWorld[0].GetOrigin()
+								};
+								FVector4f TranslatedWorldPlane[] = {
+									FVector4f::Zero()
+								};
+								FVector4f SplitDepths[] = {
+									FVector4f::Zero()
+								};
 
 								// Generic data
 								int32 MaxSampleCount = 2;
-								FVector4f TranslatedWorldPlane = FVector4f::Zero();
 								float DownsampleFactor = HeterogeneousVolumes::GetDownsampleFactor();
 								bool bIsDirectionalLight = false;
 								FRDGBufferRef VolumetricShadowLinkedListBuffer = GSystemTextures.GetDefaultStructuredBuffer(GraphBuilder, sizeof(FAVSMLinkedListPackedData));
@@ -1316,6 +1327,7 @@ void FDeferredShadingSceneRenderer::RenderHeterogeneousVolumes(
 									TranslatedWorldOrigin,
 									TranslatedWorldPlane,
 									TranslatedWorldToShadow,
+									SplitDepths,
 									ShadowMapResolution,
 									DownsampleFactor,
 									NumShadowMatrices,
diff --git a/Engine/Source/Runtime/Renderer/Private/HeterogeneousVolumes/HeterogeneousVolumes.h b/Engine/Source/Runtime/Renderer/Private/HeterogeneousVolumes/HeterogeneousVolumes.h
--- a/Engine/Source/Runtime/Renderer/Private/HeterogeneousVolumes/HeterogeneousVolumes.h
+++ b/Engine/Source/Runtime/Renderer/Private/HeterogeneousVolumes/HeterogeneousVolumes.h
@@ -145,6 +145,7 @@ namespace HeterogeneousVolumes
 	bool UseLightingCacheForInscattering();
 	bool UseLightingCacheForTransmittance();
 	bool UseAdaptiveVolumetricShadowMapForSelfShadowing(const FPrimitiveSceneProxy* PrimitiveSceneProxy);
+	bool ShouldDepthSort();
 	bool ShouldApplyHeightFog();
 	bool ShouldApplyVolumetricFog();
 	bool SupportsOverlappingVolumes();
@@ -153,7 +154,18 @@ namespace HeterogeneousVolumes
 	bool UseAnalyticDerivatives();
 	bool SupportsLightType(uint32 LightType);
 	bool SupportsShadowForLightType(uint32 LightType);
-
+	bool SupportsCascadeShadowsForDirectionalLight();
+	
+	enum class ECascadeShadowMode
+	{
+		Disabled,
+		Frustums,
+		Clipmaps,
+		Autofit
+	};
+	ECascadeShadowMode GetCascadeShadowMode();
+	int32 GetCascadeShadowsQuantizationUnit();
+	
 	enum class EFogMode
 	{
 		Disabled,
@@ -463,8 +475,9 @@ struct FAVSMSamplePackedData
 
 BEGIN_UNIFORM_BUFFER_STRUCT(FAdaptiveVolumetricShadowMapUniformBufferParameters, )
 	SHADER_PARAMETER_ARRAY(FMatrix44f, TranslatedWorldToShadow, [6])
-	SHADER_PARAMETER(FVector3f, TranslatedWorldOrigin)
-	SHADER_PARAMETER(FVector4f, TranslatedWorldPlane)
+	SHADER_PARAMETER_ARRAY(FVector4f, TranslatedWorldOrigin, [6])
+	SHADER_PARAMETER_ARRAY(FVector4f, TranslatedWorldPlane, [6])
+	SHADER_PARAMETER_ARRAY(FVector4f, SplitDepths, [6])
 
 	SHADER_PARAMETER(FIntPoint, Resolution)
 	SHADER_PARAMETER(int32, NumShadowMatrices)
@@ -483,8 +496,9 @@ END_UNIFORM_BUFFER_STRUCT()
 
 BEGIN_SHADER_PARAMETER_STRUCT(FAdaptiveVolumetricShadowMapParameters, RENDERER_API)
 	SHADER_PARAMETER_ARRAY(FMatrix44f, TranslatedWorldToShadow, [6])
-	SHADER_PARAMETER(FVector3f, TranslatedWorldOrigin)
-	SHADER_PARAMETER(FVector4f, TranslatedWorldPlane)
+	SHADER_PARAMETER_ARRAY(FVector4f, TranslatedWorldOrigin, [6])
+	SHADER_PARAMETER_ARRAY(FVector4f, TranslatedWorldPlane, [6])
+	SHADER_PARAMETER_ARRAY(FVector4f, SplitDepths, [6])
 
 	SHADER_PARAMETER(FIntPoint, Resolution)
 	SHADER_PARAMETER(int32, NumShadowMatrices)
@@ -686,11 +700,25 @@ void ConvertBeerLawShadowMapToVolumetricShadowMap(
 	FRDGBufferRef& VolumetricShadowTransmittanceBuffer
 );
 
+void ConvertVolumetricShadowMapToBeerLawShadowMap(
+	FRDGBuilder& GraphBuilder,
+	const FSceneTextures& SceneTextures,
+	FViewInfo& View,
+	// Input
+	const FString& ShadowMapName,
+	FIntPoint ShadowMapResolution,
+	FIntVector GroupCount,
+	FRDGBufferRef VolumetricShadowLinkedListBuffer,
+	// Output
+	FRDGTextureRef& BeerShadowMapTexture
+);
+
 void CreateAdaptiveVolumetricShadowMapUniformBufferParameters(
 	FRDGBuilder& GraphBuilder,
-	const FVector3f& TranslatedWorldOrigin,
-	const FVector4f& TranslatedWorldPlane,
+	const FVector3f* TranslatedWorldOrigin,
+	const FVector4f* TranslatedWorldPlane,
 	const FMatrix44f* TranslatedWorldToShadow,
+	const FVector4f* SplitDepths,
 	FIntPoint VolumetricShadowMapResolution,
 	float VolumetricShadowMapDownsampleFactor,
 	int32 NumShadowMatrices,
@@ -704,9 +732,10 @@ void CreateAdaptiveVolumetricShadowMapUniformBufferParameters(
 
 void CreateAdaptiveVolumetricShadowMapUniformBuffer(
 	FRDGBuilder& GraphBuilder,
-	const FVector3f& TranslatedWorldOrigin,
-	const FVector4f& TranslatedWorldPlane,
+	const FVector3f* TranslatedWorldOrigin,
+	const FVector4f* TranslatedWorldPlane,
 	const FMatrix44f* TranslatedWorldToShadow,
+	const FVector4f* SplitDepths,
 	FIntPoint VolumetricShadowMapResolution,
 	float VolumetricShadowMapDownsampleFactor,
 	int32 NumShadowMatrices,
@@ -801,6 +830,18 @@ void DilateExistenceMask(
 	FRDGTextureRef& DilatedExistenceTexture
 );
 
+void VisualizeCascades(
+	FRDGBuilder& GraphBuilder,
+	// Scene data
+	const FSceneTextures& SceneTextures,
+	const FScene* Scene,
+	const FViewInfo& View,
+	// Light data
+	const FLightSceneInfo* LightSceneInfo,
+	// Output
+	FRDGTextureRef& VisualizationTexture
+);
+
 void RenderAmbientOcclusionWithLiveShading(
 	FRDGBuilder& GraphBuilder,
 	// Scene data
diff --git a/Engine/Source/Runtime/Renderer/Private/HeterogeneousVolumes/HeterogeneousVolumesLiveShadingPipeline.cpp b/Engine/Source/Runtime/Renderer/Private/HeterogeneousVolumes/HeterogeneousVolumesLiveShadingPipeline.cpp
--- a/Engine/Source/Runtime/Renderer/Private/HeterogeneousVolumes/HeterogeneousVolumesLiveShadingPipeline.cpp
+++ b/Engine/Source/Runtime/Renderer/Private/HeterogeneousVolumes/HeterogeneousVolumesLiveShadingPipeline.cpp
@@ -1000,9 +1000,11 @@ FAdaptiveVolumetricShadowMapParameters GetAdaptiveVolumetricShadowMapParametersF
 		for (int32 i = 0; i < Parameters.NumShadowMatrices; ++i)
 		{
 			Parameters.TranslatedWorldToShadow[i] = UniformBuffer->GetParameters()->TranslatedWorldToShadow[i];
+			Parameters.TranslatedWorldOrigin[i] = UniformBuffer->GetParameters()->TranslatedWorldOrigin[i];
+			Parameters.TranslatedWorldPlane[i] = UniformBuffer->GetParameters()->TranslatedWorldPlane[i];
+			Parameters.SplitDepths[i] = UniformBuffer->GetParameters()->SplitDepths[i];
 		}
-		Parameters.TranslatedWorldOrigin = UniformBuffer->GetParameters()->TranslatedWorldOrigin;
-		Parameters.TranslatedWorldPlane = UniformBuffer->GetParameters()->TranslatedWorldPlane;
+
 		Parameters.Resolution = UniformBuffer->GetParameters()->Resolution;
 		Parameters.MaxSampleCount = UniformBuffer->GetParameters()->MaxSampleCount;
 		Parameters.bIsEmpty = UniformBuffer->GetParameters()->bIsEmpty;
@@ -1864,6 +1866,89 @@ public:
 };
 IMPLEMENT_TYPE_LAYOUT(FRenderShadowMapLooseBindings);
 
+class FClearVolumetricShadowMapCS : public FGlobalShader
+{
+	DECLARE_GLOBAL_SHADER(FClearVolumetricShadowMapCS);
+	SHADER_USE_PARAMETER_STRUCT(FClearVolumetricShadowMapCS, FGlobalShader);
+
+	using FPermutationDomain = TShaderPermutationDomain<>;
+
+	BEGIN_SHADER_PARAMETER_STRUCT(FParameters, )
+		SHADER_PARAMETER(FIntPoint, ShadowResolution)
+		SHADER_PARAMETER(int32, MaxSampleCount)
+
+		// Output
+		SHADER_PARAMETER_RDG_BUFFER_UAV(RWStructuredBuffer<int2>, RWVolumetricShadowLinkedListBuffer)
+	END_SHADER_PARAMETER_STRUCT()
+
+	static bool ShouldCompilePermutation(
+		const FGlobalShaderPermutationParameters& Parameters
+	)
+	{
+		return DoesPlatformSupportHeterogeneousVolumes(Parameters.Platform);
+	}
+
+	static void ModifyCompilationEnvironment(
+		const FGlobalShaderPermutationParameters& Parameters,
+		FShaderCompilerEnvironment& OutEnvironment
+	)
+	{
+		FMaterialShader::ModifyCompilationEnvironment(Parameters, OutEnvironment);
+		OutEnvironment.SetDefine(TEXT("THREADGROUP_SIZE_1D"), GetThreadGroupSize1D());
+		OutEnvironment.SetDefine(TEXT("THREADGROUP_SIZE_2D"), GetThreadGroupSize2D());
+		OutEnvironment.SetDefine(TEXT("THREADGROUP_SIZE_3D"), GetThreadGroupSize3D());
+
+		// Disable in-scattering features
+		//OutEnvironment.SetDefine(TEXT("DIM_USE_TRANSMITTANCE_VOLUME"), 0);
+		//OutEnvironment.SetDefine(TEXT("DIM_USE_INSCATTERING_VOLUME"), 0);
+		//OutEnvironment.SetDefine(TEXT("DIM_USE_LUMEN_GI"), 0);
+
+		// This shader takes a very long time to compile with FXC, so we pre-compile it with DXC first and then forward the optimized HLSL to FXC.
+		//OutEnvironment.CompilerFlags.Add(CFLAG_PrecompileWithDXC); // @lh-todo - Disabled to workaround SPIRV-Cross bug: StructuredBuffer<uint> is translated to ByteAddressBuffer in HLSL backend
+		OutEnvironment.CompilerFlags.Add(CFLAG_AllowTypedUAVLoads);
+
+		//OutEnvironment.SetDefine(TEXT("GET_PRIMITIVE_DATA_OVERRIDE"), 1);
+	}
+
+	static int32 GetThreadGroupSize1D() { return GetThreadGroupSize3D() * GetThreadGroupSize3D() * GetThreadGroupSize3D(); }
+	static int32 GetThreadGroupSize2D() { return 8; }
+	static int32 GetThreadGroupSize3D() { return 4; }
+};
+
+IMPLEMENT_GLOBAL_SHADER(FClearVolumetricShadowMapCS, "/Engine/Private/HeterogeneousVolumes/HeterogeneousVolumesVoxelGridShadows.usf", "ClearVolumetricShadowMapCS", SF_Compute);
+
+void ClearVolumetricShadowMap(
+	FRDGBuilder& GraphBuilder,
+	const FScene* Scene,
+	const FViewInfo& View,
+	// Shadow data
+	FIntPoint ShadowMapResolution,
+	FIntVector& GroupCount,
+	// Output
+	FRDGBufferRef& VolumetricShadowLinkedListBuffer
+)
+{
+	FClearVolumetricShadowMapCS::FParameters* PassParameters = GraphBuilder.AllocParameters<FClearVolumetricShadowMapCS::FParameters>();
+	{
+		PassParameters->ShadowResolution = HeterogeneousVolumes::GetShadowMapResolution();
+		PassParameters->MaxSampleCount = HeterogeneousVolumes::GetShadowMaxSampleCount();
+		// Output
+		PassParameters->RWVolumetricShadowLinkedListBuffer = GraphBuilder.CreateUAV(VolumetricShadowLinkedListBuffer);
+	}
+
+	FClearVolumetricShadowMapCS::FPermutationDomain PermutationVector;
+	TShaderRef<FClearVolumetricShadowMapCS> ComputeShader = View.ShaderMap->GetShader<FClearVolumetricShadowMapCS>(PermutationVector);
+	FComputeShaderUtils::AddPass(
+		GraphBuilder,
+		// Add Light name..
+		RDG_EVENT_NAME("ClearVolumetricShadowMapsCS"),
+		ERDGPassFlags::Compute | ERDGPassFlags::NeverCull,
+		ComputeShader,
+		PassParameters,
+		GroupCount
+	);
+}
+
 class FRenderVolumetricShadowMapForLightWithLiveShadingCS : public FMeshMaterialShader
 {
 	DECLARE_SHADER_TYPE(FRenderVolumetricShadowMapForLightWithLiveShadingCS, MeshMaterial);
@@ -1882,7 +1967,7 @@ class FRenderVolumetricShadowMapForLightWithLiveShadingCS : public FMeshMaterial
 		SHADER_PARAMETER_STRUCT_REF(FBlueNoise, BlueNoise)
 
 		// Volumetric Shadow Map data
-		SHADER_PARAMETER(FVector3f, TranslatedWorldOrigin)
+		SHADER_PARAMETER_ARRAY(FVector4f, TranslatedWorldOrigin, [6])
 		SHADER_PARAMETER(FIntPoint, ShadowResolution)
 		SHADER_PARAMETER(int, MaxSampleCount)
 		SHADER_PARAMETER(float, AbsoluteErrorThreshold)
@@ -1916,7 +2001,6 @@ class FRenderVolumetricShadowMapForLightWithLiveShadingCS : public FMeshMaterial
 		// Output
 		SHADER_PARAMETER_RDG_BUFFER_UAV(RWBuffer<int>, RWVolumetricShadowLinkedListAllocatorBuffer)
 		SHADER_PARAMETER_RDG_BUFFER_UAV(RWStructuredBuffer<int2>, RWVolumetricShadowLinkedListBuffer)
-		SHADER_PARAMETER_RDG_TEXTURE_UAV(RWTexture2D<float4>, RWBeerShadowMapTexture)
 
 		// Debug
 		//SHADER_PARAMETER_RDG_BUFFER_UAV(RWStructuredBuffer<FVolumetricShadowMapDebugData>, RWDebugBuffer)
@@ -2012,9 +2096,10 @@ void CollectHeterogeneousVolumeMeshBatchesForView(
 	}
 }
 
-void CollectHeterogeneousVolumeMeshBatchesForLight(
+void CollectHeterogeneousVolumeMeshBatchesForLightWithShadowIndex(
 	const FLightSceneInfo* LightSceneInfo,
 	const FVisibleLightInfo* VisibleLightInfo,
+	int32 ShadowIndexDiscard,
 	const FViewInfo& View,
 	TSet<FVolumetricMeshBatch>& HeterogeneousVolumesMeshBatches,
 	FBoxSphereBounds::Builder& WorldBoundsBuilder
@@ -2023,16 +2108,19 @@ void CollectHeterogeneousVolumeMeshBatchesForLight(
 	check(LightSceneInfo);
 	check(VisibleLightInfo);
 
+	Experimental::TRobinHoodHashMap<const FPrimitiveSceneProxy*, FVolumetricMeshBatch> MeshBatchHashMap;
+
 	if (LightSceneInfo->Proxy->CastsVolumetricShadow())
 	{
-		//for (int32 ViewIndex = 0; ViewIndex < Views.Num(); ViewIndex++)
+		int32 ShadowIndexMin = ShadowIndexDiscard;
+		int32 ShadowIndexMax = ShadowIndexDiscard + 1;
+		if (ShadowIndexDiscard < 0)
 		{
-			//const FViewInfo& View = Views[ViewIndex];
-			bool bCollectForShadowCasting = true;
-			CollectHeterogeneousVolumeMeshBatchesForView(View, bCollectForShadowCasting, HeterogeneousVolumesMeshBatches, WorldBoundsBuilder);
+			ShadowIndexMin = 0;
+			ShadowIndexMax = VisibleLightInfo->ShadowsToProject.Num();
 		}
 
-		for (int32 ShadowIndex = 0; ShadowIndex < VisibleLightInfo->ShadowsToProject.Num(); ++ShadowIndex)
+		for (int32 ShadowIndex = ShadowIndexMin; ShadowIndex < ShadowIndexMax; ++ShadowIndex)
 		{
 			const FProjectedShadowInfo* ProjectedShadowInfo = HeterogeneousVolumes::GetProjectedShadowInfo(VisibleLightInfo, ShadowIndex);
 			if (ProjectedShadowInfo != nullptr)
@@ -2053,12 +2141,20 @@ void CollectHeterogeneousVolumeMeshBatchesForLight(
 
 					if (bIsShadowCast && bIsVolumeMaterial)
 					{
-						HeterogeneousVolumesMeshBatches.FindOrAdd(FVolumetricMeshBatch(MeshBatch.Mesh, MeshBatch.PrimitiveSceneProxy));
-						WorldBoundsBuilder += MeshBatch.PrimitiveSceneProxy->GetBounds();
+						FBoxSphereBounds VolumeBounds = MeshBatch.PrimitiveSceneProxy->GetBounds();
+						WorldBoundsBuilder += VolumeBounds;
+
+						MeshBatchHashMap.FindOrAdd(MeshBatch.PrimitiveSceneProxy, FVolumetricMeshBatch(MeshBatch.Mesh, MeshBatch.PrimitiveSceneProxy));
 					}
 				}
 			}
 		}
+
+		// TODO: Determine if coarse sorting is beneficial for shadow combining
+		for (auto Itr = MeshBatchHashMap.begin(); Itr != MeshBatchHashMap.end(); ++Itr)
+		{
+			HeterogeneousVolumesMeshBatches.Add((*Itr).Value);
+		}
 	}
 }
 
@@ -2074,12 +2170,13 @@ bool RenderVolumetricShadowMapForLightForHeterogeneousVolumeWithLiveShading(
 	const FVisibleLightInfo* VisibleLightInfo,
 	// Shadow data
 	HeterogeneousVolumes::FLODInfo LODInfo,
-	const FVector3f& TranslatedWorldOrigin,
+	const FVector3f* TranslatedWorldOrigin,
 	int32 NumShadowMatrices,
 	FMatrix44f* TranslatedWorldToShadow,
 	FMatrix44f* ShadowToTranslatedWorld,
 	FIntPoint ShadowMapResolution,
 	uint32 MaxSampleCount,
+	float MaxShadowTraceDistance,
 	bool bUseCameraSceneDepth,
 	// Volume
 	const FVolumetricMeshBatch& VolumetricMeshBatch,
@@ -2117,7 +2214,6 @@ bool RenderVolumetricShadowMapForLightForHeterogeneousVolumeWithLiveShading(
 		PassParameters->BlueNoise = CreateUniformBufferImmediate(BlueNoise, EUniformBufferUsage::UniformBuffer_SingleDraw);
 
 		// Shadow map data
-		PassParameters->TranslatedWorldOrigin = TranslatedWorldOrigin;
 		PassParameters->ShadowResolution = ShadowMapResolution;
 		PassParameters->MaxSampleCount = MaxSampleCount;
 		PassParameters->AbsoluteErrorThreshold = HeterogeneousVolumes::GetShadowAbsoluteErrorThreshold();
@@ -2126,6 +2222,7 @@ bool RenderVolumetricShadowMapForLightForHeterogeneousVolumeWithLiveShading(
 		PassParameters->NumShadowMatrices = NumShadowMatrices;
 		for (int32 i = 0; i < PassParameters->NumShadowMatrices; ++i)
 		{
+			PassParameters->TranslatedWorldOrigin[i] = TranslatedWorldOrigin[i];
 			PassParameters->TranslatedWorldToShadow[i] = TranslatedWorldToShadow[i];
 			PassParameters->ShadowToTranslatedWorld[i] = ShadowToTranslatedWorld[i];
 		}
@@ -2153,8 +2250,12 @@ bool RenderVolumetricShadowMapForLightForHeterogeneousVolumeWithLiveShading(
 		// Ray Data
 		float LODFactor = HeterogeneousVolumes::CalcLODFactor(LODInfo, HeterogeneousVolumeInterface);
 		PassParameters->ShadowStepSize = HeterogeneousVolumes::GetShadowStepSize();
-		PassParameters->ShadowStepFactor = HeterogeneousVolumeInterface->GetShadowStepFactor() * LODFactor;
-		PassParameters->MaxTraceDistance = HeterogeneousVolumes::GetMaxTraceDistance();
+		// TODO: Assign LOD factor per sub-region shadow for CSM
+		//PassParameters->ShadowStepFactor = HeterogeneousVolumeInterface->GetShadowStepFactor() * LODFactor;
+		PassParameters->ShadowStepFactor = HeterogeneousVolumeInterface->GetShadowStepFactor();
+		//PassParameters->MaxTraceDistance = HeterogeneousVolumes::GetMaxTraceDistance();
+		//PassParameters->MaxTraceDistance = HeterogeneousVolumes::GetMaxShadowTraceDistance();
+		PassParameters->MaxTraceDistance = MaxShadowTraceDistance;
 		PassParameters->MaxStepCount = HeterogeneousVolumes::GetMaxStepCount();
 		PassParameters->bJitter = HeterogeneousVolumes::ShouldJitter();
 
@@ -2167,7 +2268,6 @@ bool RenderVolumetricShadowMapForLightForHeterogeneousVolumeWithLiveShading(
 		// Output
 		PassParameters->RWVolumetricShadowLinkedListAllocatorBuffer = GraphBuilder.CreateUAV(VolumetricShadowLinkedListAllocatorBuffer, PF_R32_UINT);
 		PassParameters->RWVolumetricShadowLinkedListBuffer = GraphBuilder.CreateUAV(VolumetricShadowLinkedListBuffer);
-		PassParameters->RWBeerShadowMapTexture = GraphBuilder.CreateUAV(BeerShadowMapTexture);
 		//PassParameters->RWDebugBuffer = GraphBuilder.CreateUAV(DebugBuffer);
 	}
 
@@ -2221,6 +2321,47 @@ bool RenderVolumetricShadowMapForLightForHeterogeneousVolumeWithLiveShading(
 	return true;
 }
 
+FVector FloorToDouble(const FVector Vec)
+{
+	return FVector(
+		FMath::FloorToDouble(Vec.X),
+		FMath::FloorToDouble(Vec.Y),
+		FMath::FloorToDouble(Vec.Z)
+	);
+}
+
+TArray<FVector> GetFrustumVertices(const FMatrix& ClipToWorld)
+{
+	TArray<FVector> Vertices;
+
+	Vertices.SetNum(8);
+	for (int32 i = 0; i < Vertices.Num(); ++i)
+	{
+		float x = i & 1 ? 1.0 : -1.0;
+		float y = (i >> 1) & 1 ? 1.0 : -1.0;
+		float z = (i >> 2) & 1 ? 1.0 : -1.0;
+
+		//Vertices[i] = ClipToWorld.TransformPosition(FVector(x, y, z));
+		FVector4 Result = ClipToWorld.TransformFVector4(FVector4(x, y, z, 1.0));
+		Vertices[i] = Result / Result.W;
+	}
+
+	return Vertices;
+}
+
+TArray<FVector> TransformVertices(const TArray<FVector>& Vertices, const FMatrix& ToMatrix)
+{
+	TArray<FVector> ToVertices;
+
+	ToVertices.SetNum(Vertices.Num());
+	for (int i = 0; i < ToVertices.Num(); ++i)
+	{
+		ToVertices[i] = ToMatrix.TransformPosition(Vertices[i]);
+	}
+
+	return ToVertices;
+}
+
 bool RenderVolumetricShadowMapForLightWithLiveShading(
 	FRDGBuilder& GraphBuilder,
 	// Scene data
@@ -2233,9 +2374,10 @@ bool RenderVolumetricShadowMapForLightWithLiveShading(
 	const FVisibleLightInfo* VisibleLightInfo,
 	// Output
 	bool& bIsDirectionalLight,
-	FVector3f& TranslatedWorldOrigin,
-	FVector4f& TranslatedWorldPlane,
+	FVector3f* TranslatedWorldOrigin,
+	FVector4f* TranslatedWorldPlane,
 	FMatrix44f* TranslatedWorldToShadow,
+	FVector4f* SplitDepths,
 	FIntVector& GroupCount,
 	int32& NumShadowMatrices,
 	FIntPoint& ShadowMapResolution,
@@ -2249,8 +2391,18 @@ bool RenderVolumetricShadowMapForLightWithLiveShading(
 	check(LightSceneInfo);
 	check(VisibleLightInfo);
 
-	const FProjectedShadowInfo* ProjectedShadowInfo = HeterogeneousVolumes::GetProjectedShadowInfo(VisibleLightInfo, 0);
-	check(ProjectedShadowInfo != NULL)
+	check(HeterogeneousVolumes::GetProjectedShadowInfo(VisibleLightInfo, 0) != NULL)
+
+	// Collect all shadow-casting volumes
+	int32 GlobalCascadeShadowIndex = -1;
+	TSet<FVolumetricMeshBatch> GlobalHeterogeneousVolumesMeshBatches;
+	FBoxSphereBounds::Builder GlobalWorldVolumeBoundsBuilder;
+	CollectHeterogeneousVolumeMeshBatchesForLightWithShadowIndex(LightSceneInfo, VisibleLightInfo, GlobalCascadeShadowIndex, View, GlobalHeterogeneousVolumesMeshBatches, GlobalWorldVolumeBoundsBuilder);
+	if (GlobalHeterogeneousVolumesMeshBatches.IsEmpty() || !GlobalWorldVolumeBoundsBuilder.IsValid())
+	{
+		return false;
+	}
+	FBoxSphereBounds GlobalWorldVolumeBounds(GlobalWorldVolumeBoundsBuilder);
 
 	ShadowMapResolution = HeterogeneousVolumes::GetShadowMapResolution();
 
@@ -2260,113 +2412,371 @@ bool RenderVolumetricShadowMapForLightWithLiveShading(
 	GroupCount.Y = FMath::DivideAndRoundUp(ShadowMapResolution.Y, FRenderVolumetricShadowMapForLightWithLiveShadingCS::GetThreadGroupSize2D());
 	GroupCount.Z = bIsMultiProjection ? 6 : 1;
 
-	// Collect shadow-casting volumes
-	TSet<FVolumetricMeshBatch> HeterogeneousVolumesMeshBatches;
-	FBoxSphereBounds::Builder WorldVolumeBoundsBuilder;
-	CollectHeterogeneousVolumeMeshBatchesForLight(LightSceneInfo, VisibleLightInfo, View, HeterogeneousVolumesMeshBatches, WorldVolumeBoundsBuilder);
-	if (HeterogeneousVolumesMeshBatches.IsEmpty() || !WorldVolumeBoundsBuilder.IsValid())
-	{
-		return false;
-	}
+	float MaxShadowTraceDistance = HeterogeneousVolumes::GetMaxShadowTraceDistance();
 
-	// Build shadow transform
-	FBoxSphereBounds WorldVolumeBounds(WorldVolumeBoundsBuilder);
-	NumShadowMatrices = ProjectedShadowInfo->OnePassShadowViewProjectionMatrices.Num();
 	FMatrix44f ShadowToTranslatedWorld[6];
-
-	if (NumShadowMatrices > 0)
+	if (LightType == LightType_Directional)
 	{
-		FVector PreViewTranslation = View.ViewMatrices.GetPreViewTranslation();
-		FMatrix TranslatedWorldToWorldMatrix = FTranslationMatrix(-PreViewTranslation);
-		FVector LightPosition = LightSceneInfo->Proxy->GetPosition();
-		FMatrix WorldToLightMatrix = LightSceneInfo->Proxy->GetWorldToLight();
+		bIsDirectionalLight = true;
 
-		// Remove light rotation when building the RectLight projections..
-		FMatrix RotationalAdjustmentMatrix = FMatrix::Identity;
-		if (LightType == LIGHT_TYPE_RECT)
+		if (HeterogeneousVolumes::SupportsCascadeShadowsForDirectionalLight())
 		{
-			FVector LightDirection = LightSceneInfo->Proxy->GetDirection().GetSafeNormal();
-			RotationalAdjustmentMatrix = FRotationMatrix(LightDirection.Rotation());
-		}
+			NumShadowMatrices = VisibleLightInfo->ShadowsToProject.Num();
+			bIsMultiProjection = true;
+			GroupCount.Z = NumShadowMatrices;
 
-		FMatrix ViewMatrix[] = {
-			FLookFromMatrix(FVector::Zero(), FVector(-1, 0, 0), FVector(0, 0, 1)),
-			FLookFromMatrix(FVector::Zero(), FVector(1, 0, 0), FVector(0, 0, 1)),
-			FLookFromMatrix(FVector::Zero(), FVector(0, -1, 0), FVector(0, 0, 1)),
-			FLookFromMatrix(FVector::Zero(), FVector(0, 1, 0), FVector(0, 0, 1)),
-			FLookFromMatrix(FVector::Zero(), FVector(0, 0, -1), FVector(1, 0, 0)),
-			FLookFromMatrix(FVector::Zero(), FVector(0, 0, 1), FVector(1, 0, 0))
-		};
+			// Auto-fit may cause quality popping as new volumes are spawned.
+			float MaxTraceDistance = HeterogeneousVolumes::GetMaxShadowTraceDistance();
+			if (MaxTraceDistance < 0)
+			{
+				MaxTraceDistance = (GlobalWorldVolumeBounds.Origin - View.ViewMatrices.GetViewOrigin()).Length() + GlobalWorldVolumeBounds.SphereRadius;
+			}
 
-		FMatrix PerspectiveMatrix = FPerspectiveMatrix(
-			PI / 4.0f,
-			ShadowMapResolution.X,
-			ShadowMapResolution.Y,
-			HeterogeneousVolumes::GetShadowNearClippingDistance(),
-			LightSceneInfo->Proxy->GetRadius()
-		);
+			for (int32 ShadowIndex = 0; ShadowIndex < NumShadowMatrices; ++ShadowIndex)
+			{
+				// Build orthographic projection centered around volume..
+				FVector PreViewTranslation = View.ViewMatrices.GetPreViewTranslation();
+				FMatrix TranslatedWorldToWorldMatrix = FTranslationMatrix(-PreViewTranslation);
+				FVector LightDirection = LightSceneInfo->Proxy->GetDirection().GetSafeNormal();
 
-		FMatrix ScreenMatrix = FScaleMatrix(FVector(0.5, -0.5, -0.5)) * FTranslationMatrix(FVector(0.5, 0.5, 0.5));
+				FMatrix TranslationMatrix;
+				FMatrix RotationMatrix;
+				FMatrix ScaleMatrix;
+
+				if (HeterogeneousVolumes::GetCascadeShadowMode() == HeterogeneousVolumes::ECascadeShadowMode::Frustums)
+				{
+					// Define viewing frustum from SplitNear to Split Far
+					const FMatrix ViewMatrix = View.ViewMatrices.GetViewMatrix();
+					const FMatrix ProjectionMatrix = View.ViewMatrices.GetProjectionMatrix();
+					const FVector4 ViewOrigin = View.ViewMatrices.GetViewOrigin();
+
+					float AspectRatio = ProjectionMatrix.M[1][1] / ProjectionMatrix.M[0][0];
+					float ActualFOV = (ViewOrigin.W > 0.0f) ? FMath::Atan(1.0f / ProjectionMatrix.M[0][0]) : PI / 4.0f;
+					const FProjectedShadowInfo* ProjectedShadowInfo = HeterogeneousVolumes::GetProjectedShadowInfo(VisibleLightInfo, ShadowIndex);
+					float Near = ProjectedShadowInfo->CascadeSettings.SplitNear;
+					float Far = ProjectedShadowInfo->CascadeSettings.SplitFar;
+
+					// Get frustum vertices in world-space
+					FMatrix CascadeViewProjectionMatrix = ViewMatrix * FPerspectiveMatrix(ActualFOV, AspectRatio, 1.0f, Near, Far);
+					FMatrix CascadeToWorld = CascadeViewProjectionMatrix.Inverse();
+					TArray<FVector> WorldFrustumVertices = GetFrustumVertices(CascadeToWorld);
+
+					// Build AABB about the frustum in world-space
+					FBoxSphereBounds::Builder WorldSpaceBoundsBuilder;
+					for (int32 i = 0; i < WorldFrustumVertices.Num(); ++i)
+					{
+						WorldSpaceBoundsBuilder += WorldFrustumVertices[i];
+					}
+					FBoxSphereBounds WorldSpaceBounds(WorldSpaceBoundsBuilder);
+
+					// Define the light projection
+					float MaxRadius = FMath::Max(WorldSpaceBounds.SphereRadius, MaxShadowTraceDistance / 2.0);
+					FVector ToLight = WorldSpaceBounds.Origin - LightDirection * MaxRadius;
+					TranslatedWorldOrigin[ShadowIndex] = FVector3f(PreViewTranslation + ToLight);
+					MaxShadowTraceDistance = MaxRadius * 2.0;
+
+					TranslationMatrix = FTranslationMatrix(-ToLight);
+					RotationMatrix = FInverseRotationMatrix(LightDirection.Rotation());
+					FMatrix WorldToLight = TranslationMatrix * RotationMatrix;
+					FMatrix LightToWorld = WorldToLight.Inverse();
+
+					// Build a bounding volume for light rays in light-space
+					TArray<FVector> LightFrustumVertices = TransformVertices(WorldFrustumVertices, WorldToLight);
+
+					FBoxSphereBounds::Builder LightSpaceBoundsBuilder;
+					for (int32 i = 0; i < LightFrustumVertices.Num(); ++i)
+					{
+						LightSpaceBoundsBuilder += LightFrustumVertices[i];
+					}
+					// The virtual light position bounds the parallelpiped
+					LightSpaceBoundsBuilder += FVector::Zero();
+					FBoxSphereBounds LightSpaceBounds(LightSpaceBoundsBuilder);
+
+					// Transform the centroid to world-space and use this as the new translation matrix
+					WorldSpaceBounds.Origin = LightToWorld.TransformPosition(LightSpaceBounds.Origin);
+
+					int32 QuantizeDistance = HeterogeneousVolumes::GetCascadeShadowsQuantizationUnit();
+					if (QuantizeDistance > 0)
+					{
+						float PixelDistance = FMath::Max(LightSpaceBounds.BoxExtent.Y, LightSpaceBounds.BoxExtent.Z) / ShadowMapResolution.X;
+						float QuantizePixelDistance = QuantizeDistance * PixelDistance;
+						FVector OriginSnapUnits = FVector(
+							FMath::FloorToDouble(WorldSpaceBounds.Origin.X / QuantizePixelDistance),
+							FMath::FloorToDouble(WorldSpaceBounds.Origin.Y / QuantizePixelDistance),
+							FMath::FloorToDouble(WorldSpaceBounds.Origin.Z / QuantizePixelDistance)
+						);
+						FVector QuantizedWorldSpaceBoundsOrigin = OriginSnapUnits * QuantizePixelDistance;
+						FVector OriginDelta = QuantizedWorldSpaceBoundsOrigin - WorldSpaceBounds.Origin;
+						WorldSpaceBounds.Origin = QuantizedWorldSpaceBoundsOrigin;
+
+						FVector QuantizedLocalSpaceBoundsBoxExtent = LightSpaceBounds.BoxExtent;
+						FVector BoxExtentSnapUnits = FVector(
+							FMath::FloorToDouble(QuantizedLocalSpaceBoundsBoxExtent.X / QuantizePixelDistance + 1.0),
+							FMath::FloorToDouble(QuantizedLocalSpaceBoundsBoxExtent.Y / QuantizePixelDistance + 1.0),
+							FMath::FloorToDouble(QuantizedLocalSpaceBoundsBoxExtent.Z / QuantizePixelDistance + 1.0)
+						);
+						BoxExtentSnapUnits += FVector(1);
+						QuantizedLocalSpaceBoundsBoxExtent = BoxExtentSnapUnits * QuantizePixelDistance;
+						LightSpaceBounds.BoxExtent.Y = QuantizedLocalSpaceBoundsBoxExtent.Y;
+						LightSpaceBounds.BoxExtent.Z = QuantizedLocalSpaceBoundsBoxExtent.Z;
+
+						// Recompute light virtual origin
+						ToLight += OriginDelta;
+						TranslatedWorldOrigin[ShadowIndex] = FVector3f(PreViewTranslation + ToLight);
+					}
+
+					TranslationMatrix = FTranslationMatrix(-WorldSpaceBounds.Origin);
+					ScaleMatrix = FScaleMatrix(FVector(1.0) / LightSpaceBounds.BoxExtent);
+				}
+				else if (HeterogeneousVolumes::GetCascadeShadowMode() == HeterogeneousVolumes::ECascadeShadowMode::Clipmaps)
+				{
+					const FProjectedShadowInfo* ProjectedShadowInfo = HeterogeneousVolumes::GetProjectedShadowInfo(VisibleLightInfo, ShadowIndex);
+					float Far = ProjectedShadowInfo->CascadeSettings.SplitFar;
+
+					GlobalWorldVolumeBounds.Origin = View.ViewMatrices.GetViewOrigin();
+					GlobalWorldVolumeBounds.BoxExtent = FVector(Far, Far, Far);
+					GlobalWorldVolumeBounds.SphereRadius = Far; // GlobalWorldVolumeBounds.BoxExtent.Length();
+
+					// Build AABB about the camera in world-space
+					FBoxSphereBounds::Builder WorldSpaceBoundsBuilder;
+					TArray<FVector> GlobalWorldVolumeBoundsVertices;
+					GlobalWorldVolumeBoundsVertices.SetNum(8);
+					for (int32 i = 0; i < GlobalWorldVolumeBoundsVertices.Num(); ++i)
+					{
+						FVector Sign((i >> 2) & 1 ? 1 : -1, (i >> 1) & 1 ? 1 : -1, i & 1 ? 1 : -1);
+						GlobalWorldVolumeBoundsVertices[i] = GlobalWorldVolumeBounds.Origin + Sign * GlobalWorldVolumeBounds.BoxExtent;
+						WorldSpaceBoundsBuilder += GlobalWorldVolumeBoundsVertices[i];
+					}
+					FBoxSphereBounds WorldSpaceBounds(WorldSpaceBoundsBuilder);
+
+					// Define the light projection
+					float MaxRadius = FMath::Max(WorldSpaceBounds.SphereRadius, MaxShadowTraceDistance / 2.0);
+					FVector ToLight = WorldSpaceBounds.Origin - LightDirection * MaxRadius;
+					TranslatedWorldOrigin[ShadowIndex] = FVector3f(PreViewTranslation + ToLight);
+					MaxShadowTraceDistance = MaxRadius * 2.0;
+
+					TranslationMatrix = FTranslationMatrix(-ToLight);
+					RotationMatrix = FInverseRotationMatrix(LightDirection.Rotation());
+					FMatrix WorldToLight = TranslationMatrix * RotationMatrix;
+					FMatrix LightToWorld = WorldToLight.Inverse();
 
-		for (int32 i = 0; i < NumShadowMatrices; ++i)
+					// Build a bounding volume for light rays in light-space
+					TArray<FVector> LightVolumeBoundsVertices = TransformVertices(GlobalWorldVolumeBoundsVertices, WorldToLight);
+
+					FBoxSphereBounds::Builder LightSpaceBoundsBuilder;
+					for (int32 i = 0; i < LightVolumeBoundsVertices.Num(); ++i)
+					{
+						LightSpaceBoundsBuilder += LightVolumeBoundsVertices[i];
+					}
+					// The virtual light position bounds the parallelpiped
+					LightSpaceBoundsBuilder += FVector::Zero();
+					FBoxSphereBounds LightSpaceBounds(LightSpaceBoundsBuilder);
+
+					// Transform the centroid to world-space and use this as the new translation matrix
+					WorldSpaceBounds.Origin = LightToWorld.TransformPosition(LightSpaceBounds.Origin);
+
+					int32 QuantizeDistance = HeterogeneousVolumes::GetCascadeShadowsQuantizationUnit();
+					if (QuantizeDistance > 0)
+					{
+						float PixelDistance = FMath::Max(LightSpaceBounds.BoxExtent.Y, LightSpaceBounds.BoxExtent.Z) / ShadowMapResolution.X;
+						float QuantizePixelDistance = QuantizeDistance * PixelDistance;
+						FVector OriginSnapUnits = FVector(
+							FMath::FloorToDouble(WorldSpaceBounds.Origin.X / QuantizePixelDistance),
+							FMath::FloorToDouble(WorldSpaceBounds.Origin.Y / QuantizePixelDistance),
+							FMath::FloorToDouble(WorldSpaceBounds.Origin.Z / QuantizePixelDistance)
+						);
+						FVector QuantizedWorldSpaceBoundsOrigin = OriginSnapUnits * QuantizePixelDistance;
+						FVector OriginDelta = QuantizedWorldSpaceBoundsOrigin - WorldSpaceBounds.Origin;
+						WorldSpaceBounds.Origin = QuantizedWorldSpaceBoundsOrigin;
+
+						FVector QuantizedLocalSpaceBoundsBoxExtent = LightSpaceBounds.BoxExtent;
+						FVector BoxExtentSnapUnits = FVector(
+							FMath::FloorToDouble(QuantizedLocalSpaceBoundsBoxExtent.X / QuantizePixelDistance + 1.0),
+							FMath::FloorToDouble(QuantizedLocalSpaceBoundsBoxExtent.Y / QuantizePixelDistance + 1.0),
+							FMath::FloorToDouble(QuantizedLocalSpaceBoundsBoxExtent.Z / QuantizePixelDistance + 1.0)
+						);
+						BoxExtentSnapUnits += FVector(1);
+						QuantizedLocalSpaceBoundsBoxExtent = BoxExtentSnapUnits * QuantizePixelDistance;
+						LightSpaceBounds.BoxExtent.Y = QuantizedLocalSpaceBoundsBoxExtent.Y;
+						LightSpaceBounds.BoxExtent.Z = QuantizedLocalSpaceBoundsBoxExtent.Z;
+
+						// Recompute light virtual origin
+						ToLight += OriginDelta;
+						TranslatedWorldOrigin[ShadowIndex] = FVector3f(PreViewTranslation + ToLight);
+					}
+
+					TranslationMatrix = FTranslationMatrix(-WorldSpaceBounds.Origin);
+					ScaleMatrix = FScaleMatrix(FVector(1.0) / LightSpaceBounds.BoxExtent);
+				}
+				else if (HeterogeneousVolumes::GetCascadeShadowMode() == HeterogeneousVolumes::ECascadeShadowMode::Autofit)
+				{
+					TSet<FVolumetricMeshBatch> HeterogeneousVolumesMeshBatches;
+					FBoxSphereBounds::Builder WorldVolumeBoundsBuilder;
+					CollectHeterogeneousVolumeMeshBatchesForLightWithShadowIndex(LightSceneInfo, VisibleLightInfo, ShadowIndex, View, HeterogeneousVolumesMeshBatches, WorldVolumeBoundsBuilder);
+					FBoxSphereBounds WorldVolumeBounds(WorldVolumeBoundsBuilder);
+
+					// Define transforms
+					TranslationMatrix = FTranslationMatrix(-WorldVolumeBounds.Origin);
+					RotationMatrix = FInverseRotationMatrix(LightDirection.Rotation());
+					ScaleMatrix = FScaleMatrix(FVector(1.0 / WorldVolumeBounds.SphereRadius));
+
+					TranslatedWorldOrigin[ShadowIndex] = FVector3f(PreViewTranslation + WorldVolumeBounds.Origin - LightDirection * WorldVolumeBounds.SphereRadius);
+				}
+				else
+				{
+					const FProjectedShadowInfo* ProjectedShadowInfo = HeterogeneousVolumes::GetProjectedShadowInfo(VisibleLightInfo, ShadowIndex);
+					FShadowCascadeSettings ShadowCascadeSettings;
+					bool IsPrecomputedLightingValid = false;
+					LightSceneInfo->Proxy->GetShadowSplitBounds(View, ProjectedShadowInfo->CascadeSettings.ShadowSplitIndex, IsPrecomputedLightingValid, &ShadowCascadeSettings);
+
+					float DepthSpan = ShadowCascadeSettings.SplitFar - ShadowCascadeSettings.SplitNear;
+					float QuantizeDepthSpan = FMath::FloorToDouble(DepthSpan);
+					float PixelDistance = QuantizeDepthSpan / ShadowMapResolution.GetMax();
+
+					// Define transforms
+					TranslationMatrix = FTranslationMatrix(-GlobalWorldVolumeBounds.Origin);
+					RotationMatrix = FInverseRotationMatrix(LightDirection.Rotation());
+					ScaleMatrix = FScaleMatrix(FVector(1.0 / GlobalWorldVolumeBounds.SphereRadius));
+
+					TranslatedWorldOrigin[ShadowIndex] = FVector3f(PreViewTranslation + GlobalWorldVolumeBounds.Origin - LightDirection * GlobalWorldVolumeBounds.SphereRadius);
+				}
+
+				FMatrix FaceMatrix(
+					FPlane(0, 0, 1, 0),
+					FPlane(0, 1, 0, 0),
+					FPlane(-1, 0, 0, 0),
+					FPlane(0, 0, 0, 1));
+
+				// Invert Z to match reverse-Z for the rest of the shadow types!
+				FMatrix ScreenMatrix = FScaleMatrix(FVector(0.5, -0.5, -0.5)) * FTranslationMatrix(FVector(0.5, 0.5, 0.5));
+				FMatrix WorldToShadowMatrix = TranslationMatrix * RotationMatrix * ScaleMatrix * FaceMatrix * ScreenMatrix;
+				FMatrix TranslatedWorldToShadowMatrix = TranslatedWorldToWorldMatrix * WorldToShadowMatrix;
+
+				TranslatedWorldToShadow[ShadowIndex] = FMatrix44f(TranslatedWorldToShadowMatrix);
+				ShadowToTranslatedWorld[ShadowIndex] = TranslatedWorldToShadow[ShadowIndex].Inverse();
+
+				float W = -FVector3f::DotProduct(TranslatedWorldOrigin[ShadowIndex], FVector3f(LightDirection));
+				TranslatedWorldPlane[ShadowIndex] = FVector4f(LightDirection.X, LightDirection.Y, LightDirection.Z, W);
+
+				const FProjectedShadowInfo* ProjectedShadowInfo = HeterogeneousVolumes::GetProjectedShadowInfo(VisibleLightInfo, ShadowIndex);
+				SplitDepths[ShadowIndex] = FVector4f(
+					ProjectedShadowInfo->CascadeSettings.SplitNear,
+					ProjectedShadowInfo->CascadeSettings.SplitFar,
+					ProjectedShadowInfo->CascadeSettings.SplitNearFadeRegion,
+					ProjectedShadowInfo->CascadeSettings.SplitFarFadeRegion
+				);
+			}
+		}
+		else
 		{
-			FMatrix WorldToShadowMatrix = WorldToLightMatrix * RotationalAdjustmentMatrix * ViewMatrix[i] * PerspectiveMatrix * ScreenMatrix;
-			TranslatedWorldToShadow[i] = FMatrix44f(TranslatedWorldToWorldMatrix * WorldToShadowMatrix);
-			ShadowToTranslatedWorld[i] = TranslatedWorldToShadow[i].Inverse();
+			// Build orthographic projection centered around volume..
+			FVector PreViewTranslation = View.ViewMatrices.GetPreViewTranslation();
+			FMatrix TranslatedWorldToWorldMatrix = FTranslationMatrix(-PreViewTranslation);
+
+			FMatrix TranslationMatrix = FTranslationMatrix(-GlobalWorldVolumeBounds.Origin);
+
+			FVector LightDirection = LightSceneInfo->Proxy->GetDirection().GetSafeNormal();
+			FMatrix RotationMatrix = FInverseRotationMatrix(LightDirection.Rotation());
+			FMatrix ScaleMatrix = FScaleMatrix(FVector(1.0 / GlobalWorldVolumeBounds.SphereRadius));
+
+			const FMatrix FaceMatrix(
+				FPlane(0, 0, 1, 0),
+				FPlane(0, 1, 0, 0),
+				FPlane(-1, 0, 0, 0),
+				FPlane(0, 0, 0, 1));
+
+			// Invert Z to match reverse-Z for the rest of the shadow types!
+			FMatrix ScreenMatrix = FScaleMatrix(FVector(0.5, -0.5, -0.5)) * FTranslationMatrix(FVector(0.5, 0.5, 0.5));
+			FMatrix WorldToShadowMatrix = TranslationMatrix * RotationMatrix * ScaleMatrix * FaceMatrix * ScreenMatrix;
+			FMatrix TranslatedWorldToShadowMatrix = TranslatedWorldToWorldMatrix * WorldToShadowMatrix;
+
+			NumShadowMatrices = 1;
+			TranslatedWorldToShadow[0] = FMatrix44f(TranslatedWorldToShadowMatrix);
+			ShadowToTranslatedWorld[0] = TranslatedWorldToShadow[0].Inverse();
+			TranslatedWorldOrigin[0] = FVector3f(PreViewTranslation + GlobalWorldVolumeBounds.Origin - LightDirection * GlobalWorldVolumeBounds.SphereRadius);
+
+			//FVector LightDirection = LightSceneInfo->Proxy->GetDirection().GetSafeNormal();
+			float W = -FVector3f::DotProduct(TranslatedWorldOrigin[0], FVector3f(LightDirection));
+			TranslatedWorldPlane[0] = FVector4f(LightDirection.X, LightDirection.Y, LightDirection.Z, W);
 		}
-		TranslatedWorldOrigin = FVector3f(PreViewTranslation + LightPosition);
 	}
-	else if (LightType == LightType_Directional)
+	else
 	{
-		bIsDirectionalLight = true;
-		// Build orthographic projection centered around volume..
-		FVector PreViewTranslation = View.ViewMatrices.GetPreViewTranslation();
-		FMatrix TranslatedWorldToWorldMatrix = FTranslationMatrix(-PreViewTranslation);
 
-		FMatrix TranslationMatrix = FTranslationMatrix(-WorldVolumeBounds.Origin);
+		// Collect shadow-casting volumes
+		//const FProjectedShadowInfo* ProjectedShadowInfo = HeterogeneousVolumes::GetProjectedShadowInfo(VisibleLightInfo, ShadowIndex);
+		const FProjectedShadowInfo* ProjectedShadowInfo = HeterogeneousVolumes::GetProjectedShadowInfo(VisibleLightInfo, 0);
+		NumShadowMatrices = ProjectedShadowInfo->OnePassShadowViewProjectionMatrices.Num();
+		if (NumShadowMatrices > 0)
+		{
+			FVector PreViewTranslation = View.ViewMatrices.GetPreViewTranslation();
+			FMatrix TranslatedWorldToWorldMatrix = FTranslationMatrix(-PreViewTranslation);
+			FVector LightPosition = LightSceneInfo->Proxy->GetPosition();
+			FMatrix WorldToLightMatrix = LightSceneInfo->Proxy->GetWorldToLight();
+
+			// Remove light rotation when building the RectLight projections..
+			FMatrix RotationalAdjustmentMatrix = FMatrix::Identity;
+			if (LightType == LIGHT_TYPE_RECT)
+			{
+				FVector LightDirection = LightSceneInfo->Proxy->GetDirection().GetSafeNormal();
+				RotationalAdjustmentMatrix = FRotationMatrix(LightDirection.Rotation());
+			}
 
-		FVector LightDirection = LightSceneInfo->Proxy->GetDirection().GetSafeNormal();
-		FMatrix RotationMatrix = FInverseRotationMatrix(LightDirection.Rotation());
-		FMatrix ScaleMatrix = FScaleMatrix(FVector(1.0 / WorldVolumeBounds.SphereRadius));
+			FMatrix ViewMatrix[] = {
+				FLookFromMatrix(FVector::Zero(), FVector(-1, 0, 0), FVector(0, 0, 1)),
+				FLookFromMatrix(FVector::Zero(), FVector(1, 0, 0), FVector(0, 0, 1)),
+				FLookFromMatrix(FVector::Zero(), FVector(0, -1, 0), FVector(0, 0, 1)),
+				FLookFromMatrix(FVector::Zero(), FVector(0, 1, 0), FVector(0, 0, 1)),
+				FLookFromMatrix(FVector::Zero(), FVector(0, 0, -1), FVector(1, 0, 0)),
+				FLookFromMatrix(FVector::Zero(), FVector(0, 0, 1), FVector(1, 0, 0))
+			};
+
+			FMatrix PerspectiveMatrix = FPerspectiveMatrix(
+				PI / 4.0f,
+				ShadowMapResolution.X,
+				ShadowMapResolution.Y,
+				HeterogeneousVolumes::GetShadowNearClippingDistance(),
+				LightSceneInfo->Proxy->GetRadius()
+			);
 
-		const FMatrix FaceMatrix(
-			FPlane(0, 0, 1, 0),
-			FPlane(0, 1, 0, 0),
-			FPlane(-1, 0, 0, 0),
-			FPlane(0, 0, 0, 1));
+			FMatrix ScreenMatrix = FScaleMatrix(FVector(0.5, -0.5, -0.5)) * FTranslationMatrix(FVector(0.5, 0.5, 0.5));
 
-		// Invert Z to match reverse-Z for the rest of the shadow types!
-		FMatrix ScreenMatrix = FScaleMatrix(FVector(0.5, -0.5, -0.5)) * FTranslationMatrix(FVector(0.5, 0.5, 0.5));
-		FMatrix WorldToShadowMatrix = TranslationMatrix * RotationMatrix * ScaleMatrix * FaceMatrix * ScreenMatrix;
-		FMatrix TranslatedWorldToShadowMatrix = TranslatedWorldToWorldMatrix * WorldToShadowMatrix;
+			for (int32 i = 0; i < NumShadowMatrices; ++i)
+			{
+				FMatrix WorldToShadowMatrix = WorldToLightMatrix * RotationalAdjustmentMatrix * ViewMatrix[i] * PerspectiveMatrix * ScreenMatrix;
+				TranslatedWorldToShadow[i] = FMatrix44f(TranslatedWorldToWorldMatrix * WorldToShadowMatrix);
+				ShadowToTranslatedWorld[i] = TranslatedWorldToShadow[i].Inverse();
+				TranslatedWorldOrigin[i] = FVector3f(PreViewTranslation + LightPosition);
+				FVector LightDirection = LightSceneInfo->Proxy->GetDirection().GetSafeNormal();
+				float W = -FVector3f::DotProduct(TranslatedWorldOrigin[i], FVector3f(LightDirection));
+				TranslatedWorldPlane[i] = FVector4f(LightDirection.X, LightDirection.Y, LightDirection.Z, W);
+			}
+		}
+		else
+		{
+			FVector PreViewTranslation = View.ViewMatrices.GetPreViewTranslation();
+			FMatrix TranslatedWorldToWorldMatrix = FTranslationMatrix(-PreViewTranslation);
+			FVector4f ShadowmapMinMax = FVector4f::Zero();
+			FMatrix WorldToShadowMatrix = ProjectedShadowInfo->GetWorldToShadowMatrix(ShadowmapMinMax);
+			FMatrix TranslatedWorldToShadowMatrix = TranslatedWorldToWorldMatrix * WorldToShadowMatrix;
 
-		NumShadowMatrices = 1;
-		TranslatedWorldToShadow[0] = FMatrix44f(TranslatedWorldToShadowMatrix);
-		ShadowToTranslatedWorld[0] = TranslatedWorldToShadow[0].Inverse();
-		TranslatedWorldOrigin = FVector3f(PreViewTranslation + WorldVolumeBounds.Origin - LightDirection * WorldVolumeBounds.SphereRadius);
-	}
-	else
-	{
-		FVector PreViewTranslation = View.ViewMatrices.GetPreViewTranslation();
-		FMatrix TranslatedWorldToWorldMatrix = FTranslationMatrix(-PreViewTranslation);
-		FVector4f ShadowmapMinMax = FVector4f::Zero();
-		FMatrix WorldToShadowMatrix = ProjectedShadowInfo->GetWorldToShadowMatrix(ShadowmapMinMax);
-		FMatrix TranslatedWorldToShadowMatrix = TranslatedWorldToWorldMatrix * WorldToShadowMatrix;
+			NumShadowMatrices = 1;
+			TranslatedWorldToShadow[0] = FMatrix44f(TranslatedWorldToShadowMatrix);
+			ShadowToTranslatedWorld[0] = TranslatedWorldToShadow[0].Inverse();
+			TranslatedWorldOrigin[0] = FVector3f(View.ViewMatrices.GetPreViewTranslation() - ProjectedShadowInfo->PreShadowTranslation);
 
-		NumShadowMatrices = 1;
-		TranslatedWorldToShadow[0] = FMatrix44f(TranslatedWorldToShadowMatrix);
-		ShadowToTranslatedWorld[0] = TranslatedWorldToShadow[0].Inverse();
-		TranslatedWorldOrigin = FVector3f(View.ViewMatrices.GetPreViewTranslation() - ProjectedShadowInfo->PreShadowTranslation);
+			FVector LightDirection = LightSceneInfo->Proxy->GetDirection().GetSafeNormal();
+			float W = -FVector3f::DotProduct(TranslatedWorldOrigin[0], FVector3f(LightDirection));
+			TranslatedWorldPlane[0] = FVector4f(LightDirection.X, LightDirection.Y, LightDirection.Z, W);
+		}
 	}
 
-	FVector LightDirection = LightSceneInfo->Proxy->GetDirection().GetSafeNormal();
-	float W = -FVector3f::DotProduct(TranslatedWorldOrigin, FVector3f(LightDirection));
-	TranslatedWorldPlane = FVector4f(LightDirection.X, LightDirection.Y, LightDirection.Z, W);
-
+	// TODO: LODInfo needs to be calculated independently for each cascade
 	FVector PreViewTranslation = View.ViewMatrices.GetPreViewTranslation();
 	HeterogeneousVolumes::FLODInfo LODInfo;
-	LODInfo.WorldSceneBounds = WorldVolumeBounds;
-	LODInfo.WorldOrigin = FVector(TranslatedWorldOrigin) - PreViewTranslation;
+	LODInfo.WorldSceneBounds = GlobalWorldVolumeBounds;
+	//LODInfo.WorldOrigin = FVector(TranslatedWorldOrigin) - PreViewTranslation;
+	LODInfo.WorldOrigin = FVector(TranslatedWorldOrigin[0]) - PreViewTranslation;
 	LODInfo.ViewRect = FIntRect(0, 0, ShadowMapResolution.X, ShadowMapResolution.Y);
 
 	FMatrix WorldToTranslatedWorldMatrix = FTranslationMatrix(PreViewTranslation);
@@ -2380,20 +2790,21 @@ bool RenderVolumetricShadowMapForLightWithLiveShading(
 
 	// Iterate over shadow-casting volumes
 	bool bHasShadowCastingVolume = false;
-	if (!HeterogeneousVolumesMeshBatches.IsEmpty())
+	if (!GlobalHeterogeneousVolumesMeshBatches.IsEmpty())
 	{
-		auto VolumeMeshBatchItr = HeterogeneousVolumesMeshBatches.begin();
+		auto VolumeMeshBatchItr = GlobalHeterogeneousVolumesMeshBatches.begin();
 
 		MaxSampleCount = HeterogeneousVolumes::GetShadowMaxSampleCount();
 		int32 VolumetricShadowLinkedListElementCount = ShadowMapResolution.X * ShadowMapResolution.Y * MaxSampleCount;
 		if (bIsMultiProjection)
 		{
-			VolumetricShadowLinkedListElementCount *= 6;
+			VolumetricShadowLinkedListElementCount *= NumShadowMatrices;
 		}
 		VolumetricShadowLinkedListBuffer = GraphBuilder.CreateBuffer(
 			FRDGBufferDesc::CreateStructuredDesc(sizeof(FAVSMLinkedListPackedData), VolumetricShadowLinkedListElementCount),
 			TEXT("HeterogeneousVolume.VolumetricShadowLinkedListBuffer")
 		);
+		//ClearVolumetricShadowMap(GraphBuilder, Scene, View, ShadowMapResolution, GroupCount, VolumetricShadowLinkedListBuffer);
 
 		bool bUseCameraSceneDepth = false;
 		RenderVolumetricShadowMapForLightForHeterogeneousVolumeWithLiveShading(
@@ -2413,6 +2824,7 @@ bool RenderVolumetricShadowMapForLightWithLiveShading(
 			ShadowToTranslatedWorld,
 			ShadowMapResolution,
 			MaxSampleCount,
+			MaxShadowTraceDistance,
 			bUseCameraSceneDepth,
 			// Volume
 			*VolumeMeshBatchItr,
@@ -2424,7 +2836,7 @@ bool RenderVolumetricShadowMapForLightWithLiveShading(
 		);
 
 		++VolumeMeshBatchItr;
-		for (; VolumeMeshBatchItr != HeterogeneousVolumesMeshBatches.end(); ++VolumeMeshBatchItr)
+		for (; VolumeMeshBatchItr != GlobalHeterogeneousVolumesMeshBatches.end(); ++VolumeMeshBatchItr)
 		{
 			FRDGBufferRef VolumetricShadowLinkedListBuffer1 = GraphBuilder.CreateBuffer(
 				FRDGBufferDesc::CreateStructuredDesc(sizeof(FAVSMLinkedListPackedData), VolumetricShadowLinkedListElementCount),
@@ -2448,6 +2860,7 @@ bool RenderVolumetricShadowMapForLightWithLiveShading(
 				ShadowToTranslatedWorld,
 				ShadowMapResolution,
 				MaxSampleCount,
+				MaxShadowTraceDistance,
 				bUseCameraSceneDepth,
 				// Volume
 				*VolumeMeshBatchItr,
@@ -2477,6 +2890,25 @@ bool RenderVolumetricShadowMapForLightWithLiveShading(
 
 			VolumetricShadowLinkedListBuffer = VolumetricShadowLinkedListBuffer2;
 		}
+
+		if (HeterogeneousVolumes::GetDebugMode() != 0)
+		{
+			// Create debug beer shadow map texture
+			FString LightName;
+			FSceneRenderer::GetLightNameForDrawEvent(LightSceneInfo->Proxy, LightName);
+			RDG_EVENT_SCOPE(GraphBuilder, "%s", *LightName);
+
+			ConvertVolumetricShadowMapToBeerLawShadowMap(
+				GraphBuilder,
+				SceneTextures,
+				View,
+				LightName,
+				ShadowMapResolution,
+				GroupCount,
+				VolumetricShadowLinkedListBuffer,
+				BeerShadowMapTexture
+			);
+		}
 	}
 
 	return true;
@@ -2537,6 +2969,9 @@ void RenderAdaptiveVolumetricShadowMapWithLiveShading(
 			bool bCreateShadowMap = bShouldRenderShadowMaps && bCastsVolumetricShadow && bDynamicallyShadowed && !ShouldRenderRayTracingShadowsForLight(*View.Family, LightSceneInfoCompact[PassIndex]);
 			if (bCreateShadowMap)
 			{
+#if 1
+				FRDGTextureRef BeerShadowMapTexture;
+#else
 				FString LightName;
 				FSceneRenderer::GetLightNameForDrawEvent(LightSceneInfo->Proxy, LightName);
 				RDG_EVENT_SCOPE(GraphBuilder, "%s", *LightName);
@@ -2544,21 +2979,24 @@ void RenderAdaptiveVolumetricShadowMapWithLiveShading(
 				FRDGTextureDesc Desc = SceneTextures.Color.Target->Desc;
 				Desc.Format = PF_FloatRGBA;
 				Desc.Flags &= ~(TexCreate_FastVRAM);
-				FRDGTextureRef BeerShadowMapTexture = GraphBuilder.CreateTexture(Desc, TEXT("BeerShadowMapTexture"));
+				FString ShadowMapName = FString::Printf(TEXT("HeterogeneousVolumes.%s.BeerShadowMapTexture"), *LightName);
+				//FRDGTextureRef BeerShadowMapTexture = GraphBuilder.CreateTexture(Desc, TEXT("HeterogeneousVolumes.BeerShadowMapTexture"));
+				FRDGTextureRef BeerShadowMapTexture = GraphBuilder.CreateTexture(Desc, *ShadowMapName);
 				AddClearUAVPass(GraphBuilder, GraphBuilder.CreateUAV(BeerShadowMapTexture), FLinearColor::Transparent);
-
+#endif
 				bool bIsDirectionalLight = false;
-				FVector3f TranslatedWorldOrigin = FVector3f::Zero();
-				FVector4f TranslatedWorldPlane = FVector4f::Zero();
-				FMatrix44f TranslatedWorldToShadow[] =
+				FVector3f TranslatedWorldOrigin[6];
+				FVector4f TranslatedWorldPlane[6];
+				FVector4f SplitDepths[6];
+				FMatrix44f TranslatedWorldToShadow[6];
+				for (int i = 0; i < 6; ++i)
 				{
-					FMatrix44f::Identity,
-					FMatrix44f::Identity,
-					FMatrix44f::Identity,
-					FMatrix44f::Identity,
-					FMatrix44f::Identity,
-					FMatrix44f::Identity
-				};
+					TranslatedWorldOrigin[i] = FVector3f::Zero();
+					TranslatedWorldPlane[i] = FVector4f::Zero();
+					SplitDepths[i] = FVector4f::Zero();
+					TranslatedWorldToShadow[i] = FMatrix44f::Identity;
+				}
+
 				FIntVector GroupCount = FIntVector::ZeroValue;
 				int32 NumShadowMatrices = 0;
 				FIntPoint VolumetricShadowMapResolution = FIntPoint::NoneValue;
@@ -2579,6 +3017,7 @@ void RenderAdaptiveVolumetricShadowMapWithLiveShading(
 					TranslatedWorldOrigin,
 					TranslatedWorldPlane,
 					TranslatedWorldToShadow,
+					SplitDepths,
 					GroupCount,
 					NumShadowMatrices,
 					VolumetricShadowMapResolution,
@@ -2608,6 +3047,7 @@ void RenderAdaptiveVolumetricShadowMapWithLiveShading(
 						TranslatedWorldOrigin,
 						TranslatedWorldPlane,
 						TranslatedWorldToShadow,
+						SplitDepths,
 						VolumetricShadowMapResolution,
 						DownsampleFactor,
 						NumShadowMatrices,
@@ -2634,6 +3074,22 @@ void RenderAdaptiveVolumetricShadowMapWithLiveShading(
 				TRDGUniformBufferRef<FAdaptiveVolumetricShadowMapUniformBufferParameters>& AdaptiveVolumetricShadowMap = View.ViewState->AdaptiveVolumetricShadowMapUniformBufferMap.FindOrAdd(LightSceneInfo->Id);
 				AdaptiveVolumetricShadowMap = AdaptiveVolumetricShadowMapUniformBuffer;
 			}
+
+			if (HeterogeneousVolumes::GetDebugMode() != 0)
+			{
+				FRDGTextureRef VisualizationTexture;
+				VisualizeCascades(
+					GraphBuilder,
+					// Scene data
+					SceneTextures,
+					Scene,
+					View,
+					// Light data
+					LightSceneInfo,
+					// Output
+					VisualizationTexture
+				);
+			}
 		}
 	}
 }
@@ -2690,7 +3146,9 @@ void RenderAdaptiveVolumetricCameraMapWithLiveShading(
 		FMatrix44f ShadowToTranslatedWorld[] = {
 			TranslatedWorldToShadow[0].Inverse()
 		};
-		FVector3f TranslatedWorldOrigin = ShadowToTranslatedWorld[0].GetOrigin();
+		FVector3f TranslatedWorldOrigin[] = {
+			ShadowToTranslatedWorld[0].GetOrigin()
+		};
 
 		// Dispatch
 		FIntVector GroupCount = FIntVector(1);
@@ -2718,7 +3176,7 @@ void RenderAdaptiveVolumetricCameraMapWithLiveShading(
 
 		HeterogeneousVolumes::FLODInfo LODInfo;
 		LODInfo.WorldSceneBounds = WorldVolumeBounds;
-		LODInfo.WorldOrigin = FVector(TranslatedWorldOrigin) - PreViewTranslation;
+		LODInfo.WorldOrigin = FVector(TranslatedWorldOrigin[0]) - PreViewTranslation;
 		LODInfo.ViewRect = FIntRect(0, 0, ShadowMapResolution.X, ShadowMapResolution.Y);
 		GetViewFrustumBounds(LODInfo.WorldShadowFrustum, WorldToTranslatedWorldMatrix * FMatrix(TranslatedWorldToShadow[0]), true);
 		LODInfo.FOV = FOV;
@@ -2748,6 +3206,7 @@ void RenderAdaptiveVolumetricCameraMapWithLiveShading(
 			ShadowToTranslatedWorld,
 			ShadowMapResolution,
 			MaxSampleCount,
+			HeterogeneousVolumes::GetMaxTraceDistance(),
 			bUseCameraSceneDepth,
 			// Volume
 			*VolumeMeshBatchItr,
@@ -2784,6 +3243,7 @@ void RenderAdaptiveVolumetricCameraMapWithLiveShading(
 				ShadowToTranslatedWorld,
 				ShadowMapResolution,
 				MaxSampleCount,
+				HeterogeneousVolumes::GetMaxTraceDistance(),
 				bUseCameraSceneDepth,
 				// Volume
 				*VolumeMeshBatchItr,
@@ -2827,7 +3287,12 @@ void RenderAdaptiveVolumetricCameraMapWithLiveShading(
 			VolumetricShadowSampleBuffer
 		);
 
-		FVector4f TranslatedWorldPlane = FVector4f::Zero();
+		FVector4f TranslatedWorldPlane[] = {
+			FVector4f::Zero()
+		};
+		FVector4f SplitDepths[] = {
+			FVector4f::Zero()
+		};
 		float DownsampleFactor = HeterogeneousVolumes::GetCameraDownsampleFactor();
 		bool bIsDirectionalLight = false;
 		CreateAdaptiveVolumetricShadowMapUniformBufferParameters(
@@ -2835,6 +3300,7 @@ void RenderAdaptiveVolumetricCameraMapWithLiveShading(
 			TranslatedWorldOrigin,
 			TranslatedWorldPlane,
 			TranslatedWorldToShadow,
+			SplitDepths,
 			ShadowMapResolution,
 			DownsampleFactor,
 			NumShadowMatrices,
diff --git a/Engine/Source/Runtime/Renderer/Private/HeterogeneousVolumes/HeterogeneousVolumesVoxelGridPipeline.cpp b/Engine/Source/Runtime/Renderer/Private/HeterogeneousVolumes/HeterogeneousVolumesVoxelGridPipeline.cpp
--- a/Engine/Source/Runtime/Renderer/Private/HeterogeneousVolumes/HeterogeneousVolumesVoxelGridPipeline.cpp
+++ b/Engine/Source/Runtime/Renderer/Private/HeterogeneousVolumes/HeterogeneousVolumesVoxelGridPipeline.cpp
@@ -280,6 +280,24 @@ static TAutoConsoleVariable<bool> CVarHeterogeneousShadowsForLightTypeDirectiona
 	ECVF_RenderThreadSafe
 );
 
+static TAutoConsoleVariable<int32> CVarHeterogeneousCascadeShadowsForDirectionalLight(
+	TEXT("r.HeterogeneousVolumes.Shadows.Cascades"),
+	1,
+	TEXT("Enables cascade shadows for the directional light (Default = 1)\n")
+	TEXT("0: Disabled\n")
+	TEXT("1: View-based frustums (Default)\n")
+	TEXT("2: View-based clipmaps centered about the camera\n")
+	TEXT("3: Auto-fit shadow maps based on bounding dimensions of contained volumes"),
+	ECVF_RenderThreadSafe
+);
+
+static TAutoConsoleVariable<int32> CVarHeterogeneousVolumesCascadeShadowsQuantization(
+	TEXT("r.HeterogeneousVolumes.Shadows.Cascades.PixelSnapping"),
+	16,
+	TEXT("Cascade quantization unit in pixel-space (Default = 16)"),
+	ECVF_RenderThreadSafe
+);
+
 static TAutoConsoleVariable<bool> CVarHeterogeneousShadowsForLightTypePoint(
 	TEXT("r.HeterogeneousVolumes.Shadows.LightType.Point"),
 	1,
@@ -509,6 +527,31 @@ namespace HeterogeneousVolumes
 		return false;
 	}
 
+	bool SupportsCascadeShadowsForDirectionalLight()
+	{
+		if (SupportsLightType(LightType_Directional) && SupportsShadowForLightType(LightType_Directional))
+		{
+			return CVarHeterogeneousCascadeShadowsForDirectionalLight.GetValueOnRenderThread() != 0;
+		}
+
+		return false;
+	}
+
+	ECascadeShadowMode GetCascadeShadowMode()
+	{
+		return static_cast<ECascadeShadowMode>(CVarHeterogeneousCascadeShadowsForDirectionalLight.GetValueOnRenderThread());
+	}
+
+	int32 GetCascadeShadowsQuantizationUnit()
+	{
+		if (SupportsCascadeShadowsForDirectionalLight())
+		{
+			return CVarHeterogeneousVolumesCascadeShadowsQuantization.GetValueOnRenderThread();
+		}
+
+		return 0;
+	}
+
 	bool IsDynamicShadow(const FVisibleLightInfo* VisibleLightInfo)
 	{
 		check(VisibleLightInfo != nullptr);
@@ -3303,11 +3346,106 @@ void ConvertBeerLawShadowMapToVolumetricShadowMap(
 	);
 }
 
+class FConvertVolumetricShadowMapToBeerLawShadowMapCS : public FGlobalShader
+{
+	DECLARE_GLOBAL_SHADER(FConvertVolumetricShadowMapToBeerLawShadowMapCS);
+	SHADER_USE_PARAMETER_STRUCT(FConvertVolumetricShadowMapToBeerLawShadowMapCS, FGlobalShader);
+
+	using FPermutationDomain = TShaderPermutationDomain<>;
+
+	BEGIN_SHADER_PARAMETER_STRUCT(FParameters, )
+		// Input
+		SHADER_PARAMETER(FIntPoint, ShadowResolution)
+		SHADER_PARAMETER(int32, MaxSampleCount)
+		SHADER_PARAMETER_RDG_BUFFER_SRV(StructuredBuffer<FAVSMLinkedListPackedData>, VolumetricShadowLinkedListBuffer0)
+
+		// Output
+		SHADER_PARAMETER_RDG_TEXTURE_UAV(Texture2D<float4>, RWBeerShadowMapTexture)
+	END_SHADER_PARAMETER_STRUCT()
+
+	static bool ShouldCompilePermutation(
+		const FGlobalShaderPermutationParameters& Parameters
+	)
+	{
+		return DoesPlatformSupportHeterogeneousVolumes(Parameters.Platform);
+	}
+
+	static void ModifyCompilationEnvironment(
+		const FGlobalShaderPermutationParameters& Parameters,
+		FShaderCompilerEnvironment& OutEnvironment
+	)
+	{
+		FMaterialShader::ModifyCompilationEnvironment(Parameters, OutEnvironment);
+		OutEnvironment.SetDefine(TEXT("THREADGROUP_SIZE_1D"), GetThreadGroupSize1D());
+		OutEnvironment.SetDefine(TEXT("THREADGROUP_SIZE_2D"), GetThreadGroupSize2D());
+		OutEnvironment.SetDefine(TEXT("THREADGROUP_SIZE_3D"), GetThreadGroupSize3D());
+
+		OutEnvironment.SetDefine(TEXT("DIM_USE_TRANSMITTANCE_VOLUME"), 0);
+		OutEnvironment.SetDefine(TEXT("DIM_USE_INSCATTERING_VOLUME"), 0);
+		OutEnvironment.SetDefine(TEXT("DIM_USE_LUMEN_GI"), 0);
+
+		OutEnvironment.CompilerFlags.Add(CFLAG_AllowTypedUAVLoads);
+	}
+
+	static int32 GetThreadGroupSize1D() { return GetThreadGroupSize3D() * GetThreadGroupSize3D() * GetThreadGroupSize3D(); }
+	static int32 GetThreadGroupSize2D() { return 8; }
+	static int32 GetThreadGroupSize3D() { return 4; }
+};
+
+IMPLEMENT_GLOBAL_SHADER(FConvertVolumetricShadowMapToBeerLawShadowMapCS, "/Engine/Private/HeterogeneousVolumes/HeterogeneousVolumesVoxelGridShadows.usf", "ConvertVolumetricShadowMapToBeerLawShadowMapCS", SF_Compute);
+
+void ConvertVolumetricShadowMapToBeerLawShadowMap(
+	FRDGBuilder& GraphBuilder,
+	const FSceneTextures& SceneTextures,
+	FViewInfo& View,
+	// Input
+	const FString& ShadowMapName,
+	FIntPoint ShadowMapResolution,
+	FIntVector GroupCount,
+	FRDGBufferRef VolumetricShadowLinkedListBuffer,
+	// Output
+	FRDGTextureRef& BeerShadowMapTexture
+)
+{
+	FRDGTextureDesc Desc = SceneTextures.Color.Target->Desc;
+	Desc.Format = PF_FloatRGBA;
+	Desc.Flags &= ~(TexCreate_FastVRAM);
+	Desc.Extent = FIntPoint(ShadowMapResolution.X * GroupCount.Z, ShadowMapResolution.Y);
+	//FString ShadowMapString = FString::Printf(TEXT("HeterogeneousVolumes.%s.BeerShadowMapTexture"), *ShadowMapName);
+	//BeerShadowMapTexture = GraphBuilder.CreateTexture(Desc, *ShadowMapString);
+	BeerShadowMapTexture = GraphBuilder.CreateTexture(Desc, TEXT("HeterogeneousVolumes.Light.BeerShadowMapTexture"));
+	AddClearUAVPass(GraphBuilder, GraphBuilder.CreateUAV(BeerShadowMapTexture), FLinearColor::Transparent);
+
+	FConvertVolumetricShadowMapToBeerLawShadowMapCS::FParameters* PassParameters = GraphBuilder.AllocParameters<FConvertVolumetricShadowMapToBeerLawShadowMapCS::FParameters>();
+	{
+		// Input
+		PassParameters->ShadowResolution = ShadowMapResolution;
+		PassParameters->MaxSampleCount = HeterogeneousVolumes::GetShadowMaxSampleCount();
+		PassParameters->VolumetricShadowLinkedListBuffer0 = GraphBuilder.CreateSRV(VolumetricShadowLinkedListBuffer);
+
+		// Output
+		PassParameters->RWBeerShadowMapTexture = GraphBuilder.CreateUAV(BeerShadowMapTexture);
+	}
+
+	FConvertVolumetricShadowMapToBeerLawShadowMapCS::FPermutationDomain PermutationVector;
+	TShaderRef<FConvertVolumetricShadowMapToBeerLawShadowMapCS> ComputeShader = View.ShaderMap->GetShader<FConvertVolumetricShadowMapToBeerLawShadowMapCS>(PermutationVector);
+	FComputeShaderUtils::AddPass(
+		GraphBuilder,
+		// Add Light name..
+		RDG_EVENT_NAME("FConvertBeerLawShadowMapToVolumetricShadowMapCS"),
+		ERDGPassFlags::Compute,
+		ComputeShader,
+		PassParameters,
+		GroupCount
+	);
+}
+
 void CreateAdaptiveVolumetricShadowMapUniformBufferParameters(
 	FRDGBuilder& GraphBuilder,
-	const FVector3f& TranslatedWorldOrigin,
-	const FVector4f& TranslatedWorldPlane,
+	const FVector3f* TranslatedWorldOrigin,
+	const FVector4f* TranslatedWorldPlane,
 	const FMatrix44f* TranslatedWorldToShadow,
+	const FVector4f* SplitDepths,
 	FIntPoint VolumetricShadowMapResolution,
 	float VolumetricShadowMapDownsampleFactor,
 	int32 NumShadowMatrices,
@@ -3325,9 +3463,11 @@ void CreateAdaptiveVolumetricShadowMapUniformBufferParameters(
 		for (int32 i = 0; i < NumShadowMatrices; ++i)
 		{
 			Parameters->TranslatedWorldToShadow[i] = TranslatedWorldToShadow[i];
+			Parameters->TranslatedWorldOrigin[i] = TranslatedWorldOrigin[i];
+			Parameters->TranslatedWorldPlane[i] = TranslatedWorldPlane[i];
+			Parameters->SplitDepths[i] = SplitDepths[i];
 		}
-		Parameters->TranslatedWorldOrigin = TranslatedWorldOrigin;
-		Parameters->TranslatedWorldPlane = TranslatedWorldPlane;
+
 		Parameters->Resolution = VolumetricShadowMapResolution;
 		Parameters->DownsampleFactor = VolumetricShadowMapDownsampleFactor;
 		Parameters->MaxSampleCount = VolumetricShadowMapMaxSampleCount;
@@ -3344,9 +3484,10 @@ void CreateAdaptiveVolumetricShadowMapUniformBufferParameters(
 
 void CreateAdaptiveVolumetricShadowMapUniformBuffer(
 	FRDGBuilder& GraphBuilder,
-	const FVector3f& TranslatedWorldOrigin,
-	const FVector4f& TranslatedWorldPlane,
+	const FVector3f* TranslatedWorldOrigin,
+	const FVector4f* TranslatedWorldPlane,
 	const FMatrix44f* TranslatedWorldToShadow,
+	const FVector4f* SplitDepths,
 	FIntPoint VolumetricShadowMapResolution,
 	float VolumetricShadowMapDownsampleFactor,
 	int32 NumShadowMatrices,
@@ -3364,6 +3505,7 @@ void CreateAdaptiveVolumetricShadowMapUniformBuffer(
 		TranslatedWorldOrigin,
 		TranslatedWorldPlane,
 		TranslatedWorldToShadow,
+		SplitDepths,
 		VolumetricShadowMapResolution,
 		VolumetricShadowMapDownsampleFactor,
 		NumShadowMatrices,
@@ -3386,12 +3528,14 @@ namespace HeterogeneousVolumes
 		FAdaptiveVolumetricShadowMapUniformBufferParameters* UniformBufferParameters = GraphBuilder.AllocParameters<FAdaptiveVolumetricShadowMapUniformBufferParameters>();
 		{
 			UniformBufferParameters->NumShadowMatrices = 1;
-			for (int32 i = 0; i < UniformBufferParameters->NumShadowMatrices; ++i)
+			for (int32 i = 0; i < 6; ++i)
 			{
 				UniformBufferParameters->TranslatedWorldToShadow[i] = FMatrix44f::Identity;
+				UniformBufferParameters->TranslatedWorldOrigin[i] = FVector3f::ZeroVector;
+				UniformBufferParameters->TranslatedWorldPlane[i] = FVector4f::Zero();
+				UniformBufferParameters->SplitDepths[i] = FVector4f::Zero();
 			}
-			UniformBufferParameters->TranslatedWorldOrigin = FVector3f::ZeroVector;
-			UniformBufferParameters->TranslatedWorldPlane = FVector4f::Zero();
+
 			UniformBufferParameters->Resolution = FIntPoint::ZeroValue;
 			UniformBufferParameters->DownsampleFactor = 1.0f;
 			UniformBufferParameters->MaxSampleCount = 0;
@@ -3466,9 +3610,11 @@ namespace HeterogeneousVolumes
 			for (int32 i = 0; i < Parameters.NumShadowMatrices; ++i)
 			{
 				Parameters.TranslatedWorldToShadow[i] = UniformBuffer->GetParameters()->TranslatedWorldToShadow[i];
+				Parameters.TranslatedWorldOrigin[i] = UniformBuffer->GetParameters()->TranslatedWorldOrigin[i];
+				Parameters.TranslatedWorldPlane[i] = UniformBuffer->GetParameters()->TranslatedWorldPlane[i];
+				Parameters.SplitDepths[i] = UniformBuffer->GetParameters()->SplitDepths[i];
 			}
-			Parameters.TranslatedWorldOrigin = UniformBuffer->GetParameters()->TranslatedWorldOrigin;
-			Parameters.TranslatedWorldPlane = UniformBuffer->GetParameters()->TranslatedWorldPlane;
+
 			Parameters.Resolution = UniformBuffer->GetParameters()->Resolution;
 			Parameters.DownsampleFactor = UniformBuffer->GetParameters()->DownsampleFactor;
 			Parameters.MaxSampleCount = UniformBuffer->GetParameters()->MaxSampleCount;
@@ -3486,9 +3632,10 @@ namespace HeterogeneousVolumes
 			for (int32 i = 0; i < Parameters.NumShadowMatrices; ++i)
 			{
 				Parameters.TranslatedWorldToShadow[i] = FMatrix44f::Identity;
+				Parameters.TranslatedWorldOrigin[i] = FVector3f::ZeroVector;
+				Parameters.TranslatedWorldPlane[i] = FVector4f::Zero();
 			}
-			Parameters.TranslatedWorldOrigin = FVector3f::ZeroVector;
-			Parameters.TranslatedWorldPlane = FVector4f::Zero();
+
 			Parameters.Resolution = FIntPoint::ZeroValue;
 			Parameters.DownsampleFactor = 1.0f;
 			Parameters.MaxSampleCount = 0;
@@ -3579,16 +3726,11 @@ void CombineVolumetricShadowMap(
 )
 {
 	bool bIsMultiProjection = (LightType == LightType_Point) || (LightType == LightType_Rect);
-	GroupCount.X = FMath::DivideAndRoundUp(ShadowMapResolution.X, FCombineVolumetricShadowMapsCS::GetThreadGroupSize2D());
-	GroupCount.Y = FMath::DivideAndRoundUp(ShadowMapResolution.Y, FCombineVolumetricShadowMapsCS::GetThreadGroupSize2D());
-	GroupCount.Z = bIsMultiProjection ? 6 : 1;
 
 	MaxSampleCount = HeterogeneousVolumes::GetShadowMaxSampleCount();
 	int32 VolumetricShadowLinkedListElementCount = ShadowMapResolution.X * ShadowMapResolution.Y * MaxSampleCount;
-	if (bIsMultiProjection)
-	{
-		VolumetricShadowLinkedListElementCount *= 6;
-	}
+	VolumetricShadowLinkedListElementCount *= GroupCount.Z;
+
 	VolumetricShadowLinkedListBuffer = GraphBuilder.CreateBuffer(
 		FRDGBufferDesc::CreateStructuredDesc(sizeof(FAVSMLinkedListPackedData), VolumetricShadowLinkedListElementCount),
 		TEXT("HeterogeneousVolume.VolumetricShadowLinkedListBuffer")
@@ -3643,8 +3785,8 @@ void RenderVolumetricShadowMapForLightWithVoxelGrid(
 	const TRDGUniformBufferRef<FFrustumVoxelGridUniformBufferParameters>& FrustumGridUniformBuffer,
 	// Output
 	bool& bIsDirectionalLight,
-	FVector3f& TranslatedWorldOrigin,
-	FVector4f& TranslatedWorldPlane,
+	FVector3f* TranslatedWorldOrigin,
+	FVector4f* TranslatedWorldPlane,
 	FMatrix44f* TranslatedWorldToShadow,
 	FIntVector& GroupCount,
 	int32& NumShadowMatrices,
@@ -3812,7 +3954,7 @@ void RenderVolumetricShadowMapForLightWithVoxelGrid(
 			PassParameters->TranslatedWorldOrigin = FVector3f(View.ViewMatrices.GetPreViewTranslation() - ProjectedShadowInfo->PreShadowTranslation);
 		}
 
-		TranslatedWorldOrigin = PassParameters->TranslatedWorldOrigin;
+		TranslatedWorldOrigin[0] = PassParameters->TranslatedWorldOrigin;
 		NumShadowMatrices = PassParameters->NumShadowMatrices;
 		for (int32 i = 0; i < PassParameters->NumShadowMatrices; ++i)
 		{
@@ -3820,8 +3962,8 @@ void RenderVolumetricShadowMapForLightWithVoxelGrid(
 		}
 
 		FVector LightDirection = LightSceneInfo->Proxy->GetDirection().GetSafeNormal();
-		float W = -FVector3f::DotProduct(TranslatedWorldOrigin, FVector3f(LightDirection));
-		TranslatedWorldPlane = FVector4f(LightDirection.X, LightDirection.Y, LightDirection.Z, W);
+		float W = -FVector3f::DotProduct(TranslatedWorldOrigin[0], FVector3f(LightDirection));
+		TranslatedWorldPlane[0] = FVector4f(LightDirection.X, LightDirection.Y, LightDirection.Z, W);
 
 		PassParameters->ShadowResolution = ShadowMapResolution;
 		PassParameters->MaxSampleCount = MaxSampleCount;
@@ -4048,17 +4190,18 @@ void RenderAdaptiveVolumetricShadowMapWithVoxelGrid(
 				AddClearUAVPass(GraphBuilder, GraphBuilder.CreateUAV(BeerShadowMapTexture), FLinearColor::Transparent);
 
 				bool bIsDirectionalLight = false;
-				FVector3f TranslatedWorldOrigin = FVector3f::Zero();
-				FVector4f TranslatedWorldPlane = FVector4f::Zero();
-				FMatrix44f TranslatedWorldToShadow[] =
+				FVector3f TranslatedWorldOrigin[6];
+				FVector4f TranslatedWorldPlane[6];
+				FMatrix44f TranslatedWorldToShadow[6];
+				FVector4f SplitDepths[6];
+				for (int i = 0; i < 6; ++i)
 				{
-					FMatrix44f::Identity,
-					FMatrix44f::Identity,
-					FMatrix44f::Identity,
-					FMatrix44f::Identity,
-					FMatrix44f::Identity,
-					FMatrix44f::Identity
-				};
+					TranslatedWorldOrigin[i] = FVector3f::Zero();
+					TranslatedWorldPlane[i] = FVector4f::Zero();
+					TranslatedWorldToShadow[i] = FMatrix44f::Identity;
+					SplitDepths[i] = FVector4f::Zero();
+				}
+
 				FIntVector GroupCount = FIntVector::ZeroValue;
 				int32 NumShadowMatrices = 0;
 				FIntPoint VolumetricShadowMapResolution = FIntPoint::NoneValue;
@@ -4114,6 +4257,7 @@ void RenderAdaptiveVolumetricShadowMapWithVoxelGrid(
 					TranslatedWorldOrigin,
 					TranslatedWorldPlane,
 					TranslatedWorldToShadow,
+					SplitDepths,
 					VolumetricShadowMapResolution,
 					DownsampleFactor,
 					NumShadowMatrices,
@@ -4156,12 +4300,19 @@ void RenderAdaptiveVolumetricCameraMapWithVoxelGrid(
 	}
 	RDG_EVENT_SCOPE(GraphBuilder, "Adaptive Volumetric Camera Map");
 
-	FVector3f TranslatedWorldOrigin = FVector3f::ZeroVector;
 	int32 NumShadowMatrices = 1;
-	FMatrix44f TranslatedWorldToShadow[] =
-	{
+	FMatrix44f TranslatedWorldToShadow[] = {
 		FMatrix44f::Identity
 	};
+	FVector3f TranslatedWorldOrigin[] = {
+		FVector3f::ZeroVector
+	};
+	FVector4f TranslatedWorldPlane[] = {
+		FVector4f::Zero()
+	};
+	FVector4f SplitDepths[] = {
+		FVector4f::Zero()
+	};
 	FIntPoint VolumetricShadowMapResolution = FIntPoint::NoneValue;
 	uint32 VolumetricShadowMapMaxSampleCount = 0;
 	FRDGBufferRef VolumetricShadowMapLinkedListBuffer = GSystemTextures.GetDefaultStructuredBuffer(GraphBuilder, sizeof(FAVSMLinkedListPackedData));;
@@ -4190,7 +4341,7 @@ void RenderAdaptiveVolumetricCameraMapWithVoxelGrid(
 			OrthoGridUniformBuffer,
 			FrustumGridUniformBuffer,
 			// Output
-			TranslatedWorldOrigin,
+			TranslatedWorldOrigin[0],
 			TranslatedWorldToShadow[0],
 			GroupCount,
 			VolumetricShadowMapResolution,
@@ -4210,7 +4361,6 @@ void RenderAdaptiveVolumetricCameraMapWithVoxelGrid(
 			VolumetricShadowMapSampleBuffer
 		);
 
-		FVector4f TranslatedWorldPlane = FVector4f::Zero();
 		bool bIsDirectionalLight = false;
 		float DownsampleFactor = HeterogeneousVolumes::GetCameraDownsampleFactor();
 
@@ -4219,6 +4369,7 @@ void RenderAdaptiveVolumetricCameraMapWithVoxelGrid(
 			TranslatedWorldOrigin,
 			TranslatedWorldPlane,
 			TranslatedWorldToShadow,
+			SplitDepths,
 			VolumetricShadowMapResolution,
 			DownsampleFactor,
 			NumShadowMatrices,
@@ -4231,3 +4382,95 @@ void RenderAdaptiveVolumetricCameraMapWithVoxelGrid(
 		);
 	}
 }
+
+class FVisualizeCascadesCS : public FGlobalShader
+{
+	DECLARE_GLOBAL_SHADER(FVisualizeCascadesCS);
+	SHADER_USE_PARAMETER_STRUCT(FVisualizeCascadesCS, FGlobalShader);
+
+	BEGIN_SHADER_PARAMETER_STRUCT(FParameters, )
+		// Scene data
+		SHADER_PARAMETER_STRUCT_REF(FViewUniformShaderParameters, View)
+		SHADER_PARAMETER_STRUCT_INCLUDE(FSceneTextureParameters, SceneTextures)
+		SHADER_PARAMETER_RDG_UNIFORM_BUFFER(FSceneUniformParameters, Scene)
+
+		// Shadow data
+		SHADER_PARAMETER_RDG_UNIFORM_BUFFER(FAdaptiveVolumetricShadowMapUniformBufferParameters, AVSM)
+
+		// Dispatch data
+		SHADER_PARAMETER(int32, DownsampleFactor)
+
+		// Output
+		SHADER_PARAMETER_RDG_TEXTURE_UAV(RWTexture3D<float>, RWVisualizationTexture)
+	END_SHADER_PARAMETER_STRUCT()
+
+	static bool ShouldCompilePermutation(const FGlobalShaderPermutationParameters& Parameters)
+	{
+		return DoesPlatformSupportHeterogeneousVolumes(Parameters.Platform);
+	}
+
+	static void ModifyCompilationEnvironment(
+		const FGlobalShaderPermutationParameters& Parameters,
+		FShaderCompilerEnvironment& OutEnvironment
+	)
+	{
+		FMaterialShader::ModifyCompilationEnvironment(Parameters, OutEnvironment);
+		OutEnvironment.SetDefine(TEXT("THREADGROUP_SIZE_1D"), GetThreadGroupSize1D());
+		OutEnvironment.SetDefine(TEXT("THREADGROUP_SIZE_2D"), GetThreadGroupSize2D());
+		OutEnvironment.SetDefine(TEXT("THREADGROUP_SIZE_3D"), GetThreadGroupSize3D());
+
+		// This shader takes a very long time to compile with FXC, so we pre-compile it with DXC first and then forward the optimized HLSL to FXC.
+		OutEnvironment.CompilerFlags.Add(CFLAG_AllowTypedUAVLoads);
+	}
+
+	static int32 GetThreadGroupSize1D() { return 64; }
+	static int32 GetThreadGroupSize2D() { return 8; }
+	static int32 GetThreadGroupSize3D() { return 4; }
+};
+
+IMPLEMENT_GLOBAL_SHADER(FVisualizeCascadesCS, "/Engine/Private/HeterogeneousVolumes/HeterogeneousVolumesVisualizeCascades.usf", "VisualizeCascadesCS", SF_Compute);
+
+void VisualizeCascades(
+	FRDGBuilder& GraphBuilder,
+	// Scene data
+	const FSceneTextures& SceneTextures,
+	const FScene* Scene,
+	const FViewInfo& View,
+	// Light data
+	const FLightSceneInfo* LightSceneInfo,
+	// Output
+	FRDGTextureRef& VisualizationTexture
+)
+{
+	FRDGTextureDesc VisualizationTextureDesc = SceneTextures.Color.Target->Desc;
+	VisualizationTextureDesc.Extent = HeterogeneousVolumes::GetScaledViewRect(View.ViewRect);
+	VisualizationTexture = GraphBuilder.CreateTexture(VisualizationTextureDesc, TEXT("HeterogeneousVolumes.VisualizationTexture"));
+
+	FVisualizeCascadesCS::FParameters* PassParameters = GraphBuilder.AllocParameters<FVisualizeCascadesCS::FParameters>();
+	{
+		// Scene
+		PassParameters->View = View.ViewUniformBuffer;
+		PassParameters->SceneTextures = GetSceneTextureParameters(GraphBuilder, SceneTextures);
+		PassParameters->Scene = View.GetSceneUniforms().GetBuffer(GraphBuilder);
+
+		// Shadow
+		PassParameters->AVSM = HeterogeneousVolumes::GetAdaptiveVolumetricShadowMapUniformBuffer(GraphBuilder, View.ViewState, LightSceneInfo);
+
+		// Dispatch
+		PassParameters->DownsampleFactor = HeterogeneousVolumes::GetDownsampleFactor();
+
+		// Output
+		PassParameters->RWVisualizationTexture = GraphBuilder.CreateUAV(VisualizationTexture);
+	}
+
+	FIntVector GroupCount = FComputeShaderUtils::GetGroupCount(VisualizationTextureDesc.Extent, FVisualizeCascadesCS::GetThreadGroupSize2D());
+
+	TShaderRef<FVisualizeCascadesCS> ComputeShader = View.ShaderMap->GetShader<FVisualizeCascadesCS>();
+	FComputeShaderUtils::AddPass(
+		GraphBuilder,
+		RDG_EVENT_NAME("FVisualizeCascadesCS"),
+		ComputeShader,
+		PassParameters,
+		GroupCount
+	);
+}
\ No newline at end of file
