Fix shader compilation (RWBeerShadowMapTexture type)

#rnx

diff --git a/Engine/Source/Runtime/Renderer/Private/HeterogeneousVolumes/HeterogeneousVolumesVoxelGridPipeline.cpp b/Engine/Source/Runtime/Renderer/Private/HeterogeneousVolumes/HeterogeneousVolumesVoxelGridPipeline.cpp
--- a/Engine/Source/Runtime/Renderer/Private/HeterogeneousVolumes/HeterogeneousVolumesVoxelGridPipeline.cpp
+++ b/Engine/Source/Runtime/Renderer/Private/HeterogeneousVolumes/HeterogeneousVolumesVoxelGridPipeline.cpp
@@ -3360,7 +3360,7 @@ class FConvertVolumetricShadowMapToBeerLawShadowMapCS : public FGlobalShader
 		SHADER_PARAMETER_RDG_BUFFER_SRV(StructuredBuffer<FAVSMLinkedListPackedData>, VolumetricShadowLinkedListBuffer0)
 
 		// Output
-		SHADER_PARAMETER_RDG_TEXTURE_UAV(Texture2D<float4>, RWBeerShadowMapTexture)
+		SHADER_PARAMETER_RDG_TEXTURE_UAV(RWTexture2D<float4>, RWBeerShadowMapTexture)
 	END_SHADER_PARAMETER_STRUCT()
 
 	static bool ShouldCompilePermutation(
