Use GEngine->GetMonospaceFont() in debugging code instead of loading font multiple times in separate places.

diff --git a/Engine/Plugins/Animation/GameplayInsights/Source/RewindDebuggerVLog/Private/RewindDebuggerVLog.cpp b/Engine/Plugins/Animation/GameplayInsights/Source/RewindDebuggerVLog/Private/RewindDebuggerVLog.cpp
--- a/Engine/Plugins/Animation/GameplayInsights/Source/RewindDebuggerVLog/Private/RewindDebuggerVLog.cpp
+++ b/Engine/Plugins/Animation/GameplayInsights/Source/RewindDebuggerVLog/Private/RewindDebuggerVLog.cpp
@@ -93,9 +93,6 @@ void FRewindDebuggerVLog::Initialize()
 	TraceDevice.ImmediateRenderDelegate.BindRaw(this, &FRewindDebuggerVLog::ImmediateRender);
 
 	DelegateHandle = UDebugDrawService::Register(TEXT("VirtualTextureResidency")/*TEXT("VisLog")*/, FDebugDrawDelegate::CreateRaw(this, &FRewindDebuggerVLog::OnShowDebugInfo));
-
-	FTopLevelAssetPath MonospaceFontPath = FTopLevelAssetPath("/Engine/EngineFonts/DroidSansMono.DroidSansMono");
-	MonospaceFont = LoadObject<UFont>(nullptr, *MonospaceFontPath.ToString(), nullptr, LOAD_None, nullptr);
 }
 
 bool MatchCategoryFilters(const FName& CategoryName, ELogVerbosity::Type Verbosity)
@@ -117,7 +114,7 @@ void FRewindDebuggerVLog::RenderLogEntry(const FVisualLogEntry& Entry, UCanvas*
 	else
 	{
 		UWorld* World = IRewindDebugger::Instance()->GetWorldToVisualize();
-		FVisualLogEntryRenderer::RenderLogEntry(World, Entry, &MatchCategoryFilters, Canvas, GEngine->GetMediumFont(), MonospaceFont, ScreenTextY);
+		FVisualLogEntryRenderer::RenderLogEntry(World, Entry, &MatchCategoryFilters, Canvas, GEngine->GetMediumFont(), GEngine->GetMonospaceFont(), ScreenTextY);
 
 	}
 }
diff --git a/Engine/Plugins/Animation/GameplayInsights/Source/RewindDebuggerVLog/Private/RewindDebuggerVLog.h b/Engine/Plugins/Animation/GameplayInsights/Source/RewindDebuggerVLog/Private/RewindDebuggerVLog.h
--- a/Engine/Plugins/Animation/GameplayInsights/Source/RewindDebuggerVLog/Private/RewindDebuggerVLog.h
+++ b/Engine/Plugins/Animation/GameplayInsights/Source/RewindDebuggerVLog/Private/RewindDebuggerVLog.h
@@ -47,7 +47,6 @@ private:
 	int32 ScreenTextY = 0;
 
 	FDelegateHandle DelegateHandle;
-	UFont* MonospaceFont = nullptr;
 
 	TArray<uint64> DebuggedObjectIds;
 	TArray<FVisualLogEntry> ImmediateRenderQueue;
diff --git a/Engine/Source/Runtime/Engine/Classes/Engine/Engine.h b/Engine/Source/Runtime/Engine/Classes/Engine/Engine.h
--- a/Engine/Source/Runtime/Engine/Classes/Engine/Engine.h
+++ b/Engine/Source/Runtime/Engine/Classes/Engine/Engine.h
@@ -3045,7 +3045,7 @@ public:
 	/**
 	 * Gets the engine's default monospace font.
 	 *
-	 * @return Subtitle font.
+	 * @return Monospace font.
 	 */
 	static ENGINE_API class UFont* GetMonospaceFont();
 
diff --git a/Engine/Source/Runtime/Renderer/Private/VirtualShadowMaps/VirtualShadowMapCacheManager.cpp b/Engine/Source/Runtime/Renderer/Private/VirtualShadowMaps/VirtualShadowMapCacheManager.cpp
--- a/Engine/Source/Runtime/Renderer/Private/VirtualShadowMaps/VirtualShadowMapCacheManager.cpp
+++ b/Engine/Source/Runtime/Renderer/Private/VirtualShadowMaps/VirtualShadowMapCacheManager.cpp
@@ -1033,12 +1033,6 @@ FScreenPassTexture FVirtualShadowMapArrayCacheManager::AddVisualizePass(FRDGBuil
 #if !UE_BUILD_SHIPPING
 	if (Pass == EVSMVisualizationPostPass::PostEditorPrimitives && (FVirtualShadowMap::GetEnabledStatSections() & EVSMStatSection::Clusters))
 	{
-		if (!MonospaceFont)
-		{
-			FTopLevelAssetPath MonospaceFontPath = FTopLevelAssetPath("/Engine/EngineFonts/DroidSansMono.DroidSansMono");
-			MonospaceFont = LoadObject<UFont>(nullptr, *MonospaceFontPath.ToString());
-		}
-		
 		AddDrawCanvasPass(GraphBuilder, RDG_EVENT_NAME("NPFDiagnostics"), View, Output,
 			[&](FCanvas& Canvas)
 			{
@@ -1078,7 +1072,7 @@ FScreenPassTexture FVirtualShadowMapArrayCacheManager::AddVisualizePass(FRDGBuil
 						MessageBoxCharacters.X = FMath::Max(MessageBoxCharacters.X, Message.Len());
 					}
 
-					FVector2D CharSize = MeasureString(TEXT("A"), MonospaceFont);
+					FVector2D CharSize = MeasureString(TEXT("A"), GEngine->GetMonospaceFont());
 					MessageBoxSize = CharSize * MessageBoxCharacters + FVector2D(10, 10);
 				}
 				
@@ -1091,7 +1085,7 @@ FScreenPassTexture FVirtualShadowMapArrayCacheManager::AddVisualizePass(FRDGBuil
 				int Y = 0;
 				for (const FString& Message : NPFDiagnosticMessages)
 				{
-					Y += DrawString(BaselineX, BaselineY + Y, Message, FLinearColor::White, MonospaceFont).Y;
+					Y += DrawString(BaselineX, BaselineY + Y, Message, FLinearColor::White, GEngine->GetMonospaceFont()).Y;
 				}
 			});
 		NPFDiagnosticMessages.SetNum(0, EAllowShrinking::No);
diff --git a/Engine/Source/Runtime/Renderer/Private/VirtualShadowMaps/VirtualShadowMapCacheManager.h b/Engine/Source/Runtime/Renderer/Private/VirtualShadowMaps/VirtualShadowMapCacheManager.h
--- a/Engine/Source/Runtime/Renderer/Private/VirtualShadowMaps/VirtualShadowMapCacheManager.h
+++ b/Engine/Source/Runtime/Renderer/Private/VirtualShadowMaps/VirtualShadowMapCacheManager.h
@@ -543,7 +543,6 @@ private:
 	TArray<FString> NPFDiagnosticMessages {};
 	TMap<uint32, uint8> NPFDiagnosticTimer {};
 
-	TObjectPtr<UFont> MonospaceFont = nullptr;
 #endif // UE_BUILD_SHIPPING
 
 	FShadowInvalidatingInstancesImplementation ShadowInvalidatingInstancesImplementation;
