Add PSO precaching for mobile deferred local lights

diff --git a/Engine/Source/Runtime/Renderer/Private/MobileDeferredShadingPass.cpp b/Engine/Source/Runtime/Renderer/Private/MobileDeferredShadingPass.cpp
--- a/Engine/Source/Runtime/Renderer/Private/MobileDeferredShadingPass.cpp
+++ b/Engine/Source/Runtime/Renderer/Private/MobileDeferredShadingPass.cpp
@@ -797,14 +797,14 @@ static int RenderDirectionalLights(
 	return NumLights;
 }
 
-template<uint32 LightingChannel, bool bWithStencilCulling, bool bEnableShadingModelSupport>
-static void SetLocalLightRasterizerAndDepthState(FGraphicsPipelineStateInitializer& GraphicsPSOInit, const FViewInfo& View, const FSphere& LightBounds)
+template<uint32 LightingChannel, bool bEnableShadingModelSupport>
+static void SetLocalLightRasterizerAndDepthState(FGraphicsPipelineStateInitializer& GraphicsPSOInit, bool bReverseCulling, bool bCameraInsideLightGeometry)
 {
-	if (bWithStencilCulling)
+	if (GMobileUseLightStencilCulling != 0)
 	{
 		// Render backfaces with depth and stencil tests
 		// and clear stencil to zero for next light mask
-		GraphicsPSOInit.RasterizerState = View.bReverseCulling ? TStaticRasterizerState<FM_Solid, CM_CW>::GetRHI() : TStaticRasterizerState<FM_Solid, CM_CCW>::GetRHI();
+		GraphicsPSOInit.RasterizerState = bReverseCulling ? TStaticRasterizerState<FM_Solid, CM_CW>::GetRHI() : TStaticRasterizerState<FM_Solid, CM_CCW>::GetRHI();
 		GraphicsPSOInit.DepthStencilState = TStaticDepthStencilState<
 			false, CF_LessEqual,
 			false, CF_Equal, SO_Keep, SO_Keep, SO_Keep,
@@ -815,16 +815,10 @@ static void SetLocalLightRasterizerAndDepthState(FGraphicsPipelineStateInitializ
 	}
 	else
 	{
-
-		const bool bCameraInsideLightGeometry = ((FVector)View.ViewMatrices.GetViewOrigin() - LightBounds.Center).SizeSquared() < FMath::Square(LightBounds.W * 1.05f + View.NearClippingDistance * 2.0f)
-			// Always draw backfaces in ortho
-			//@todo - accurate ortho camera / light intersection
-			|| !View.IsPerspectiveProjection();
-
 		if (bCameraInsideLightGeometry)
 		{
 			// Render backfaces with depth tests disabled since the camera is inside (or close to inside) the light geometry
-			GraphicsPSOInit.RasterizerState = View.bReverseCulling ? TStaticRasterizerState<FM_Solid, CM_CW>::GetRHI() : TStaticRasterizerState<FM_Solid, CM_CCW>::GetRHI();
+			GraphicsPSOInit.RasterizerState = bReverseCulling ? TStaticRasterizerState<FM_Solid, CM_CW>::GetRHI() : TStaticRasterizerState<FM_Solid, CM_CCW>::GetRHI();
 			GraphicsPSOInit.DepthStencilState = TStaticDepthStencilState<
 				false, CF_Always,
 				true, CF_Equal, SO_Keep, SO_Keep, SO_Keep,
@@ -834,7 +828,7 @@ static void SetLocalLightRasterizerAndDepthState(FGraphicsPipelineStateInitializ
 		else
 		{
 			// Render frontfaces with depth tests on to get the speedup from HiZ since the camera is outside the light geometry
-			GraphicsPSOInit.RasterizerState = View.bReverseCulling ? TStaticRasterizerState<FM_Solid, CM_CCW>::GetRHI() : TStaticRasterizerState<FM_Solid, CM_CW>::GetRHI();
+			GraphicsPSOInit.RasterizerState = bReverseCulling ? TStaticRasterizerState<FM_Solid, CM_CCW>::GetRHI() : TStaticRasterizerState<FM_Solid, CM_CW>::GetRHI();
 			GraphicsPSOInit.DepthStencilState = TStaticDepthStencilState<
 				false, CF_DepthNearOrEqual,
 				true, CF_Equal, SO_Keep, SO_Keep, SO_Keep,
@@ -844,46 +838,52 @@ static void SetLocalLightRasterizerAndDepthState(FGraphicsPipelineStateInitializ
 	}
 }
 
-template <bool bEnableShadingModelSupport, bool bWithStencilCulling>
-static void SetLocalLightRasterizerAndDepthState(FGraphicsPipelineStateInitializer& GraphicsPSOInit, const FViewInfo& View, const FSphere& LightBounds, uint32 LightingChannel)
+template <bool bEnableShadingModelSupport>
+static void SetLocalLightRasterizerAndDepthState(FGraphicsPipelineStateInitializer& GraphicsPSOInit, bool bReverseCulling, bool bCameraInsideLightGeometry, uint32 LightingChannel)
 {
 	// TODO: support multi-channel ligths?
 	switch (LightingChannel)
 	{
 	default:
-		SetLocalLightRasterizerAndDepthState<0, bWithStencilCulling, bEnableShadingModelSupport>(GraphicsPSOInit, View, LightBounds);
+		SetLocalLightRasterizerAndDepthState<0, bEnableShadingModelSupport>(GraphicsPSOInit, bReverseCulling, bCameraInsideLightGeometry);
 		break;
 	case 1:
-		SetLocalLightRasterizerAndDepthState<1, bWithStencilCulling, bEnableShadingModelSupport>(GraphicsPSOInit, View, LightBounds);
+		SetLocalLightRasterizerAndDepthState<1, bEnableShadingModelSupport>(GraphicsPSOInit, bReverseCulling, bCameraInsideLightGeometry);
 		break;
 	case 2:
-		SetLocalLightRasterizerAndDepthState<2, bWithStencilCulling, bEnableShadingModelSupport>(GraphicsPSOInit, View, LightBounds);
+		SetLocalLightRasterizerAndDepthState<2, bEnableShadingModelSupport>(GraphicsPSOInit, bReverseCulling, bCameraInsideLightGeometry);
 		break;
 	}
 }
 
-template <bool bEnableShadingModelSupport>
-static void SetLocalLightRasterizerAndDepthState(FGraphicsPipelineStateInitializer& GraphicsPSOInit, const FViewInfo& View, const FSphere& LightBounds, uint32 LightingChannel)
+static void SetLocalLightRasterizerAndDepthState(FGraphicsPipelineStateInitializer& GraphicsPSOInit, bool bReverseCulling, bool bCameraInsideLightGeometry, uint32 LightingChannel, bool bEnableShadingModelSupport)
 {
-	if (GMobileUseLightStencilCulling != 0)
+	if (bEnableShadingModelSupport)
 	{
-		SetLocalLightRasterizerAndDepthState<bEnableShadingModelSupport, true>(GraphicsPSOInit, View, LightBounds, LightingChannel);
+		SetLocalLightRasterizerAndDepthState<true>(GraphicsPSOInit, bReverseCulling, bCameraInsideLightGeometry, LightingChannel);
 	}
 	else
 	{
-		SetLocalLightRasterizerAndDepthState<bEnableShadingModelSupport, false>(GraphicsPSOInit, View, LightBounds, LightingChannel);
+		SetLocalLightRasterizerAndDepthState<false>(GraphicsPSOInit, bReverseCulling, bCameraInsideLightGeometry, LightingChannel);
 	}
 }
 
 static void SetLocalLightRasterizerAndDepthState(FGraphicsPipelineStateInitializer& GraphicsPSOInit, const FViewInfo& View, const FSphere& LightBounds, uint32 LightingChannel, bool bEnableShadingModelSupport)
 {
+	const bool bCameraInsideLightGeometry =
+		(GMobileUseLightStencilCulling == 0) &&
+		(((FVector)View.ViewMatrices.GetViewOrigin() - LightBounds.Center).SizeSquared() < FMath::Square(LightBounds.W * 1.05f + View.NearClippingDistance * 2.0f)
+			// Always draw backfaces in ortho
+			//@todo - accurate ortho camera / light intersection
+			|| !View.IsPerspectiveProjection());
+
 	if (bEnableShadingModelSupport)
 	{
-		SetLocalLightRasterizerAndDepthState<true>(GraphicsPSOInit, View, LightBounds, LightingChannel);
+		SetLocalLightRasterizerAndDepthState<true>(GraphicsPSOInit, View.bReverseCulling, bCameraInsideLightGeometry, LightingChannel);
 	}
 	else
 	{
-		SetLocalLightRasterizerAndDepthState<false>(GraphicsPSOInit, View, LightBounds, LightingChannel);
+		SetLocalLightRasterizerAndDepthState<false>(GraphicsPSOInit, View.bReverseCulling, bCameraInsideLightGeometry, LightingChannel);
 	}
 }
 
@@ -1179,6 +1179,16 @@ public:
 		const FPSOPrecacheParams& PreCacheParams,
 		TArray<FPSOPrecacheData>& PSOInitializers
 	) override final;
+
+private:
+	void CollectPSOInitializersDirectional(
+		const FSceneTexturesConfig& SceneTexturesConfig,
+		const FMaterial& Material,
+		TArray<FPSOPrecacheData>& PSOInitializers);
+	void CollectPSOInitializersLocal(
+		const FSceneTexturesConfig& SceneTexturesConfig,
+		const FMaterial& Material,
+		TArray<FPSOPrecacheData>& PSOInitializers);
 };
 
 void FDeferredMobileLightMaterialPSOCollector::CollectPSOInitializers(
@@ -1186,8 +1196,16 @@ void FDeferredMobileLightMaterialPSOCollector::CollectPSOInitializers(
 	const FMaterial& Material,
 	const FPSOPrecacheVertexFactoryData& VertexFactoryData,
 	const FPSOPrecacheParams& PreCacheParams,
-	TArray<FPSOPrecacheData>& PSOInitializers
-)
+	TArray<FPSOPrecacheData>& PSOInitializers)
+{
+	CollectPSOInitializersDirectional(SceneTexturesConfig, Material, PSOInitializers);
+	CollectPSOInitializersLocal(SceneTexturesConfig, Material, PSOInitializers);
+}
+
+void FDeferredMobileLightMaterialPSOCollector::CollectPSOInitializersDirectional(
+	const FSceneTexturesConfig& SceneTexturesConfig,
+	const FMaterial& Material,
+	TArray<FPSOPrecacheData>& PSOInitializers)
 {
 	FMaterialShaderTypes ShaderTypesToGetAnyPermutation;
 	ShaderTypesToGetAnyPermutation.AddShaderType<FMobileDirectionalLightFunctionPS>();
@@ -1307,6 +1325,105 @@ void FDeferredMobileLightMaterialPSOCollector::CollectPSOInitializers(
 	AddPSOInitializer(0, true, true, 0, false);
 }
 
+void FDeferredMobileLightMaterialPSOCollector::CollectPSOInitializersLocal(
+	const FSceneTexturesConfig& SceneTexturesConfig,
+	const FMaterial& Material,
+	TArray<FPSOPrecacheData>& PSOInitializers)
+{
+	FMaterialShaderTypes ShaderTypesToGetAnyPermutation;
+	ShaderTypesToGetAnyPermutation.AddShaderType<FMobileRadialLightFunctionPS>();
+
+	FMaterialShaders ShadersAnyPermutation;
+	if (!Material.TryGetShaders(ShaderTypesToGetAnyPermutation, nullptr, ShadersAnyPermutation))
+	{
+		return;
+	}
+
+	auto AddPSOInitializer =
+		[&](int32 PassIndex, uint8 LightType, bool bUseIESTexture, bool bCameraInsideLightGeometry)
+		{
+			const bool bEnableShadingModelSupport = (PassIndex > 0);
+			const bool bShouldCastShadow = false;
+			const int32 LightingChannel = 0;
+			const bool bReverseCulling = false;
+
+			FMobileRadialLightFunctionPS::FPermutationDomain PermutationVector;
+			PermutationVector.Set<FMobileRadialLightFunctionPS::FEnableShadingModelSupport>(bEnableShadingModelSupport);
+			PermutationVector.Set<FMobileRadialLightFunctionPS::FRadialLightTypeDim>(LightType);
+			PermutationVector.Set<FMobileRadialLightFunctionPS::FIESProfileDim>(bUseIESTexture);
+			PermutationVector.Set<FMobileRadialLightFunctionPS::FSpotLightShadowDim>(bShouldCastShadow);
+			FMaterialShaderTypes ShaderTypesToGet;
+			ShaderTypesToGet.AddShaderType<FMobileRadialLightFunctionPS>(PermutationVector.ToDimensionValueId());
+
+			FMaterialShaders Shaders;
+			if (!Material.TryGetShaders(ShaderTypesToGet, nullptr, Shaders))
+			{
+				return;
+			}
+
+			TShaderRef<FMobileRadialLightFunctionPS> PixelShader;
+			Shaders.TryGetPixelShader(PixelShader);
+			if (!PixelShader.IsValid())
+			{
+				return;
+			}
+
+			const bool bRequired = false;
+			FRHIPixelShader* RHIPixelShader = PixelShader.GetPixelShader(bRequired);
+			if (RHIPixelShader == nullptr)
+			{
+				return;
+			}
+
+			FDeferredLightVS::FPermutationDomain PermutationVectorVS;
+			PermutationVectorVS.Set<FDeferredLightVS::FRadialLight>(true);
+			TShaderMapRef<FDeferredLightVS> VertexShader(GetGlobalShaderMap(GMaxRHIShaderPlatform), PermutationVectorVS);
+			FRHIVertexShader* RHIVertexShader = VertexShader.GetVertexShader(bRequired);
+			if (RHIVertexShader == nullptr)
+			{
+				return;
+			}
+
+			FGraphicsPipelineStateInitializer GraphicsPSOInit;
+			GraphicsPSOInit.RasterizerState = TStaticRasterizerState<>::GetRHI();
+			GraphicsPSOInit.BlendState = TStaticBlendState<CW_RGB, BO_Add, BF_One, BF_One, BO_Add, BF_One, BF_One>::GetRHI();
+			SetLocalLightRasterizerAndDepthState(GraphicsPSOInit, bReverseCulling, bCameraInsideLightGeometry, LightingChannel, bEnableShadingModelSupport);
+
+			GraphicsPSOInit.BoundShaderState.VertexDeclarationRHI = GFilterVertexDeclaration.VertexDeclarationRHI;
+			GraphicsPSOInit.BoundShaderState.VertexShaderRHI = RHIVertexShader;
+			GraphicsPSOInit.BoundShaderState.PixelShaderRHI = RHIPixelShader;
+			GraphicsPSOInit.PrimitiveType = PT_TriangleList;
+
+			FGraphicsPipelineRenderTargetsInfo RenderTargetsInfo;
+			RenderTargetsInfo.NumSamples = 1;
+			AddRenderTargetInfo(SceneTexturesConfig.ColorFormat, SceneTexturesConfig.ColorCreateFlags, RenderTargetsInfo);
+
+			GraphicsPSOInit.StatePrecachePSOHash = RHIComputeStatePrecachePSOHash(GraphicsPSOInit);
+			ApplyTargetsInfo(GraphicsPSOInit, RenderTargetsInfo);
+
+			FPSOPrecacheData PSOPrecacheData;
+			PSOPrecacheData.bRequired = true;
+			PSOPrecacheData.Type = FPSOPrecacheData::EType::Graphics;
+			PSOPrecacheData.GraphicsPSOInitializer = GraphicsPSOInit;
+#if PSO_PRECACHING_VALIDATE
+			PSOPrecacheData.PSOCollectorIndex = PSOCollectorIndex;
+			PSOPrecacheData.VertexFactoryType = nullptr;
+#endif // PSO_PRECACHING_VALIDATE
+
+			PSOInitializers.Add(MoveTemp(PSOPrecacheData));
+		};
+
+	// int32 PassIndex, uint8 LightType, bool bUseIESTexture, bool bCameraInsideLightGeometry
+	AddPSOInitializer(0, LIGHT_TYPE_POINT, false, false);
+	AddPSOInitializer(1, LIGHT_TYPE_POINT, false, false);
+	AddPSOInitializer(0, LIGHT_TYPE_POINT, true, false);
+	AddPSOInitializer(1, LIGHT_TYPE_POINT, true, false);
+	AddPSOInitializer(0, LIGHT_TYPE_POINT, true, true);
+	AddPSOInitializer(1, LIGHT_TYPE_POINT, true, true);
+	AddPSOInitializer(0, LIGHT_TYPE_SPOT, false, false);
+	AddPSOInitializer(1, LIGHT_TYPE_SPOT, false, false);
+}
+
 IPSOCollector* CreateDeferredMobileLightMaterialPSOCollector(ERHIFeatureLevel::Type FeatureLevel)
 {
 	return new FDeferredMobileLightMaterialPSOCollector(FeatureLevel);
