Occulsion Query Helpers correct number of frames for Metal mobile platforms.

diff --git a/Engine/Source/Runtime/Renderer/Private/SceneOcclusion.cpp b/Engine/Source/Runtime/Renderer/Private/SceneOcclusion.cpp
--- a/Engine/Source/Runtime/Renderer/Private/SceneOcclusion.cpp
+++ b/Engine/Source/Runtime/Renderer/Private/SceneOcclusion.cpp
@@ -112,12 +112,12 @@ int32 FOcclusionQueryHelpers::GetNumBufferedFrames(ERHIFeatureLevel::Type Featur
 	EShaderPlatform ShaderPlatform = GShaderPlatformForFeatureLevel[FeatureLevel];
 
 	int32 NumExtraMobileFrames = 0;
-	if ((FeatureLevel <= ERHIFeatureLevel::ES3_1 || IsVulkanMobileSM5Platform(ShaderPlatform)) && CVarMobileEnableOcclusionExtraFrame.GetValueOnAnyThread())
+	if ((FeatureLevel <= ERHIFeatureLevel::ES3_1 || IsVulkanMobileSM5Platform(ShaderPlatform) || IsMetalMobileSM5Platform(ShaderPlatform)) && CVarMobileEnableOcclusionExtraFrame.GetValueOnAnyThread())
 	{
 		NumExtraMobileFrames++; // the mobile renderer just doesn't do much after the basepass, and hence it will be asking for the query results almost immediately; the results can't possibly be ready in 1 frame.
 		
-		bool bNeedsAnotherExtraMobileFrame = IsVulkanPlatform(ShaderPlatform) || IsOpenGLPlatform(ShaderPlatform);
-		bNeedsAnotherExtraMobileFrame = bNeedsAnotherExtraMobileFrame || IsVulkanMobileSM5Platform(ShaderPlatform);
+		bool bNeedsAnotherExtraMobileFrame = IsVulkanPlatform(ShaderPlatform) || IsMetalPlatform(ShaderPlatform) || IsOpenGLPlatform(ShaderPlatform);
+		bNeedsAnotherExtraMobileFrame = bNeedsAnotherExtraMobileFrame || IsVulkanMobileSM5Platform(ShaderPlatform) || IsMetalMobileSM5Platform(ShaderPlatform);
 		bNeedsAnotherExtraMobileFrame = bNeedsAnotherExtraMobileFrame || FDataDrivenShaderPlatformInfo::GetNeedsExtraMobileFrames(ShaderPlatform);
 		bNeedsAnotherExtraMobileFrame = bNeedsAnotherExtraMobileFrame && IsRunningRHIInSeparateThread();
 
