Temporarily disable ensure in RayTracing.cpp until UE-112448 is fixed.

diff --git a/Engine/Source/Runtime/Renderer/Private/RayTracing/RayTracing.cpp b/Engine/Source/Runtime/Renderer/Private/RayTracing/RayTracing.cpp
--- a/Engine/Source/Runtime/Renderer/Private/RayTracing/RayTracing.cpp
+++ b/Engine/Source/Runtime/Renderer/Private/RayTracing/RayTracing.cpp
@@ -1328,7 +1328,9 @@ namespace RayTracing
 					if (bUsingReferenceBasedResidency)
 					{
 						FRayTracingGeometry* TargetRayTracingGeometry = SceneInfo->GetStaticRayTracingGeometry(LODIndex);
-						if (!ensure(TargetRayTracingGeometry != nullptr))
+
+						// TODO: Should have an assert here but disabled it due to UE-112448
+						if (TargetRayTracingGeometry == nullptr)
 						{
 							return;
 						}
