[SingleLayerWater] Fixed race condition when reading the view's InstanceFactor in WaterMeshSceneProxy::GDME() by moving InstanceFactor initializiation before kicking off GDME.

diff --git a/Engine/Source/Runtime/Renderer/Private/MobileShadingRenderer.cpp b/Engine/Source/Runtime/Renderer/Private/MobileShadingRenderer.cpp
--- a/Engine/Source/Runtime/Renderer/Private/MobileShadingRenderer.cpp
+++ b/Engine/Source/Runtime/Renderer/Private/MobileShadingRenderer.cpp
@@ -485,14 +485,6 @@ void FMobileSceneRenderer::InitViews(
 		InstanceCullingManager.InstanceOcclusionQueryBufferFormat = Scene->InstanceCullingOcclusionQueryRenderer->InstanceOcclusionQueryBufferFormat;
 	}
 
-	for (FViewInfo& ViewInfo : Views)
-	{
-		uint32 InstanceFactor = ViewInfo.bIsInstancedStereoEnabled && IStereoRendering::IsStereoEyeView(ViewInfo) && GEngine->StereoRenderingDevice.IsValid() ?
-			GEngine->StereoRenderingDevice->GetDesiredNumberOfViews(true) : 1;
-
-		ViewInfo.InstanceFactor = InstanceFactor > 0 ? InstanceFactor : 1;
-	}
-
 	FILCUpdatePrimTaskData* ILCTaskData = nullptr;
 
 	const FExclusiveDepthStencil::Type BasePassDepthStencilAccess = FExclusiveDepthStencil::DepthWrite_StencilWrite;
diff --git a/Engine/Source/Runtime/Renderer/Private/SceneVisibility.cpp b/Engine/Source/Runtime/Renderer/Private/SceneVisibility.cpp
--- a/Engine/Source/Runtime/Renderer/Private/SceneVisibility.cpp
+++ b/Engine/Source/Runtime/Renderer/Private/SceneVisibility.cpp
@@ -4907,11 +4907,24 @@ static bool IsLargeCameraMovement(FSceneView& View, const FMatrix& PrevViewMatri
 void FSceneRenderer::PreVisibilityFrameSetup(FRDGBuilder& GraphBuilder)
 {
 	// Possible stencil dither optimization approach
-	for (int32 ViewIndex = 0; ViewIndex < Views.Num(); ViewIndex++)
+	for (FViewInfo& View : Views)
 	{
-		FViewInfo& View = Views[ViewIndex];
 		View.bAllowStencilDither = DepthPass.bDitheredLODTransitionsUseStencil;
 	}
+
+	// Initialize InstanceFactor. AllViews includes CustomRenderPass views, which might possibly also require InstanceFactor to be correctly assigned.
+	const int32 InstanceFactor = FMath::Max(1, GEngine->StereoRenderingDevice.IsValid() ? GEngine->StereoRenderingDevice->GetDesiredNumberOfViews(true) : 1);
+	if (InstanceFactor > 1)
+	{
+		for (FViewInfo* View : AllViews)
+		{
+			if (View->bIsInstancedStereoEnabled && IStereoRendering::IsStereoEyeView(*View))
+			{
+				View->InstanceFactor = InstanceFactor;
+			}
+		}
+	}
+	
 	FRHICommandListImmediate& RHICmdList = GraphBuilder.RHICmdList;
 	
 	if (GetRendererOutput() == ERendererOutput::FinalSceneColor)
@@ -5863,15 +5876,6 @@ void FDeferredShadingSceneRenderer::BeginInitViews(
 		InstanceCullingManager.InstanceOcclusionQueryBufferFormat = Scene->InstanceCullingOcclusionQueryRenderer->InstanceOcclusionQueryBufferFormat;
 	}
 
-	// Create GPU-side representation of the view for instance culling.
-	for (FViewInfo* ViewInfo : AllViews)
-	{
-		uint32 InstanceFactor = ViewInfo->bIsInstancedStereoEnabled && IStereoRendering::IsStereoEyeView(*ViewInfo) && GEngine->StereoRenderingDevice.IsValid() ?
-			GEngine->StereoRenderingDevice->GetDesiredNumberOfViews(true) : 1;
-
-		ViewInfo->InstanceFactor = InstanceFactor > 0 ? InstanceFactor : 1;
-	}
-
 	LumenScenePDIVisualization();
 
 	// This must happen before we start initialising and using views.
