Fix broken receiver mask sampling on platforms with gather_lod.

#rnx

diff --git a/Engine/Source/Runtime/Renderer/Private/VirtualShadowMaps/VirtualShadowMapArray.cpp b/Engine/Source/Runtime/Renderer/Private/VirtualShadowMaps/VirtualShadowMapArray.cpp
--- a/Engine/Source/Runtime/Renderer/Private/VirtualShadowMaps/VirtualShadowMapArray.cpp
+++ b/Engine/Source/Runtime/Renderer/Private/VirtualShadowMaps/VirtualShadowMapArray.cpp
@@ -2236,7 +2236,6 @@ void FVirtualShadowMapArray::BuildPageAllocations(
 		if (IsVirtualShadowMapLocalReceiverMaskEnabled())
 		{
 			PageReceiverMasksRDG = AllocatePageTable(true, PF_R32_UINT, TEXT("Shadow.Virtual.PageReceiverMasks"), 2u);
-			UniformParameters.PageReceiverMaskTextureSizeInvSize = UniformParameters.PageTableTextureSizeInvSize;
 			// TODO: If we don't actually have any local lights perhaps worth doing the simpler ClearUAVPass below instead?
 			AddClearPageTableUAVPass(GraphBuilder, PageReceiverMasksRDG, 0u, 2u);
 		}
@@ -2248,7 +2247,6 @@ void FVirtualShadowMapArray::BuildPageAllocations(
 			const uint32 NumPageReceiverMaskRows = FMath::DivideAndRoundUp((uint32)NumDirectionalPageReceiverMaskEntriesRequired, NumPageTablesEntriesPerRow);
 			// Each page stores 2x2 masks, each representing a 4x4 tile
 			FIntPoint PageReceiverMaskTextureSize(2U * NumPageTablesEntriesPerRow * VSM_PAGE_TABLE_TEX2D_SIZE_X, 2U * NumPageReceiverMaskRows * VSM_PAGE_TABLE_TEX2D_SIZE_Y);
-			UniformParameters.PageReceiverMaskTextureSizeInvSize = FVector4f(FVector2f(PageReceiverMaskTextureSize), FVector2f(1.0f, 1.0f) / FVector2f(PageReceiverMaskTextureSize));
 
 			// TODO: should be uint16 - but atomics prevent this
 			PageReceiverMasksRDG = GraphBuilder.CreateTexture(
@@ -2263,10 +2261,14 @@ void FVirtualShadowMapArray::BuildPageAllocations(
 				TEXT("Shadow.Virtual.PageReceiverMasks"));
 			AddClearUAVPass(GraphBuilder, GraphBuilder.CreateUAV(PageReceiverMasksRDG), FUintVector4(0, 0, 0, 0));
 		}
+
+		const FIntPoint TextureSize = PageReceiverMasksRDG->Desc.Extent;
+		UniformParameters.PageReceiverMaskTextureSizeInvSize = FVector4f(FVector2f(TextureSize), FVector2f(1.0f, 1.0f) / FVector2f(TextureSize));
 	}
 	else
 	{
 		PageReceiverMasksRDG = GetOrCreatePageTableDummy();
+		UniformParameters.PageReceiverMaskTextureSizeInvSize = FVector4f(1, 1, 1, 1);
 	}
 
 	const uint32 DirtyFlagsPerPageCount = 4;
