Substrate - fixed missing anisotropy when using adaptative gbuffer.

#rnx
#tests pc editor dx12

diff --git a/Engine/Source/Runtime/Core/Private/UObject/FortniteMainBranchObjectVersion.cpp b/Engine/Source/Runtime/Core/Private/UObject/FortniteMainBranchObjectVersion.cpp
--- a/Engine/Source/Runtime/Core/Private/UObject/FortniteMainBranchObjectVersion.cpp
+++ b/Engine/Source/Runtime/Core/Private/UObject/FortniteMainBranchObjectVersion.cpp
@@ -6,7 +6,7 @@ TMap<FGuid, FGuid> FFortniteMainBranchObjectVersion::GetSystemGuids()
 	TMap<FGuid, FGuid> SystemGuids;
 	const FDevSystemGuids& DevGuids = FDevSystemGuids::Get();
 
-	SystemGuids.Add(DevGuids.GLOBALSHADERMAP_DERIVEDDATA_VER, FGuid("125CF21C881744418BD4A614B3B125FF"));
+	SystemGuids.Add(DevGuids.GLOBALSHADERMAP_DERIVEDDATA_VER, FGuid("210D8D0B3C7540799E2D561F27190837"));
 	SystemGuids.Add(DevGuids.LANDSCAPE_MOBILE_COOK_VERSION, FGuid("32D02EF867C74B71A0D4E0FA41392732"));
 	SystemGuids.Add(DevGuids.MATERIALSHADERMAP_DERIVEDDATA_VER, FGuid("F051877974B44D409435F321BDB39DBB"));
 	SystemGuids.Add(DevGuids.NANITE_DERIVEDDATA_VER, FGuid("B169793F25FF4493A55451FDC02A4AC1"));
diff --git a/Engine/Source/Runtime/Renderer/Private/AnisotropyRendering.cpp b/Engine/Source/Runtime/Renderer/Private/AnisotropyRendering.cpp
--- a/Engine/Source/Runtime/Renderer/Private/AnisotropyRendering.cpp
+++ b/Engine/Source/Runtime/Renderer/Private/AnisotropyRendering.cpp
@@ -20,11 +20,16 @@ static FAutoConsoleVariableRef CVarAnisotropicMaterials(
 
 bool SupportsAnisotropicMaterials(ERHIFeatureLevel::Type FeatureLevel, EShaderPlatform ShaderPlatform)
 {
-	const bool bSubstrateEnabled = Substrate::IsSubstrateEnabled();
 	return GAnisotropicMaterials
 		&& FeatureLevel >= ERHIFeatureLevel::SM5
-		&& FDataDrivenShaderPlatformInfo::GetSupportsAnisotropicMaterials(ShaderPlatform)
-		&& (!bSubstrateEnabled || (bSubstrateEnabled && Substrate::IsSubstrateBlendableGBufferEnabled(ShaderPlatform))); // Substrate renders anisotropy surface natively, without extra pass. Unless blendable gbuffer is used.
+		&& FDataDrivenShaderPlatformInfo::GetSupportsAnisotropicMaterials(ShaderPlatform);
+}
+
+bool SupportsAnisotropicMaterialPass(ERHIFeatureLevel::Type FeatureLevel, EShaderPlatform ShaderPlatform)
+{
+		const bool bSubstrateEnabled = Substrate::IsSubstrateEnabled();
+		return SupportsAnisotropicMaterials(FeatureLevel, ShaderPlatform)
+			&& (!bSubstrateEnabled || (bSubstrateEnabled && Substrate::IsSubstrateBlendableGBufferEnabled(ShaderPlatform))); // Substrate renders anisotropy surface natively, without extra pass. Unless blendable gbuffer is used.
 }
 
 static bool IsAnisotropyPassCompatible(const EShaderPlatform Platform, FMaterialShaderParameters MaterialParameters)
@@ -150,7 +155,7 @@ void FAnisotropyMeshProcessor::AddMeshBatch(
 	int32 StaticMeshId /* = -1 */ 
 	)
 {
-	if (SupportsAnisotropicMaterials(FeatureLevel, GShaderPlatformForFeatureLevel[FeatureLevel]) && MeshBatch.bUseForMaterial)
+	if (SupportsAnisotropicMaterialPass(FeatureLevel, GShaderPlatformForFeatureLevel[FeatureLevel]) && MeshBatch.bUseForMaterial)
 	{
 		const FMaterialRenderProxy* MaterialRenderProxy = MeshBatch.MaterialRenderProxy;
 		while (MaterialRenderProxy)
@@ -244,7 +249,7 @@ bool FAnisotropyMeshProcessor::Process(
 void FAnisotropyMeshProcessor::CollectPSOInitializers(const FSceneTexturesConfig& SceneTexturesConfig, const FMaterial& Material, const FPSOPrecacheVertexFactoryData& VertexFactoryData, const FPSOPrecacheParams& PreCacheParams, TArray<FPSOPrecacheData>& PSOInitializers)
 {
 	if (ShouldDraw(Material, Material.MaterialUsesAnisotropy_GameThread()) && 
-		SupportsAnisotropicMaterials(FeatureLevel, GShaderPlatformForFeatureLevel[FeatureLevel]))
+		SupportsAnisotropicMaterialPass(FeatureLevel, GShaderPlatformForFeatureLevel[FeatureLevel]))
 	{
 		const FMeshDrawingPolicyOverrideSettings OverrideSettings = ComputeMeshOverrideSettings(PreCacheParams);
 		const ERasterizerFillMode MeshFillMode = ComputeMeshFillMode(Material, OverrideSettings);
@@ -288,13 +293,29 @@ void FAnisotropyMeshProcessor::CollectPSOInitializers(const FSceneTexturesConfig
 	}
 }
 
-bool ShouldRenderAnisotropyPass(const FViewInfo& View)
+bool ShouldRenderAnisotropicLighting(const FViewInfo& View)
 {
 	if (!SupportsAnisotropicMaterials(View.FeatureLevel, View.GetShaderPlatform()))
 	{
 		return false;
 	}
 
+	if (SupportsAnisotropicMaterialPass(View.FeatureLevel, View.GetShaderPlatform()) && View.ShouldRenderView() && !HasAnyDraw(View.ParallelMeshDrawCommandPasses[EMeshPass::AnisotropyPass]))
+	{
+		// The setup requires a anisotropy pass for material data but nothing is rendered in it: then skip rendering anisotropy.
+		return false;
+	}
+
+	return true;
+}
+
+bool ShouldRenderAnisotropyPass(const FViewInfo& View)
+{
+	if (!SupportsAnisotropicMaterialPass(View.FeatureLevel, View.GetShaderPlatform()))
+	{
+		return false;
+	}
+
 	if (IsForwardShadingEnabled(GetFeatureLevelShaderPlatform(View.FeatureLevel)))
 	{
 		return false;
diff --git a/Engine/Source/Runtime/Renderer/Private/AnisotropyRendering.h b/Engine/Source/Runtime/Renderer/Private/AnisotropyRendering.h
--- a/Engine/Source/Runtime/Renderer/Private/AnisotropyRendering.h
+++ b/Engine/Source/Runtime/Renderer/Private/AnisotropyRendering.h
@@ -57,5 +57,6 @@ protected:
 		);
 };
 
-bool ShouldRenderAnisotropyPass(const FViewInfo& Views);
+bool ShouldRenderAnisotropicLighting(const FViewInfo& View);
+bool ShouldRenderAnisotropyPass(const FViewInfo& View);
 bool ShouldRenderAnisotropyPass(TArrayView<FViewInfo> Views);
\ No newline at end of file
diff --git a/Engine/Source/Runtime/Renderer/Private/LightRendering.cpp b/Engine/Source/Runtime/Renderer/Private/LightRendering.cpp
--- a/Engine/Source/Runtime/Renderer/Private/LightRendering.cpp
+++ b/Engine/Source/Runtime/Renderer/Private/LightRendering.cpp
@@ -998,11 +998,6 @@ class FDeferredLightPS : public FGlobalShader
 
 		if (PermutationVector.Get<FDeferredLightPS::FAnistropicMaterials>())
 		{
-			if (Substrate::IsSubstrateEnabled())
-			{
-				return IsFeatureLevelSupported(Parameters.Platform, ERHIFeatureLevel::SM5) && Substrate::IsSubstrateBlendableGBufferEnabled(Parameters.Platform);
-			}
-
 			// Anisotropic materials do not currently support rect lights
 			if (PermutationVector.Get<FSourceShapeDim>() == ELightSourceShape::Rect || PermutationVector.Get<FSourceTextureDim>())
 			{
@@ -2765,11 +2760,7 @@ static void RenderLight(
 	const bool bUseVirtualShadowMapMask = VirtualShadowMapId != INDEX_NONE && ShadowMaskBits;
 	const bool bNeedComplexTransmittanceSupport = View.HairCardsMeshElements.Num() && IsHairStrandsSupported(EHairStrandsShaderType::All, View.GetShaderPlatform());
 
-	bool bSupportAnisotropyPermutation = ShouldRenderAnisotropyPass(View);
-	if (Substrate::IsSubstrateEnabled())
-	{
-		bSupportAnisotropyPermutation = bSupportAnisotropyPermutation && Substrate::IsSubstrateBlendableGBufferEnabled(View.GetShaderPlatform());
-	}
+	const bool bSupportAnisotropyPermutation = ShouldRenderAnisotropicLighting(View);
 
 	check(!bUseVirtualShadowMapMask || bIsRadial);		// VSM mask only stores local lights
 
