Substrate - deferred lighting global shader permutation reduction.
It was noted as 2200 but in fact it was 1624.
After making the anisotropy perrmtutaion to be true only for Complex and ComplexSpecial tile and false for Simple and Single: 844 permutation.
Now the new top is FScreenProbeIntegrateCS with 1024 permutations.
Material now respects MATERIAL_USES_ANISOTROPY to not request complex tiles if not supported by the project and platform.

#rnx
#tests PC editor dx12

diff --git a/Engine/Shaders/Private/Substrate/Substrate.ush b/Engine/Shaders/Private/Substrate/Substrate.ush
--- a/Engine/Shaders/Private/Substrate/Substrate.ush
+++ b/Engine/Shaders/Private/Substrate/Substrate.ush
@@ -1599,8 +1599,13 @@ FSubstrateData InternalGetSubstrateSlabBSDF(
 	// When Glints are enabled, we need to not keep haziness even R0=R1 to avoid pops (second lobe disapearing, more energy from glints).
 	const bool bHasGlint = SUBSTRATE_COMPLEXSPECIALPATH && GlintValue < 1.0f;
 	const uint bHasHaziness = (((Haziness.Roughness != Roughness || bHasGlint) && Haziness.Weight > 0.0f) || Haziness.bSimpleClearCoat) ? 1 : 0;
-	const uint bHasAnisotropy = Anisotropy != 0 ? 1 : 0;
 	const uint bHasFuzz = FuzzAmount > 0.0f ? 1 : 0;
+#if MATERIAL_USES_ANISOTROPY
+	const uint bHasAnisotropy = Anisotropy != 0 ? 1 : 0;
+#else
+	Anisotropy = 0.0f;
+	const uint bHasAnisotropy = false;
+#endif
 	if (!SUBSTRATE_FASTPATH && bHasAnisotropy)
 	{
 		SubstrateRequestSharedLocalBasisTangent(SharedLocalBasisTypes, SharedLocalBasisIndex);
diff --git a/Engine/Source/Runtime/Core/Private/UObject/FortniteMainBranchObjectVersion.cpp b/Engine/Source/Runtime/Core/Private/UObject/FortniteMainBranchObjectVersion.cpp
--- a/Engine/Source/Runtime/Core/Private/UObject/FortniteMainBranchObjectVersion.cpp
+++ b/Engine/Source/Runtime/Core/Private/UObject/FortniteMainBranchObjectVersion.cpp
@@ -6,7 +6,7 @@ TMap<FGuid, FGuid> FFortniteMainBranchObjectVersion::GetSystemGuids()
 	TMap<FGuid, FGuid> SystemGuids;
 	const FDevSystemGuids& DevGuids = FDevSystemGuids::Get();
 
-	SystemGuids.Add(DevGuids.GLOBALSHADERMAP_DERIVEDDATA_VER, FGuid("210D8D0B3C7540799E2D561F27190837"));
+	SystemGuids.Add(DevGuids.GLOBALSHADERMAP_DERIVEDDATA_VER, FGuid("2ACD1C7798984B98BECABF4D32E4AA20"));
 	SystemGuids.Add(DevGuids.LANDSCAPE_MOBILE_COOK_VERSION, FGuid("32D02EF867C74B71A0D4E0FA41392732"));
 	SystemGuids.Add(DevGuids.MATERIALSHADERMAP_DERIVEDDATA_VER, FGuid("F051877974B44D409435F321BDB39DBB"));
 	SystemGuids.Add(DevGuids.NANITE_DERIVEDDATA_VER, FGuid("B169793F25FF4493A55451FDC02A4AC1"));
diff --git a/Engine/Source/Runtime/Engine/Private/ShaderCompiler/ShaderCompiler.cpp b/Engine/Source/Runtime/Engine/Private/ShaderCompiler/ShaderCompiler.cpp
--- a/Engine/Source/Runtime/Engine/Private/ShaderCompiler/ShaderCompiler.cpp
+++ b/Engine/Source/Runtime/Engine/Private/ShaderCompiler/ShaderCompiler.cpp
@@ -4760,8 +4760,8 @@ void VerifyGlobalShaders(EShaderPlatform Platform, const ITargetPlatform* Target
 		int32 PermutationCountLimit = 832;	// Nanite culling as of today (2022-01-11) can go up to 832 permutations
 		if (Substrate::IsSubstrateEnabled())
 		{
-			// SUBSTRATE_TODO reduce the number of permutation of FDeferredLightPS.
-			PermutationCountLimit = 2200;	// FDeferredLightPS as of today (2025-06-13)
+			// FDeferredLightPS as of 2025-08-08 has 844 permutations.
+			PermutationCountLimit = 1024;	// FScreenProbeIntegrateCS as of 2025-08-08 has 1024 permutations.
 		}
 		ensureMsgf(
 			PermutationCountToCompile <= PermutationCountLimit,
diff --git a/Engine/Source/Runtime/Renderer/Private/LightRendering.cpp b/Engine/Source/Runtime/Renderer/Private/LightRendering.cpp
--- a/Engine/Source/Runtime/Renderer/Private/LightRendering.cpp
+++ b/Engine/Source/Runtime/Renderer/Private/LightRendering.cpp
@@ -882,7 +882,7 @@ class FDeferredLightPS : public FGlobalShader
 	class FAtmosphereTransmittance : SHADER_PERMUTATION_BOOL("USE_ATMOSPHERE_TRANSMITTANCE");
 	class FCloudTransmittance 	: SHADER_PERMUTATION_BOOL("USE_CLOUD_TRANSMITTANCE");
 	class FAnistropicMaterials 	: SHADER_PERMUTATION_BOOL("SUPPORTS_ANISOTROPIC_MATERIALS");
-	class FSubstrateTileType	: SHADER_PERMUTATION_INT("SUBSTRATE_TILETYPE", 4);
+	class FSubstrateTileType	: SHADER_PERMUTATION_INT("SUBSTRATE_TILETYPE", 4);				// Simple, Single, Complex, ComplexSpecial
 	class FVirtualShadowMapMask : SHADER_PERMUTATION_BOOL("USE_VIRTUAL_SHADOW_MAP_MASK");
 
 	using FPermutationDomain = TShaderPermutationDomain<
@@ -1025,6 +1025,19 @@ class FDeferredLightPS : public FGlobalShader
 		{
 			return false;
 		}
+
+		if (Substrate::IsSubstrateEnabled() && PermutationVector.Get<FSubstrateTileType>() >= ESubstrateTileType::EComplex && PermutationVector.Get<FAnistropicMaterials>() == 0)
+		{
+			// Remove all permutations for tiles of Complex and ComplexSpecial complexity, when anisotropy is false. We assume that is such tile is asked for, it might have anisotropy requested by the material. 
+			// Substrate does account for MATERIAL_USES_ANISOTROPY during the base pass to avoid requesting Complex tiles if anisotropy is not supported by the project or the platform.
+			return false; 
+		}
+		if (Substrate::IsSubstrateEnabled() && PermutationVector.Get<FSubstrateTileType>() < ESubstrateTileType::EComplex && PermutationVector.Get<FAnistropicMaterials>() == 1)
+		{
+			// Remove all permutations for tiles of Simple and Single complexity, when anisotropy is true. No anisotropy is allowed in this case with Substrate.
+			return false; 
+		}
+
 		return IsFeatureLevelSupported(Parameters.Platform, ERHIFeatureLevel::SM5);
 	}
 
@@ -1045,6 +1058,15 @@ class FDeferredLightPS : public FGlobalShader
 			PermutationVector.Set< FHairComplexTransmittance >(false);
 		}
 
+		if (Substrate::IsSubstrateEnabled() && PermutationVector.Get<FSubstrateTileType>() >= ESubstrateTileType::EComplex)
+		{
+			PermutationVector.Set< FAnistropicMaterials >(true);	// Anisotropy for tiles of Complex and ComplexSpecial complexity
+		}
+		if (Substrate::IsSubstrateEnabled() && PermutationVector.Get<FSubstrateTileType>() < ESubstrateTileType::EComplex)
+		{
+			PermutationVector.Set<FAnistropicMaterials>(false);		// No anisotropy for tiles of Simple and Single complexity
+		}
+
 		return PermutationVector;
 	}
 
@@ -2857,7 +2879,6 @@ static void RenderLight(
 			PermutationVector.Set< FDeferredLightPS::FAtmosphereTransmittance >(IsLightAtmospherePerPixelTransmittanceEnabled(Scene, View, LightSceneInfo));
 			PermutationVector.Set< FDeferredLightPS::FCloudTransmittance >(PassParameters->PS.CloudShadowEnabled > 0);
 		}
-		PermutationVector = FDeferredLightPS::RemapPermutation(PermutationVector);
 
 		// Substrate tile rendering: 
 		// * if the light is directional, then dispatch a set of rect tiles
@@ -2870,6 +2891,7 @@ static void RenderLight(
 			{
 				const ESubstrateTileType TileType = ESubstrateTileType::EComplexSpecial;
 				PermutationVector.Set<FDeferredLightPS::FSubstrateTileType>(TileType);
+				PermutationVector = FDeferredLightPS::RemapPermutation(PermutationVector);
 				TShaderMapRef< FDeferredLightPS > PixelShader(View.ShaderMap, PermutationVector);
 				InternalRenderLight(GraphBuilder, Scene, View, LightSceneInfo, PixelShader, PassParameters, TileType, TEXT("Light::StandardDeferred(ComplexSpecial)"));
 			}
@@ -2877,6 +2899,7 @@ static void RenderLight(
 			{
 				const ESubstrateTileType TileType = ESubstrateTileType::EComplex;
 				PermutationVector.Set<FDeferredLightPS::FSubstrateTileType>(TileType);
+				PermutationVector = FDeferredLightPS::RemapPermutation(PermutationVector);
 				TShaderMapRef< FDeferredLightPS > PixelShader(View.ShaderMap, PermutationVector);
 				InternalRenderLight(GraphBuilder, Scene, View, LightSceneInfo, PixelShader, PassParameters, TileType, TEXT("Light::StandardDeferred(Complex)"));
 			}
@@ -2884,6 +2907,7 @@ static void RenderLight(
 			{
 				const ESubstrateTileType TileType = ESubstrateTileType::ESingle;
 				PermutationVector.Set<FDeferredLightPS::FSubstrateTileType>(TileType);
+				PermutationVector = FDeferredLightPS::RemapPermutation(PermutationVector);
 				TShaderMapRef< FDeferredLightPS > PixelShader(View.ShaderMap, PermutationVector);
 				InternalRenderLight(GraphBuilder, Scene, View, LightSceneInfo, PixelShader, PassParameters, TileType, TEXT("Light::StandardDeferred(Single)"));
 			}
@@ -2891,6 +2915,7 @@ static void RenderLight(
 			{
 				const ESubstrateTileType TileType = ESubstrateTileType::ESimple;
 				PermutationVector.Set<FDeferredLightPS::FSubstrateTileType>(TileType);
+				PermutationVector = FDeferredLightPS::RemapPermutation(PermutationVector);
 				TShaderMapRef< FDeferredLightPS > PixelShader(View.ShaderMap, PermutationVector);
 				InternalRenderLight(GraphBuilder, Scene, View, LightSceneInfo, PixelShader, PassParameters, TileType, TEXT("Light::StandardDeferred(Simple)"));
 			}
@@ -2898,6 +2923,7 @@ static void RenderLight(
 		else
 		{
 			PermutationVector.Set< FDeferredLightPS::FSubstrateTileType>(0);
+			PermutationVector = FDeferredLightPS::RemapPermutation(PermutationVector);
 			TShaderMapRef< FDeferredLightPS > PixelShader(View.ShaderMap, PermutationVector);
 			InternalRenderLight(GraphBuilder, Scene, View, LightSceneInfo, PixelShader, PassParameters, ESubstrateTileType::ECount, TEXT("Light::StandardDeferred"));
 		}
@@ -3023,6 +3049,7 @@ void FDeferredShadingSceneRenderer::RenderLightForHair(
 		PermutationVector.Set< FDeferredLightPS::FCloudTransmittance >(false);
 	}
 
+	PermutationVector = FDeferredLightPS::RemapPermutation(PermutationVector);
 	TShaderMapRef<FDeferredLightHairVS> VertexShader(View.ShaderMap);
 	TShaderMapRef<FDeferredLightPS> PixelShader(View.ShaderMap, PermutationVector);
 
@@ -3243,6 +3270,8 @@ static void InternalRenderSimpleLightsForHair(
 	PermutationVector.Set< FDeferredLightPS::FAtmosphereTransmittance >(false);
 	PermutationVector.Set< FDeferredLightPS::FCloudTransmittance >(false);
 	PermutationVector.Set< FDeferredLightPS::FSubstrateTileType>(Substrate::IsSubstrateEnabled() ? ESubstrateTileType::ESimple : 0);
+
+	PermutationVector = FDeferredLightPS::RemapPermutation(PermutationVector);
 	TShaderMapRef<FDeferredLightPS> PixelShader(View.ShaderMap, PermutationVector);
 	TShaderMapRef<FDeferredLightHairVS> VertexShader(View.ShaderMap);
 
@@ -3331,6 +3360,8 @@ static void InternalRenderSimpleLightsStandardDeferred(
 	PermutationVector.Set< FDeferredLightPS::FAtmosphereTransmittance >(false);
 	PermutationVector.Set< FDeferredLightPS::FCloudTransmittance >(false);
 	PermutationVector.Set< FDeferredLightPS::FSubstrateTileType>(TileType != ESubstrateTileType::ECount ? TileType : 0);
+
+	PermutationVector = FDeferredLightPS::RemapPermutation(PermutationVector);
 	TShaderMapRef<FDeferredLightPS> PixelShader(View.ShaderMap, PermutationVector);
 
 	FDeferredLightVS::FPermutationDomain PermutationVectorVS;
