Substrate - fix for crash for lighting with rect light.

#tests pc editor dx12
#rnx

diff --git a/Engine/Shaders/Private/DeferredLightPixelShaders.usf b/Engine/Shaders/Private/DeferredLightPixelShaders.usf
--- a/Engine/Shaders/Private/DeferredLightPixelShaders.usf
+++ b/Engine/Shaders/Private/DeferredLightPixelShaders.usf
@@ -3,7 +3,9 @@
 /*=============================================================================
 	DeferredLightPixelShaders.usf:
 =============================================================================*/
-  
+
+#pragma message("UESHADERMETADATA_VERSION 1543FD56-DAD0-4ACF-ADFE-7B1543694F24")
+
 #if USE_HAIR_LIGHTING == 1
   #define SUPPORT_CONTACT_SHADOWS 0
 #else
@@ -27,7 +29,7 @@
 #endif
 
 #if SUBSTRATE_GBUFFER_FORMAT==0
-#define SUBSTRATE_LOAD_FROM_MATERIALCONTAINER 0	// Load from a specified BSDF, not from MAterialBuffer.
+#define SUBSTRATE_LOAD_FROM_MATERIALCONTAINER 0	// Load from a specified BSDF, not from MaterialBuffer.
 #define SUBSTRATE_INLINE_SHADING 1				// Inline shading for normal setup if needed.
 #define SUBSTRATE_SSS_MATERIAL_OVERRIDE 1		// Need for SSS dither.
 #include "Substrate/SubstrateRead.ush"
diff --git a/Engine/Source/Runtime/Engine/Private/ShaderCompiler/ShaderCompiler.cpp b/Engine/Source/Runtime/Engine/Private/ShaderCompiler/ShaderCompiler.cpp
--- a/Engine/Source/Runtime/Engine/Private/ShaderCompiler/ShaderCompiler.cpp
+++ b/Engine/Source/Runtime/Engine/Private/ShaderCompiler/ShaderCompiler.cpp
@@ -4760,8 +4760,7 @@ void VerifyGlobalShaders(EShaderPlatform Platform, const ITargetPlatform* Target
 		int32 PermutationCountLimit = 832;	// Nanite culling as of today (2022-01-11) can go up to 832 permutations
 		if (Substrate::IsSubstrateEnabled())
 		{
-			// FDeferredLightPS as of 2025-08-08 has 844 permutations.
-			PermutationCountLimit = 1024;	// FScreenProbeIntegrateCS as of 2025-08-08 has 1024 permutations.
+			PermutationCountLimit = 1368;	// FDeferredLightPS as of 2025-08-08
 		}
 		ensureMsgf(
 			PermutationCountToCompile <= PermutationCountLimit,
diff --git a/Engine/Source/Runtime/Renderer/Private/LightRendering.cpp b/Engine/Source/Runtime/Renderer/Private/LightRendering.cpp
--- a/Engine/Source/Runtime/Renderer/Private/LightRendering.cpp
+++ b/Engine/Source/Runtime/Renderer/Private/LightRendering.cpp
@@ -1026,12 +1026,7 @@ class FDeferredLightPS : public FGlobalShader
 			return false;
 		}
 
-		if (Substrate::IsSubstrateEnabled() && PermutationVector.Get<FSubstrateTileType>() >= ESubstrateTileType::EComplex && PermutationVector.Get<FAnistropicMaterials>() == 0)
-		{
-			// Remove all permutations for tiles of Complex and ComplexSpecial complexity, when anisotropy is false. We assume that is such tile is asked for, it might have anisotropy requested by the material. 
-			// Substrate does account for MATERIAL_USES_ANISOTROPY during the base pass to avoid requesting Complex tiles if anisotropy is not supported by the project or the platform.
-			return false; 
-		}
+		// We must keep FAnistropicMaterials 0/1 permutations for tiles of Complex and ComplexSpecial complexity. That is because rect light or hair can disable anisotropy explicitely.
 		if (Substrate::IsSubstrateEnabled() && PermutationVector.Get<FSubstrateTileType>() < ESubstrateTileType::EComplex && PermutationVector.Get<FAnistropicMaterials>() == 1)
 		{
 			// Remove all permutations for tiles of Simple and Single complexity, when anisotropy is true. No anisotropy is allowed in this case with Substrate.
@@ -1058,10 +1053,7 @@ class FDeferredLightPS : public FGlobalShader
 			PermutationVector.Set< FHairComplexTransmittance >(false);
 		}
 
-		if (Substrate::IsSubstrateEnabled() && PermutationVector.Get<FSubstrateTileType>() >= ESubstrateTileType::EComplex)
-		{
-			PermutationVector.Set< FAnistropicMaterials >(true);	// Anisotropy for tiles of Complex and ComplexSpecial complexity
-		}
+		// We must keep FAnistropicMaterials 0/1 permutations for tiles of Complex and ComplexSpecial complexity. That is because rect light or hair can disable anisotropy explicitely.
 		if (Substrate::IsSubstrateEnabled() && PermutationVector.Get<FSubstrateTileType>() < ESubstrateTileType::EComplex)
 		{
 			PermutationVector.Set<FAnistropicMaterials>(false);		// No anisotropy for tiles of Simple and Single complexity
@@ -2891,32 +2883,32 @@ static void RenderLight(
 			{
 				const ESubstrateTileType TileType = ESubstrateTileType::EComplexSpecial;
 				PermutationVector.Set<FDeferredLightPS::FSubstrateTileType>(TileType);
-				PermutationVector = FDeferredLightPS::RemapPermutation(PermutationVector);
-				TShaderMapRef< FDeferredLightPS > PixelShader(View.ShaderMap, PermutationVector);
+				FDeferredLightPS::FPermutationDomain PermutationVectorRemap = FDeferredLightPS::RemapPermutation(PermutationVector);
+				TShaderMapRef< FDeferredLightPS > PixelShader(View.ShaderMap, PermutationVectorRemap);
 				InternalRenderLight(GraphBuilder, Scene, View, LightSceneInfo, PixelShader, PassParameters, TileType, TEXT("Light::StandardDeferred(ComplexSpecial)"));
 			}
 			// Complex tiles
 			{
 				const ESubstrateTileType TileType = ESubstrateTileType::EComplex;
 				PermutationVector.Set<FDeferredLightPS::FSubstrateTileType>(TileType);
-				PermutationVector = FDeferredLightPS::RemapPermutation(PermutationVector);
-				TShaderMapRef< FDeferredLightPS > PixelShader(View.ShaderMap, PermutationVector);
+				FDeferredLightPS::FPermutationDomain PermutationVectorRemap = FDeferredLightPS::RemapPermutation(PermutationVector);
+				TShaderMapRef< FDeferredLightPS > PixelShader(View.ShaderMap, PermutationVectorRemap);
 				InternalRenderLight(GraphBuilder, Scene, View, LightSceneInfo, PixelShader, PassParameters, TileType, TEXT("Light::StandardDeferred(Complex)"));
 			}
 			// Single tiles
 			{
 				const ESubstrateTileType TileType = ESubstrateTileType::ESingle;
 				PermutationVector.Set<FDeferredLightPS::FSubstrateTileType>(TileType);
-				PermutationVector = FDeferredLightPS::RemapPermutation(PermutationVector);
-				TShaderMapRef< FDeferredLightPS > PixelShader(View.ShaderMap, PermutationVector);
+				FDeferredLightPS::FPermutationDomain PermutationVectorRemap = FDeferredLightPS::RemapPermutation(PermutationVector);
+				TShaderMapRef< FDeferredLightPS > PixelShader(View.ShaderMap, PermutationVectorRemap);
 				InternalRenderLight(GraphBuilder, Scene, View, LightSceneInfo, PixelShader, PassParameters, TileType, TEXT("Light::StandardDeferred(Single)"));
 			}
 			// Simple tiles
 			{
 				const ESubstrateTileType TileType = ESubstrateTileType::ESimple;
 				PermutationVector.Set<FDeferredLightPS::FSubstrateTileType>(TileType);
-				PermutationVector = FDeferredLightPS::RemapPermutation(PermutationVector);
-				TShaderMapRef< FDeferredLightPS > PixelShader(View.ShaderMap, PermutationVector);
+				FDeferredLightPS::FPermutationDomain PermutationVectorRemap = FDeferredLightPS::RemapPermutation(PermutationVector);
+				TShaderMapRef< FDeferredLightPS > PixelShader(View.ShaderMap, PermutationVectorRemap);
 				InternalRenderLight(GraphBuilder, Scene, View, LightSceneInfo, PixelShader, PassParameters, TileType, TEXT("Light::StandardDeferred(Simple)"));
 			}
 		}
