[SingleLayerWater] Fixed MobileSLW MeshPassProcessor CollectPSOInitializers() running even if MobileSupportsSM5MaterialNodes() is false.

diff --git a/Engine/Source/Runtime/Renderer/Private/MobileSingleLayerWaterRendering.cpp b/Engine/Source/Runtime/Renderer/Private/MobileSingleLayerWaterRendering.cpp
--- a/Engine/Source/Runtime/Renderer/Private/MobileSingleLayerWaterRendering.cpp
+++ b/Engine/Source/Runtime/Renderer/Private/MobileSingleLayerWaterRendering.cpp
@@ -248,6 +248,12 @@ void FMobileSingleLayerWaterPassMeshProcessor::CollectPSOInitializersForLMPolicy
 
 void FMobileSingleLayerWaterPassMeshProcessor::CollectPSOInitializers(const FSceneTexturesConfig& SceneTexturesConfig, const FMaterial& Material, const FPSOPrecacheVertexFactoryData& VertexFactoryData, const FPSOPrecacheParams& PreCacheParams, TArray<FPSOPrecacheData>& PSOInitializers)
 {
+	// This MeshProcessor is only used on the mobile SM5 path, otherwise all SLW is handled by the translucent path.
+	EShaderPlatform ShaderPlatform = GetFeatureLevelShaderPlatform(FeatureLevel);
+	if (!MobileSupportsSM5MaterialNodes(ShaderPlatform))
+	{
+		return;
+	}
 	const FMaterialShadingModelField ShadingModels = Material.GetShadingModels();
 	if (!ShadingModels.HasShadingModel(MSM_SingleLayerWater) || !IsOpaqueOrMaskedBlendMode(Material))
 	{
@@ -255,7 +261,6 @@ void FMobileSingleLayerWaterPassMeshProcessor::CollectPSOInitializers(const FSce
 	}
 
 	// Determine the mesh's material and blend mode.
-	EShaderPlatform ShaderPlatform = GetFeatureLevelShaderPlatform(FeatureLevel);
 	const FMeshDrawingPolicyOverrideSettings OverrideSettings = ComputeMeshOverrideSettings(PreCacheParams);
 	const ERasterizerFillMode MeshFillMode = ComputeMeshFillMode(Material, OverrideSettings);
 	ERasterizerCullMode MeshCullMode = ComputeMeshCullMode(Material, OverrideSettings);
