Reduce permutation in DeferredLightPixelShaders when using Format=0 by trimming ComplexSpecial tiles (1368 ->968).

diff --git a/Engine/Source/Runtime/Renderer/Private/LightRendering.cpp b/Engine/Source/Runtime/Renderer/Private/LightRendering.cpp
--- a/Engine/Source/Runtime/Renderer/Private/LightRendering.cpp
+++ b/Engine/Source/Runtime/Renderer/Private/LightRendering.cpp
@@ -1033,6 +1033,12 @@ class FDeferredLightPS : public FGlobalShader
 			return false; 
 		}
 
+		// Skip special complex tile with GBuffer Format=0, as GBuffer Format=0 does not handle glint
+		if (Substrate::IsSubstrateEnabled() && Substrate::IsSubstrateBlendableGBufferEnabled(Parameters.Platform) && PermutationVector.Get<FSubstrateTileType>() == ESubstrateTileType::EComplexSpecial)
+		{
+			return false;
+		}
+
 		return IsFeatureLevelSupported(Parameters.Platform, ERHIFeatureLevel::SM5);
 	}
 
@@ -2879,7 +2885,7 @@ static void RenderLight(
 		if (Substrate::IsSubstrateEnabled())
 		{
 			// Complex Special tiles
-			if (Substrate::GetSubstrateUsesComplexSpecialPath(View))
+			if (Substrate::GetSubstrateUsesComplexSpecialPath(View) && !Substrate::IsSubstrateBlendableGBufferEnabled(View.GetShaderPlatform()))
 			{
 				const ESubstrateTileType TileType = ESubstrateTileType::EComplexSpecial;
 				PermutationVector.Set<FDeferredLightPS::FSubstrateTileType>(TileType);
