Substrate - fixed vis to correctly show anisotropy value.

#rnx
#tests pc editor dx12

diff --git a/Engine/Shaders/Private/Substrate/Substrate.ush b/Engine/Shaders/Private/Substrate/Substrate.ush
--- a/Engine/Shaders/Private/Substrate/Substrate.ush
+++ b/Engine/Shaders/Private/Substrate/Substrate.ush
@@ -143,6 +143,11 @@
 #define SUBSTRATE_OPAQUE_MATERIAL 0
 #endif
 
+// In case we do not have any information about anisotropic material being supported or not, we enable that possibility to be able to call SubstrateReadGBufferBSDF correctly and get good information.
+#ifndef SUPPORTS_ANISOTROPIC_MATERIALS
+#define SUPPORTS_ANISOTROPIC_MATERIALS 1
+#endif
+
 // This is set for material writting to GBuffer for instance based on input connection from the material graph.
 // However, when InternalGetSubstrateSlabBSDF is called form SubstrateReadGBufferBSDF (from many global shaders) we have to assume that we can read and setup anisotropic material so we use the project setting only.
 #ifndef MATERIAL_USES_ANISOTROPY
diff --git a/Engine/Shaders/Shared/SubstrateVisualizeDefinitions.h b/Engine/Shaders/Shared/SubstrateVisualizeDefinitions.h
--- a/Engine/Shaders/Shared/SubstrateVisualizeDefinitions.h
+++ b/Engine/Shaders/Shared/SubstrateVisualizeDefinitions.h
@@ -11,7 +11,7 @@
 
 #ifndef __cplusplus
 // Change this to force recompilation of all Substrate dependent shaders (for instance https://guidgenerator.com/online-guid-generator.aspx)
-#pragma message("UESHADERMETADATA_VERSION E9DF359B-BC55-46BC-B4DF-A5E94467C029")
+#pragma message("UESHADERMETADATA_VERSION 24677C38-A71C-49D5-B60A-DFF59F57B39D")
 #include "/Engine/Shared/SubstrateDefinitions.h"
 #else
 #include "SubstrateDefinitions.h"
diff --git a/Engine/Source/Runtime/Renderer/Private/Substrate/SubstrateVisualize.cpp b/Engine/Source/Runtime/Renderer/Private/Substrate/SubstrateVisualize.cpp
--- a/Engine/Source/Runtime/Renderer/Private/Substrate/SubstrateVisualize.cpp
+++ b/Engine/Source/Runtime/Renderer/Private/Substrate/SubstrateVisualize.cpp
@@ -79,6 +79,7 @@ class FMaterialPrintInfoCS : public FGlobalShader
 		// Stay debug and skip optimizations to reduce compilation time on this long shader.
 		OutEnvironment.CompilerFlags.Add(CFLAG_Debug);
 		OutEnvironment.SetDefine(TEXT("SHADER_MATERIALPRINT"), 1);
+		OutEnvironment.SetDefine(TEXT("SUPPORTS_ANISOTROPIC_MATERIALS"), FDataDrivenShaderPlatformInfo::GetSupportsAnisotropicMaterials(Parameters.Platform));
 	}
 };
 IMPLEMENT_GLOBAL_SHADER(FMaterialPrintInfoCS, "/Engine/Private/Substrate/SubstrateVisualize.usf", "MaterialPrintInfoCS", SF_Compute);
@@ -125,6 +126,7 @@ class FVisualizeMaterialCountPS : public FGlobalShader
 		// Stay debug and skip optimizations to reduce compilation time on this long shader.
 		OutEnvironment.CompilerFlags.Add(CFLAG_Debug);
 		OutEnvironment.SetDefine(TEXT("SHADER_MATERIALCOUNT"), 1);
+		OutEnvironment.SetDefine(TEXT("SUPPORTS_ANISOTROPIC_MATERIALS"), FDataDrivenShaderPlatformInfo::GetSupportsAnisotropicMaterials(Parameters.Platform));
 	}
 };
 IMPLEMENT_GLOBAL_SHADER(FVisualizeMaterialCountPS, "/Engine/Private/Substrate/SubstrateVisualize.usf", "VisualizeMaterialPS", SF_Pixel);
@@ -170,6 +172,7 @@ class FSubstrateSystemInfoCS : public FGlobalShader
 		// Stay debug and skip optimizations to reduce compilation time on this long shader.
 		OutEnvironment.CompilerFlags.Add(CFLAG_Debug);
 		OutEnvironment.SetDefine(TEXT("SHADER_SYSTEMINFO"), 1);
+		OutEnvironment.SetDefine(TEXT("SUPPORTS_ANISOTROPIC_MATERIALS"), FDataDrivenShaderPlatformInfo::GetSupportsAnisotropicMaterials(Parameters.Platform));
 	}
 };
 IMPLEMENT_GLOBAL_SHADER(FSubstrateSystemInfoCS, "/Engine/Private/Substrate/SubstrateVisualize.usf", "MainCS", SF_Compute);
@@ -217,6 +220,7 @@ class FMaterialDebugSubstrateTreeCS : public FGlobalShader
 		// Stay debug and skip optimizations to reduce compilation time on this long shader.
 		OutEnvironment.CompilerFlags.Add(CFLAG_Debug);
 		OutEnvironment.SetDefine(TEXT("SHADER_DEBUGSUBSTRATETREE_CS"), 1);
+		OutEnvironment.SetDefine(TEXT("SUPPORTS_ANISOTROPIC_MATERIALS"), FDataDrivenShaderPlatformInfo::GetSupportsAnisotropicMaterials(Parameters.Platform));
 	}
 };
 IMPLEMENT_GLOBAL_SHADER(FMaterialDebugSubstrateTreeCS, "/Engine/Private/Substrate/SubstrateVisualize.usf", "MaterialDebugSubstrateTreeCS", SF_Compute);
@@ -268,6 +272,7 @@ class FMaterialDebugSubstrateTreePS : public FGlobalShader
 		// Stay debug and skip optimizations to reduce compilation time on this long shader.
 		OutEnvironment.CompilerFlags.Add(CFLAG_Debug);
 		OutEnvironment.SetDefine(TEXT("SHADER_DEBUGSUBSTRATETREE_PS"), 1);
+		OutEnvironment.SetDefine(TEXT("SUPPORTS_ANISOTROPIC_MATERIALS"), FDataDrivenShaderPlatformInfo::GetSupportsAnisotropicMaterials(Parameters.Platform));
 	}
 };
 IMPLEMENT_GLOBAL_SHADER(FMaterialDebugSubstrateTreePS, "/Engine/Private/Substrate/SubstrateVisualize.usf", "MaterialDebugSubstrateTreePS", SF_Pixel);
