[NaniteSkinning] Prevent an assert in skinning scene extension when disabling Nanite.

diff --git a/Engine/Source/Runtime/Renderer/Private/Skinning/SkinningSceneExtension.cpp b/Engine/Source/Runtime/Renderer/Private/Skinning/SkinningSceneExtension.cpp
--- a/Engine/Source/Runtime/Renderer/Private/Skinning/SkinningSceneExtension.cpp
+++ b/Engine/Source/Runtime/Renderer/Private/Skinning/SkinningSceneExtension.cpp
@@ -1260,6 +1260,11 @@ void FSkinningSceneExtension::FRenderer::UpdateViewData(FRDGBuilder& GraphBuilde
 
 	FSceneRendererBase& SceneRenderer = GetSceneRenderer();
 	FSceneCullingRenderer& SceneCullingRenderer = SceneRenderer.GetSceneExtensionsRenderers().GetRenderer<FSceneCullingRenderer>();
+	if (!SceneCullingRenderer.IsEnabled())
+	{
+		return;
+	}
+
 	FInstanceHierarchyParameters InstanceHierarchyParameters = SceneCullingRenderer.GetShaderParameters(GraphBuilder);
 	int32 NumAllocatedChunks = InstanceHierarchyParameters.NumAllocatedChunks;
 	// Create a buffer with enough space for all chunks
