Removed RHI stall waiting for parallel translate to complete occlusion from previous frame. Occlusion queries were refactored so that this is no longer necessary. The new constraint is that EndQuery must be called on the RHI command list instead on the render thread timeline. We now wait for RDG async execution tasks instead of stalling the RHI thread.

diff --git a/Engine/Source/Runtime/Renderer/Private/MobileShadingRenderer.cpp b/Engine/Source/Runtime/Renderer/Private/MobileShadingRenderer.cpp
--- a/Engine/Source/Runtime/Renderer/Private/MobileShadingRenderer.cpp
+++ b/Engine/Source/Runtime/Renderer/Private/MobileShadingRenderer.cpp
@@ -803,6 +803,8 @@ static void EndOcclusionScope(FRDGBuilder& GraphBuilder, TArray<FViewInfo>& View
 			View.ViewState->OcclusionFeedback.EndOcclusionScope(GraphBuilder);
 		}
 	}
+
+	GraphBuilder.AddDispatchHint();
 }
 
 /*
@@ -919,7 +921,6 @@ void FMobileSceneRenderer::RenderFullDepthPrepass(FRDGBuilder& GraphBuilder, TAr
 	if (DoOcclusionQueries() && !bIsSceneCaptureRenderPass)
 	{
 		EndOcclusionScope(GraphBuilder, Views);
-		FenceOcclusionTests(GraphBuilder);
 	}
 }
 
@@ -1509,7 +1510,6 @@ void FMobileSceneRenderer::Render(FRDGBuilder& GraphBuilder, const FSceneRenderU
 		if (DoOcclusionQueries() && !bIsFullDepthPrepassEnabled)
 		{
 			EndOcclusionScope(GraphBuilder, Views);
-		    FenceOcclusionTests(GraphBuilder);
 		}
 
 		SceneTextures.MobileSetupMode = EMobileSceneTextureSetupMode::All;
diff --git a/Engine/Source/Runtime/Renderer/Private/Renderer.cpp b/Engine/Source/Runtime/Renderer/Private/Renderer.cpp
--- a/Engine/Source/Runtime/Renderer/Private/Renderer.cpp
+++ b/Engine/Source/Runtime/Renderer/Private/Renderer.cpp
@@ -152,9 +152,6 @@ void FRendererModule::ShutdownModule()
 
 	// Free up global resources in Lumen
 	Lumen::Shutdown();
-
-	void CleanupOcclusionSubmittedFence();
-	CleanupOcclusionSubmittedFence();
 }
 
 void FRendererModule::OnWorldCleanup(UWorld* World, bool bSessionEnded, bool bCleanupResources, bool bWorldChanged)
diff --git a/Engine/Source/Runtime/Renderer/Private/SceneOcclusion.cpp b/Engine/Source/Runtime/Renderer/Private/SceneOcclusion.cpp
--- a/Engine/Source/Runtime/Renderer/Private/SceneOcclusion.cpp
+++ b/Engine/Source/Runtime/Renderer/Private/SceneOcclusion.cpp
@@ -1509,11 +1509,6 @@ void FDeferredShadingSceneRenderer::RenderOcclusion(
 
 	const bool bUseHzbOcclusion = RenderHzb(GraphBuilder, SceneTextures.Depth.Resolve, BuildHZBAsyncComputeParams, FroxelRenderer);
 
-	if (bIsOcclusionTesting)
-	{
-		FenceOcclusionTests(GraphBuilder);
-	}
-
 	if (bUseHzbOcclusion || bIsOcclusionTesting)
 	{
 		GraphBuilder.AddDispatchHint();
@@ -1539,93 +1534,6 @@ void FMobileSceneRenderer::RenderOcclusion(FRHICommandList& RHICmdList)
 	}
 }
 
-DECLARE_CYCLE_STAT(TEXT("OcclusionSubmittedFence Dispatch"), STAT_OcclusionSubmittedFence_Dispatch, STATGROUP_SceneRendering);
-DECLARE_CYCLE_STAT(TEXT("OcclusionSubmittedFence Wait"), STAT_OcclusionSubmittedFence_Wait, STATGROUP_SceneRendering);
-
-static uint32 GetViewStateUniqueID(const FSceneRenderer* SceneRenderer)
-{
-	return SceneRenderer->Views.Num() && SceneRenderer->Views[0].ViewState ? SceneRenderer->Views[0].ViewState->UniqueID : 0;
-}
-
-void FSceneRenderer::FenceOcclusionTests(FRDGBuilder& GraphBuilder)
-{
-	if (DoOcclusionQueries() && IsRunningRHIInSeparateThread())
-	{
-		AddPass(GraphBuilder, RDG_EVENT_NAME("FenceOcclusionTests"), [this](FRDGAsyncTask, FRHICommandList& RHICmdList)
-		{
-			SCOPE_CYCLE_COUNTER(STAT_OcclusionSubmittedFence_Dispatch);
-
-			// We implement a queue of buffered fences, so we can avoid waiting on queries for another view family that's rendering
-			// in the same frame.  Here we push a new fence into the queue of buffered fences. We assume that the queue isn't full,
-			// since WaitOcclusionTests (called earlier in the frame) will always pop at least one fence if the queue is full.
-			check(OcclusionSubmittedFence[FOcclusionQueryHelpers::MaxBufferedOcclusionFrames - 1].Fence == nullptr);
-
-			for (int32 Dest = FOcclusionQueryHelpers::MaxBufferedOcclusionFrames - 1; Dest >= 1; Dest--)
-			{
-				CA_SUPPRESS(6385);
-				OcclusionSubmittedFence[Dest] = OcclusionSubmittedFence[Dest - 1];
-			}
-
-			OcclusionSubmittedFence[0].Fence = RHICmdList.RHIThreadFence();
-			OcclusionSubmittedFence[0].ViewStateUniqueID = GetViewStateUniqueID(this);
-		});
-
-		GraphBuilder.AddDispatchHint();
-	}
-}
-
-void FSceneRenderer::WaitOcclusionTests(FRHICommandListImmediate& RHICmdList)
-{
-	if (IsRunningRHIInSeparateThread())
-	{
-		SCOPE_CYCLE_COUNTER(STAT_OcclusionSubmittedFence_Wait);
-
-		// Count the number of fences in the queue for the view in question
-		const uint32 ViewStateUniqueID = GetViewStateUniqueID(this);
-
-		int32 ViewStateFenceCount = 0;
-		for (int32 FenceIndex = 0; FenceIndex < FOcclusionQueryHelpers::MaxBufferedOcclusionFrames && OcclusionSubmittedFence[FenceIndex].Fence; FenceIndex++)
-		{
-			if (OcclusionSubmittedFence[FenceIndex].Fence && (OcclusionSubmittedFence[FenceIndex].ViewStateUniqueID == ViewStateUniqueID))
-			{
-				ViewStateFenceCount++;
-			}
-		}
-
-		// Figure out how many fences are allowed to remain in the queue after the wait, which is one less than the number of buffered frames
-		const int32 FencesAllowedInQueue = FOcclusionQueryHelpers::GetNumBufferedFrames(FeatureLevel) - 1;
-		check(FencesAllowedInQueue >= 0);
-
-		// Scan the fences in the queue in order (oldest will be at the end), and wait on them as appropriate.  If the queue is full,
-		// we always want to wait on at least one fence, so we have room to push a new fence onto the queue if necessary.  Thus the
-		// loop always runs at least once before breaking out.
-		int32 FenceToWaitOn = FOcclusionQueryHelpers::MaxBufferedOcclusionFrames - 1;
-		do
-		{
-			// If our fence counting logic is correct, we should already have broken out of the loop
-			check(FenceToWaitOn >= 0);
-
-			// Is there a fence in this queue entry?
-			if (OcclusionSubmittedFence[FenceToWaitOn].Fence)
-			{
-				// If this is one of the fences for the current scene's view state, mark that we're waiting on it
-				if (OcclusionSubmittedFence[FenceToWaitOn].ViewStateUniqueID == ViewStateUniqueID)
-				{
-					ViewStateFenceCount--;
-				}
-
-				// Wait on the current fence and clear it out
-				FRHICommandListExecutor::WaitOnRHIThreadFence(OcclusionSubmittedFence[FenceToWaitOn].Fence);
-				OcclusionSubmittedFence[FenceToWaitOn].Fence = nullptr;
-				OcclusionSubmittedFence[FenceToWaitOn].ViewStateUniqueID = 0;
-			}
-
-			FenceToWaitOn--;
-
-		} while (ViewStateFenceCount > FencesAllowedInQueue);
-	}
-}
-
 FOcclusionFeedback::FOcclusionFeedback()
 	: FRenderResource()
 	, CurrentBufferIndex(0)
diff --git a/Engine/Source/Runtime/Renderer/Private/SceneRendering.cpp b/Engine/Source/Runtime/Renderer/Private/SceneRendering.cpp
--- a/Engine/Source/Runtime/Renderer/Private/SceneRendering.cpp
+++ b/Engine/Source/Runtime/Renderer/Private/SceneRendering.cpp
@@ -482,17 +482,6 @@ static TAutoConsoleVariable<float> CVarSkyLightRealTimeReflectionCapturePreExpos
 	ECVF_Scalability | ECVF_RenderThreadSafe
 );
 
-FOcclusionSubmittedFenceState FSceneRenderer::OcclusionSubmittedFence[FOcclusionQueryHelpers::MaxBufferedOcclusionFrames];
-
-// cleanup OcclusionSubmittedFence to avoid undefined order of destruction that can destroy it after its allocator
-void CleanupOcclusionSubmittedFence()
-{
-	for (FOcclusionSubmittedFenceState& FenceState : FSceneRenderer::OcclusionSubmittedFence)
-	{
-		FenceState.Fence = nullptr;
-	}
-}
-
 extern int32 GetTranslucencyLightingVolumeDim();
 extern bool GetSubstrateEnabledUseRoughRefraction();
 
diff --git a/Engine/Source/Runtime/Renderer/Private/SceneRendering.h b/Engine/Source/Runtime/Renderer/Private/SceneRendering.h
--- a/Engine/Source/Runtime/Renderer/Private/SceneRendering.h
+++ b/Engine/Source/Runtime/Renderer/Private/SceneRendering.h
@@ -1938,12 +1938,6 @@ struct FSortedShadowMaps
 	}
 };
 
-struct FOcclusionSubmittedFenceState
-{
-	FGraphEventRef	Fence;
-	uint32			ViewStateUniqueID;
-};
-
 /**
  * View family plus associated transient scene textures.
  */
@@ -2208,12 +2202,6 @@ public:
 
 	bool DoOcclusionQueries() const;
 
-	void FenceOcclusionTests(FRDGBuilder& GraphBuilder);
-	void WaitOcclusionTests(FRHICommandListImmediate& RHICmdList);
-
-	// fences to make sure the rhi thread has digested the occlusion query renders before we attempt to read them back async
-	static FOcclusionSubmittedFenceState OcclusionSubmittedFence[FOcclusionQueryHelpers::MaxBufferedOcclusionFrames];
-
 	bool ShouldDumpMeshDrawCommandInstancingStats() const { return bDumpMeshDrawCommandInstancingStats; }
 
 	/** bound shader state for occlusion test prims */
diff --git a/Engine/Source/Runtime/Renderer/Private/SceneVisibility.cpp b/Engine/Source/Runtime/Renderer/Private/SceneVisibility.cpp
--- a/Engine/Source/Runtime/Renderer/Private/SceneVisibility.cpp
+++ b/Engine/Source/Runtime/Renderer/Private/SceneVisibility.cpp
@@ -4485,7 +4485,8 @@ void FVisibilityTaskData::LaunchVisibilityTasks(const UE::Tasks::FTask& BeginIni
 
 	if (TaskConfig.Schedule == EVisibilityTaskSchedule::Parallel)
 	{
-		SceneRenderer.WaitOcclusionTests(RHICmdList);
+		// Sync the BeginOcclusionTests RDG pass so that it's safe to wait on the occlusion queries.
+		FRDGBuilder::WaitForAsyncExecuteTask();
 
 		// Parallel occlusion culling is not supported on mobile
 		check(!Views.IsEmpty())
@@ -4753,7 +4754,8 @@ void FVisibilityTaskData::ProcessRenderThreadTasks()
 			UpdatePrimitiveFading(Scene, ViewPacket.View, ViewPacket.ViewState, PrimitiveRange);
 		}
 
-		SceneRenderer.WaitOcclusionTests(RHICmdList);
+		// Sync the BeginOcclusionTests RDG pass so that it's safe to wait on the occlusion queries.
+		FRDGBuilder::WaitForAsyncExecuteTask();
 
 		for (FVisibilityViewPacket& ViewPacket : ViewPackets)
 		{
