[Lumen surface cache] allow surface cache of primitives that affect indirect lighting while hidden to persist on visibility change. This behavior is only exposed through FScene::UpdatePrimitivesDrawnInGame_RenderThread currently.

#tests tested in editor

diff --git a/Engine/Source/Runtime/Renderer/Private/RendererScene.cpp b/Engine/Source/Runtime/Renderer/Private/RendererScene.cpp
--- a/Engine/Source/Runtime/Renderer/Private/RendererScene.cpp
+++ b/Engine/Source/Runtime/Renderer/Private/RendererScene.cpp
@@ -187,6 +187,14 @@ static FAutoConsoleVariableRef CVarSafeCullDistanceUpdate(
 	ECVF_RenderThreadSafe
 );
 
+static TAutoConsoleVariable<int32> CVarLumenSceneSurfaceCacheAllowPersistOnVisibilityChange(
+	TEXT("r.LumenScene.SurfaceCache.AllowPersistOnVisibilityChange"),
+	1,
+	TEXT("Whether to allow primitives that affect indirect lighting while hidden to persist their surface cache on visibility change.\n")
+	TEXT("Only supported through FScene::UpdatePrimitivesDrawnInGame_RenderThread currently."),
+	ECVF_RenderThreadSafe
+);
+
 DECLARE_CYCLE_STAT(TEXT("DeferredShadingSceneRenderer MotionBlurStartFrame"), STAT_FDeferredShadingSceneRenderer_MotionBlurStartFrame, STATGROUP_SceneRendering);
 
 IMPLEMENT_GLOBAL_SHADER_PARAMETER_STRUCT(FDistanceCullFadeUniformShaderParameters, "PrimitiveFade");
@@ -1803,6 +1811,15 @@ void FScene::UpdatePrimitiveInstancesFromCompute(FPrimitiveSceneInfo* PrimitiveS
 	PrimitiveUpdates.Enqueue(PrimitiveSceneInfo, MoveTemp(UpdateCommand));
 }
 
+static bool AllowPersistSurfaceCacheOnVisibilityChange(const FPrimitiveSceneProxy* PrimitiveSceneProxy)
+{
+	return CVarLumenSceneSurfaceCacheAllowPersistOnVisibilityChange.GetValueOnRenderThread() != 0
+		&& PrimitiveSceneProxy->AffectsIndirectLightingWhileHidden()
+		// It doesn't seem useful to persist emissive light source surface cache as most lighting comes from the material
+		// emissive property instead of bounce lighting
+		&& !PrimitiveSceneProxy->IsEmissiveLightSource();
+}
+
 void FScene::UpdatePrimitivesDrawnInGame_RenderThread(TArrayView<FPrimitiveSceneProxy*> InPrimitiveSceneProxies, bool bDrawnInGame)
 {
 	check(IsInRenderingThread());
@@ -1815,11 +1832,16 @@ void FScene::UpdatePrimitivesDrawnInGame_RenderThread(TArrayView<FPrimitiveScene
 			const int32 PrimitiveIndex = PrimitiveSceneInfo->GetIndex();
 			const bool bPrimitiveIndexValid = PrimitiveSceneInfo->IsIndexValid();
 			check(bPrimitiveIndexValid);
+			const bool bPersistSurfaceCache = AllowPersistSurfaceCacheOnVisibilityChange(PrimitiveSceneProxy);
 
 			if (!bDrawnInGame && bPrimitiveIndexValid)
 			{
 				DistanceFieldSceneData.RemovePrimitive(PrimitiveSceneInfo);
-				LumenRemovePrimitive(PrimitiveSceneInfo, PrimitiveIndex);
+
+				if (!bPersistSurfaceCache)
+				{
+					LumenRemovePrimitive(PrimitiveSceneInfo, PrimitiveIndex);
+				}
 			}
 
 			PrimitiveSceneProxy->SetDrawnInGame_RenderThread(bDrawnInGame);
