Increased SKINNING_BUFFER_OFFSET_BITS bits from 20 to 22 to match heavily instantiated assembly foliage content

diff --git a/Engine/Shaders/Shared/SkinningDefinitions.h b/Engine/Shaders/Shared/SkinningDefinitions.h
--- a/Engine/Shaders/Shared/SkinningDefinitions.h
+++ b/Engine/Shaders/Shared/SkinningDefinitions.h
@@ -75,19 +75,21 @@ struct FCompressedBoneTransform
 #endif
 #endif
 
-#define SKINNING_BUFFER_OFFSET_BITS 20
+#define SKINNING_BUFFER_OFFSET_BITS 22
 #define SKINNING_BUFFER_OFFSET_MAX (1 << SKINNING_BUFFER_OFFSET_BITS)
 
+#define SKINNING_BUFFER_INFLUENCE_BITS 6
+#define SKINNING_BUFFER_INFLUENCE_MAX (1 << SKINNING_BUFFER_INFLUENCE_BITS)
+
 struct FSkinningHeader
 {
 	UINT_TYPE HierarchyBufferOffset   : SKINNING_BUFFER_OFFSET_BITS;
 	UINT_TYPE TransformBufferOffset   : SKINNING_BUFFER_OFFSET_BITS;
 	UINT_TYPE ObjectSpaceBufferOffset : SKINNING_BUFFER_OFFSET_BITS;
 	UINT_TYPE MaxTransformCount       : 16;
-	UINT_TYPE MaxInfluenceCount       : 8;
+	UINT_TYPE MaxInfluenceCount       : SKINNING_BUFFER_INFLUENCE_BITS;
 	UINT_TYPE UniqueAnimationCount    : 7;
 	UINT_TYPE bHasScale               : 1;
-	UINT_TYPE Padding                 : 4;
 };
 
 #ifdef __cplusplus
diff --git a/Engine/Source/Runtime/Renderer/Private/Skinning/SkinningSceneExtension.h b/Engine/Source/Runtime/Renderer/Private/Skinning/SkinningSceneExtension.h
--- a/Engine/Source/Runtime/Renderer/Private/Skinning/SkinningSceneExtension.h
+++ b/Engine/Source/Runtime/Renderer/Private/Skinning/SkinningSceneExtension.h
@@ -112,11 +112,12 @@ private:
 
 		FSkinningHeader Pack() const
 		{
-			// Verify that the buffer offsets all fit within the encoded range prior to packing
+			// Verify that values all fit within the encoded range prior to packing
 			check(
 				HierarchyBufferOffset	<= SKINNING_BUFFER_OFFSET_MAX &&
 				TransformBufferOffset	<= SKINNING_BUFFER_OFFSET_MAX &&
-				ObjectSpaceBufferOffset	<= SKINNING_BUFFER_OFFSET_MAX
+				ObjectSpaceBufferOffset	<= SKINNING_BUFFER_OFFSET_MAX &&
+				MaxInfluenceCount		<= SKINNING_BUFFER_INFLUENCE_MAX
 			);
 
 			FSkinningHeader Output;
@@ -127,7 +128,6 @@ private:
 			Output.MaxInfluenceCount		= MaxInfluenceCount;
 			Output.UniqueAnimationCount		= UniqueAnimationCount;
 			Output.bHasScale				= bHasScale;
-			Output.Padding					= 0;
 			return Output;
 		}
 	};
