Fixed: PSO precaching of a SkyPass materials on mobile platforms

diff --git a/Engine/Source/Runtime/Renderer/Private/SkyPassRendering.cpp b/Engine/Source/Runtime/Renderer/Private/SkyPassRendering.cpp
--- a/Engine/Source/Runtime/Renderer/Private/SkyPassRendering.cpp
+++ b/Engine/Source/Runtime/Renderer/Private/SkyPassRendering.cpp
@@ -247,27 +247,53 @@ void FSkyPassMeshProcessor::CollectPSOInitializersInternal(T& SkyPassShaders, co
 	FGraphicsPipelineRenderTargetsInfo RenderTargetsInfo;
 	SetupGBufferRenderTargetInfo(SceneTexturesConfig, RenderTargetsInfo, true /*bSetupDepthStencil*/);
 
-	// Remove all but the SceneColor		
-	RenderTargetsInfo.RenderTargetsEnabled = 1;
-	for (uint32 i = 1; i < MaxSimultaneousRenderTargets; ++i)
+	// Mobile path renders sky in a BasePass so we should have exactly same RT setup
+	if (GetFeatureLevelShadingPath(FeatureLevel) == EShadingPath::Deferred)
 	{
-		RenderTargetsInfo.RenderTargetFormats[i] = 0;
-		RenderTargetsInfo.RenderTargetFlags[i] = ETextureCreateFlags::None;
+		// Remove all but the SceneColor		
+		RenderTargetsInfo.RenderTargetsEnabled = 1;
+		for (uint32 i = 1; i < MaxSimultaneousRenderTargets; ++i)
+		{
+			RenderTargetsInfo.RenderTargetFormats[i] = 0;
+			RenderTargetsInfo.RenderTargetFlags[i] = ETextureCreateFlags::None;
+		}
+
+		FBasePassMeshProcessor::AddBasePassGraphicsPipelineStateInitializer(
+			FeatureLevel,
+			VertexFactoryData,
+			Material,
+			PassDrawRenderState,
+			RenderTargetsInfo,
+			SkyPassShaders,
+			MeshFillMode,
+			MeshCullMode,
+			(EPrimitiveType)PreCacheParams.PrimitiveType,
+			true /*bPrecacheAlphaColorChannel*/,
+			PSOCollectorIndex,
+			PSOInitializers);
 	}
+	else
+	{
+		// subpass info set during the submission of the draws in mobile renderer
+		bool bDeferredShading = IsMobileDeferredShadingEnabled(GMaxRHIShaderPlatform);
+		ESubpassHint SubpassHint = GetSubpassHint(GMaxRHIShaderPlatform, bDeferredShading, RenderTargetsInfo.MultiViewCount > 1, RenderTargetsInfo.NumSamples);
 
-	FBasePassMeshProcessor::AddBasePassGraphicsPipelineStateInitializer(
-		FeatureLevel,
-		VertexFactoryData,
-		Material,
-		PassDrawRenderState,
-		RenderTargetsInfo,
-		SkyPassShaders,
-		MeshFillMode,
-		MeshCullMode,
-		(EPrimitiveType)PreCacheParams.PrimitiveType,
-		true /*bPrecacheAlphaColorChannel*/,
-		PSOCollectorIndex,
-		PSOInitializers);
+		AddGraphicsPipelineStateInitializer(
+			VertexFactoryData,
+			Material,
+			PassDrawRenderState,
+			RenderTargetsInfo,
+			SkyPassShaders,
+			MeshFillMode,
+			MeshCullMode,
+			(EPrimitiveType)PreCacheParams.PrimitiveType,
+			EMeshPassFeatures::Default,
+			SubpassHint,
+			0,
+			true /*bRequired*/,
+			PSOCollectorIndex,
+			PSOInitializers);
+	}
 
 	// Also generate with depth write which is used during CaptureSkyMeshReflection
 	{
