[RayTracing] Remove call to FlushAllocationsToClear on SBT before setting the current bindings on the SBT because the clear call could flush pending clear ops which are still used this frame. Only flush clear after add & remove of primitive scene infos on RT timeline. Pending clears can be added by async stream in requests in between the building of the dirty bindings and the flush to RHI from RDG passes causing cleared data in the SBT and triggering validation errors and potentially missing entries in the SBT

diff --git a/Engine/Source/Runtime/Renderer/Private/DeferredShadingRenderer.cpp b/Engine/Source/Runtime/Renderer/Private/DeferredShadingRenderer.cpp
--- a/Engine/Source/Runtime/Renderer/Private/DeferredShadingRenderer.cpp
+++ b/Engine/Source/Runtime/Renderer/Private/DeferredShadingRenderer.cpp
@@ -1082,8 +1082,6 @@ void FDeferredShadingSceneRenderer::SetupRayTracingRenderingData(FRDGBuilder& Gr
 
 		check(ReferenceView.MaterialRayTracingData.PipelineState || ReferenceView.MaterialRayTracingData.MaterialBindings.Num() == 0);
 
-		Scene->RayTracingSBT.FlushAllocationsToClear(RHICmdList);
-
 		if (ReferenceView.MaterialRayTracingData.PipelineState && (ReferenceView.MaterialRayTracingData.MaterialBindings.Num() || ReferenceView.MaterialRayTracingData.CallableBindings.Num()))
 		{
 			SetRaytracingShaderBindings(RHICmdList, Allocator, ReferenceView.MaterialRayTracingData);
