Fix Material Domain Volume Handling on Mobile Renderer (SM5)

-Allow VolumetricFog on Mobile Renderer (SM5) by checking for mobile SM5 compatibility.
-Mirror the logic from FBasePassMeshProcessor::TryAddMeshBatch to ensure that materials with the Volume domain are excluded from the base pass. These materials should be processed exclusively by the VolumetricFog pass.

#tests console + Editor preview mode

diff --git a/Engine/Source/Runtime/Renderer/Private/MobileBasePass.cpp b/Engine/Source/Runtime/Renderer/Private/MobileBasePass.cpp
--- a/Engine/Source/Runtime/Renderer/Private/MobileBasePass.cpp
+++ b/Engine/Source/Runtime/Renderer/Private/MobileBasePass.cpp
@@ -761,6 +761,7 @@ FMobileBasePassMeshProcessor::FMobileBasePassMeshProcessor(
 
 bool FMobileBasePassMeshProcessor::ShouldDraw(const FMaterial& Material) const
 {
+	bool bShouldDraw = ShouldIncludeDomainInMeshPass(Material.GetMaterialDomain());
 	const FMaterialShadingModelField ShadingModels = Material.GetShadingModels();
 	const bool bIsSLW = ShadingModels.HasShadingModel(MSM_SingleLayerWater);
 	const bool bIsTranslucentSLW = bIsSLW && !MobileSupportsSM5MaterialNodes(GetFeatureLevelShaderPlatform(ERHIFeatureLevel::ES3_1));
@@ -768,18 +769,18 @@ bool FMobileBasePassMeshProcessor::ShouldDraw(const FMaterial& Material) const
 	if (bTranslucentBasePass)
 	{
 		// Skipping TPT_TranslucencyAfterDOFModulate. That pass is only needed for Dual Blending, which is not supported on Mobile.
-		bool bShouldDraw = bIsTranslucent && !Material.IsDeferredDecal() &&
+		bShouldDraw &= bIsTranslucent && !Material.IsDeferredDecal() &&
 			(TranslucencyPassType == ETranslucencyPass::TPT_AllTranslucency
 				|| (TranslucencyPassType == ETranslucencyPass::TPT_TranslucencyStandard && !Material.IsMobileSeparateTranslucencyEnabled())
 				|| (TranslucencyPassType == ETranslucencyPass::TPT_TranslucencyAfterDOF && Material.IsMobileSeparateTranslucencyEnabled()));
 
-		return bShouldDraw;
 	}
 	else
 	{
 		// opaque materials.
-		return !bIsTranslucent && !bIsSLW;
+		bShouldDraw &= !bIsTranslucent && !bIsSLW;
 	}
+	return bShouldDraw;
 }
 static bool MaterialCanReceiveCSM(const FMaterial& Material)
 {
diff --git a/Engine/Source/Runtime/Renderer/Private/VolumetricFog.cpp b/Engine/Source/Runtime/Renderer/Private/VolumetricFog.cpp
--- a/Engine/Source/Runtime/Renderer/Private/VolumetricFog.cpp
+++ b/Engine/Source/Runtime/Renderer/Private/VolumetricFog.cpp
@@ -176,6 +176,11 @@ static int32 GetVolumetricFogGridSizeZ()
 	return FMath::Max(1, GVolumetricFogGridSizeZ);
 }
 
+static bool IsMobileVolumetricFog(const FStaticShaderPlatform Platform)
+{
+	return IsMobilePlatform(Platform) && !MobileSupportsSM5MaterialNodes(Platform);
+}
+
 static FIntPoint GetVolumetricFogTextureResourceRes(const FViewInfo& View)
 {
 	// Allocate texture using scene render targets size so we do not reallocate every frame when dynamic resolution is used in order to avoid resources allocation hitches.
@@ -1076,7 +1081,7 @@ class FVolumetricFogLightScatteringCS : public FGlobalShader
 
 	static FPermutationDomain RemapPermutation(FPermutationDomain PermutationVector, EShaderPlatform ShaderPlatform)
 	{
-		if (IsMobilePlatform(ShaderPlatform))
+		if (IsMobileVolumetricFog(FStaticShaderPlatform(ShaderPlatform)))
 		{
 			PermutationVector.Set<FDistanceFieldSkyOcclusion>(false);
 			PermutationVector.Set<FTemporalReprojection>(false);
@@ -1188,7 +1193,7 @@ bool DoesPlatformSupportVolumetricFogVoxelization(const FStaticShaderPlatform Pl
 	// Substrate requires HLSL2021 which must be cross-compiled for D3D11 and Vulkan to be consumed by FXC compiler.
 	// This cross-compilation toolchain does not support geometry shaders.
 	const bool bCanRHICompileGeometryShaders =!Substrate::IsSubstrateEnabled() || GetMaxSupportedFeatureLevel(Platform) > ERHIFeatureLevel::SM5;
-	return !IsMobilePlatform(Platform) && bCanRHICompileGeometryShaders;
+	return !IsMobileVolumetricFog(Platform) && bCanRHICompileGeometryShaders;
 }
 bool ShouldRenderVolumetricFog(const FScene* Scene, const FSceneViewFamily& ViewFamily)
 {
@@ -1496,7 +1501,7 @@ void FSceneRenderer::ComputeVolumetricFog(FRDGBuilder& GraphBuilder,
 			IntegrationData.FrameJitterOffsetValues[FrameOffsetIndex] = VolumetricFogTemporalRandom(View.Family->FrameNumber - FrameOffsetIndex);
 		}
 
-		const bool bIsMobilePlatform = IsMobilePlatform(View.GetShaderPlatform());
+		const bool bIsMobilePlatform = IsMobileVolumetricFog(View.GetShaderPlatform());
 
 		// Mobile has limited capacities with SRV binding so do not enable atlas sampling on there.
 		const bool bUseLightFunctionAtlasEnabledAndSupported = LightFunctionAtlas::IsEnabled(*Scene, ELightFunctionAtlasSystem::VolumetricFog) && !bIsMobilePlatform;
@@ -1814,7 +1819,7 @@ void FSceneRenderer::ComputeVolumetricFog(FRDGBuilder& GraphBuilder,
 				&& bUseGlobalDistanceField
 				&& Views.Num() == 1
 				&& View.IsPerspectiveProjection()
-				&& !IsMobilePlatform(View.GetShaderPlatform());
+				&& !IsMobileVolumetricFog(View.GetShaderPlatform());
 
 			const int32 SuperSampleCount = FVolumetricFogLightScatteringCS::GetSuperSampleCount(GVolumetricFogHistoryMissSupersampleCount);
 
