[VSM] Add a deprecation and performance warning for disabling r.Shadow.Virtual.ForceOnlyVirtualShadowMaps.

diff --git a/Engine/Source/Runtime/Renderer/Private/ShadowSetup.cpp b/Engine/Source/Runtime/Renderer/Private/ShadowSetup.cpp
--- a/Engine/Source/Runtime/Renderer/Private/ShadowSetup.cpp
+++ b/Engine/Source/Runtime/Renderer/Private/ShadowSetup.cpp
@@ -345,7 +345,20 @@ static TAutoConsoleVariable<int32> CVarForceOnlyVirtualShadowMaps(
 	TEXT("r.Shadow.Virtual.ForceOnlyVirtualShadowMaps"),
 	1,
 	TEXT("If enabled, disallow creation of conventional non-virtual shadow maps for any lights that get a virtual shadow map.\n")
-	TEXT("This can improve performance and save memory, but any geometric primitives that cannot be rendered into the virtual shadow map will not cast shadows."),
+	TEXT("This variable should be enabled for optimal performance and memory usage, but geometric primitives that cannot be rendered into the virtual shadow map will not cast shadows.\n")
+	TEXT("Disabling may reduce performance. This variable is deprecated and may be removed in a future version."),
+	FConsoleVariableDelegate::CreateLambda(
+		[](IConsoleVariable* CVar)
+		{
+			bool bForceOnlyVirtualShadowMaps = CVar->GetInt() != 0;
+			if (!bForceOnlyVirtualShadowMaps)
+			{
+				UE_LOG(LogRenderer, Warning, 
+					   TEXT("Disabling r.Shadow.Virtual.ForceOnlyVirtualShadowMaps may reduce performance. ")
+					   TEXT("This variable is deprecated and may be removed in a future version."));
+			}
+		}
+	),
 	ECVF_RenderThreadSafe
 );
 
