Don't create custom depth and velocity draw commands on mobile when passes are not enabled

diff --git a/Engine/Source/Runtime/Renderer/Private/MobileShadingRenderer.cpp b/Engine/Source/Runtime/Renderer/Private/MobileShadingRenderer.cpp
--- a/Engine/Source/Runtime/Renderer/Private/MobileShadingRenderer.cpp
+++ b/Engine/Source/Runtime/Renderer/Private/MobileShadingRenderer.cpp
@@ -2618,6 +2618,11 @@ int32 FMobileSceneRenderer::ComputeNumOcclusionQueriesToBatch() const
 	return NumQueriesForBatch;
 }
 
+bool FMobileSceneRenderer::ShouldRenderCustomDepth(const FViewInfo& View)
+{
+	return GetCustomDepthStencilUsage(View).bUsesCustomDepthStencil;
+}
+
 // Whether we need a separate render-passes for translucency, decals etc
 bool FMobileSceneRenderer::RequiresMultiPass(int32 NumMSAASamples, EShaderPlatform ShaderPlatform)
 {
diff --git a/Engine/Source/Runtime/Renderer/Private/SceneRendering.h b/Engine/Source/Runtime/Renderer/Private/SceneRendering.h
--- a/Engine/Source/Runtime/Renderer/Private/SceneRendering.h
+++ b/Engine/Source/Runtime/Renderer/Private/SceneRendering.h
@@ -2827,6 +2827,9 @@ public:
 	/** Whether platform requires multiple render-passes for SceneColor rendering */
 	static bool RequiresMultiPass(int32 NumMSAASamples, EShaderPlatform ShaderPlatform);
 
+	/** Check if custom depth pass should be rendered with current view info */
+	static bool ShouldRenderCustomDepth(const FViewInfo& View);
+
 protected:
 	/** Finds the visible dynamic shadows for each view. */
 	FDynamicShadowsTaskData* InitDynamicShadows(FRDGBuilder& GraphBuilder, FInstanceCullingManager& FInstanceCullingManager);
diff --git a/Engine/Source/Runtime/Renderer/Private/SceneVisibility.cpp b/Engine/Source/Runtime/Renderer/Private/SceneVisibility.cpp
--- a/Engine/Source/Runtime/Renderer/Private/SceneVisibility.cpp
+++ b/Engine/Source/Runtime/Renderer/Private/SceneVisibility.cpp
@@ -1426,6 +1426,7 @@ void FRelevancePacket::ComputeRelevance(FDynamicPrimitiveIndexList& DynamicPrimi
 	const bool bIsTranslucentHoldoutEnabled = IsPrimitiveAlphaHoldoutEnabled(View);
 	const bool bIsTranslucentClippedDepthEnabled = NeedTSRThinGeometryDetectionPass(View);
 	const bool bHLODActive = Scene.SceneLODHierarchy.IsActive();
+	const bool bRenderCustomDepth = (ShadingPath == EShadingPath::Mobile) ? FMobileSceneRenderer::ShouldRenderCustomDepth(View) : true;
 	const FHLODVisibilityState* const HLODState = bHLODActive && ViewState ? &ViewState->HLODVisibilityState : nullptr;
 	float MaxDrawDistanceScale = GetCachedScalabilityCVars().ViewDistanceScale;
 	MaxDrawDistanceScale *= GetCachedScalabilityCVars().CalculateFieldOfViewDistanceScale(View.DesiredFOV);
@@ -1718,7 +1719,7 @@ void FRelevancePacket::ComputeRelevance(FDynamicPrimitiveIndexList& DynamicPrimi
 									DrawCommandPacket.AddCommandsForMesh(PrimitiveIndex, PrimitiveSceneInfo, StaticMeshRelevance, StaticMesh, CullingPayloadFlags, Scene, bCanCache, EMeshPass::AnisotropyPass);
 								}
 
-								if (ViewRelevance.bRenderCustomDepth)
+								if (ViewRelevance.bRenderCustomDepth && bRenderCustomDepth)
 								{
 									DrawCommandPacket.AddCommandsForMesh(PrimitiveIndex, PrimitiveSceneInfo, StaticMeshRelevance, StaticMesh, CullingPayloadFlags, Scene, bCanCache, EMeshPass::CustomDepth);
 								}
diff --git a/Engine/Source/Runtime/Renderer/Private/VelocityRendering.cpp b/Engine/Source/Runtime/Renderer/Private/VelocityRendering.cpp
--- a/Engine/Source/Runtime/Renderer/Private/VelocityRendering.cpp
+++ b/Engine/Source/Runtime/Renderer/Private/VelocityRendering.cpp
@@ -599,6 +599,13 @@ bool FVelocityMeshProcessor::PrimitiveHasVelocityForView(const FViewInfo& View,
 		return false;
 	}
 
+	// Only enabled when TAA is enabled on mobile
+	bool bUsesTAA = (View.AntiAliasingMethod == AAM_TemporalAA);
+	if (IsMobilePlatform(View.GetShaderPlatform()) && !bUsesTAA)
+	{
+		return false;
+	}
+
 	const FBoxSphereBounds& PrimitiveBounds = PrimitiveSceneProxy->GetBounds();
 	const float LODFactorDistanceSquared = (PrimitiveBounds.Origin - View.ViewMatrices.GetViewOrigin()).SizeSquared() * FMath::Square(View.LODDistanceFactor);
 
