Add support for vertex count tracking in mesh draw command stats

diff --git a/Engine/Source/Runtime/Engine/Public/MeshDrawCommandStatsSettings.h b/Engine/Source/Runtime/Engine/Public/MeshDrawCommandStatsSettings.h
--- a/Engine/Source/Runtime/Engine/Public/MeshDrawCommandStatsSettings.h
+++ b/Engine/Source/Runtime/Engine/Public/MeshDrawCommandStatsSettings.h
@@ -21,7 +21,10 @@ struct FMeshDrawCommandStatsBudget
 	TArray<FName> LinkedStatNames;
 	/** The category primitive budget. This is the maximum triangles expected, post-culling, summed across all passes. */
 	UPROPERTY(config, EditAnywhere, Category = Engine)
-	int32 PrimitiveBudget = 0;
+	int32 PrimitiveBudget = 0;	
+	/** The category vertex budget. This is the maximum number of vertex expected, post-culling, summed across all passes. */
+	UPROPERTY(config, EditAnywhere, Category = Engine)
+	int32 VertexBudget = 0;
 	/** The collection which contains this budget. */
 	UPROPERTY(config, EditAnywhere, Category = Engine)
 	int32 Collection = 0;
@@ -38,7 +41,10 @@ struct FMeshDrawCommandStatsBudgetTotals
 
 	/** The total amount of primitives budgeted for this collection. */
 	UPROPERTY(config, EditAnywhere, Category = Engine)
-	int32 PrimitiveBudget = 0;
+	int32 PrimitiveBudget = 0;	
+	/** The total amount of vertices budgeted for this collection. */
+	UPROPERTY(config, EditAnywhere, Category = Engine)
+	int32 VertexBudget = 0;
 	/** The collection this applies to. */
 	UPROPERTY(config, EditAnywhere, Category = Engine)
 	int32 Collection = 0;
diff --git a/Engine/Source/Runtime/Renderer/Private/MeshDrawCommandStats.cpp b/Engine/Source/Runtime/Renderer/Private/MeshDrawCommandStats.cpp
--- a/Engine/Source/Runtime/Renderer/Private/MeshDrawCommandStats.cpp
+++ b/Engine/Source/Runtime/Renderer/Private/MeshDrawCommandStats.cpp
@@ -24,10 +24,13 @@ void FMeshDrawCommandStatsManager::CreateInstance()
 DECLARE_STATS_GROUP(TEXT("MeshDrawCommandStats"), STATGROUP_Culling, STATCAT_Advanced);
 
 DECLARE_DWORD_COUNTER_STAT(TEXT("Total Rendered Primitives"), STAT_Culling_TotalNumPrimitives, STATGROUP_Culling);
+DECLARE_DWORD_COUNTER_STAT(TEXT("Total Rendered Vertices"), STAT_Culling_TotalNumVertices, STATGROUP_Culling);
 DECLARE_DWORD_COUNTER_STAT(TEXT("Total Rendered Instances"), STAT_Culling_TotalNumInstances, STATGROUP_Culling);
 DECLARE_DWORD_COUNTER_STAT(TEXT("InstanceCulling Indirect Rendered Primitives"), STAT_Culling_InstanceCullingIndirectNumPrimitives, STATGROUP_Culling);
+DECLARE_DWORD_COUNTER_STAT(TEXT("InstanceCulling Indirect Rendered Vertices"), STAT_Culling_InstanceCullingIndirectNumVertices, STATGROUP_Culling);
 DECLARE_DWORD_COUNTER_STAT(TEXT("InstanceCulling Indirect Rendered Instances"), STAT_Culling_InstanceCullingIndirectNumInstances, STATGROUP_Culling);
 DECLARE_DWORD_COUNTER_STAT(TEXT("Custom Indirect Rendered Primitives"), STAT_Culling_CustomIndirectNumPrimitives, STATGROUP_Culling);
+DECLARE_DWORD_COUNTER_STAT(TEXT("Custom Indirect Rendered Vertices"), STAT_Culling_CustomIndirectNumVertices, STATGROUP_Culling);
 DECLARE_DWORD_COUNTER_STAT(TEXT("Custom Indirect Rendered Instances"), STAT_Culling_CustomIndirectNumInstances, STATGROUP_Culling);
 
 CSV_DEFINE_CATEGORY(MeshDrawCommandStats, false);
@@ -50,6 +53,13 @@ static TAutoConsoleVariable<int32> CVarMeshDrawCommandStats(
 	ECVF_RenderThreadSafe
 );
 
+static TAutoConsoleVariable<int32> CVarMeshDrawCommandUseVertexStats(
+	TEXT("r.MeshDrawCommands.UseVertexCount"),
+	0,
+	TEXT("Use vertex count instead of primitive count for budget reporting."),
+	ECVF_RenderThreadSafe
+);
+
 FMeshDrawCommandStatsManager::FFrameData::~FFrameData()
 {
 	// Collect set of unique readback buffers for deletion (can be shared between MDCs and passes)
@@ -149,12 +159,21 @@ FMeshDrawCommandStatsManager::FMeshDrawCommandStatsManager()
 			const bool bShowPassNameStats = CVarMeshDrawCommandStats->GetInt() == (int)MeshDrawStatsCollection::Pass;
 			auto NumberToK = [](int32 InV) -> int32 { return InV / 1000; };
 
-			int32 TotalPrimitivesTracked = 0;
+			int64 TotalTracked = 0;
 			
 			const bool bShowStats = CVarMeshDrawCommandStats->GetInt() != (int)MeshDrawStatsCollection::None;
 			if (bShowStats)
 			{
-				OutMessages.Add(FCoreDelegates::EOnScreenMessageSeverity::Info, FText::FromString(FString::Printf(TEXT("MeshDrawCommandStats (Triangles / Budget - Category)"))));
+				const bool bUseVertexCount = CVarMeshDrawCommandUseVertexStats->GetInt() > 0;
+				if (bUseVertexCount)
+				{
+					OutMessages.Add(FCoreDelegates::EOnScreenMessageSeverity::Info, FText::FromString(FString::Printf(TEXT("MeshDrawCommandStats (Vertices / Budget - Category)"))));
+				}
+				else
+				{
+					OutMessages.Add(FCoreDelegates::EOnScreenMessageSeverity::Info, FText::FromString(FString::Printf(TEXT("MeshDrawCommandStats (Triangles / Budget - Category)"))));
+				}
+
 				int Collection = CVarMeshDrawCommandStats->GetInt();
 				FStatCollection* StatCollection = StatCollections.Find(Collection);
 
@@ -163,61 +182,65 @@ FMeshDrawCommandStatsManager::FMeshDrawCommandStatsManager()
 				{
 					for (const FCollectionCategory& Category : StatCollection->Categories)
 					{
-						uint64* PrimitiveCount = BudgetedPrimitives.Find(Category.Name);
-						if (PrimitiveCount && *PrimitiveCount > 0)
+						FDrawData* DrawData = BudgetedDrawData.Find(Category.Name);
+						if (DrawData && DrawData->PrimitiveCount > 0)
 						{
-							FCoreDelegates::EOnScreenMessageSeverity Severity = Category.PrimitiveBudget < *PrimitiveCount ? FCoreDelegates::EOnScreenMessageSeverity::Warning : FCoreDelegates::EOnScreenMessageSeverity::Info;
-							OutMessages.Add(Severity, FText::FromString(FString::Printf(TEXT("%5" UINT64_FMT "K / %5dK - %s (%s)"), 
-									NumberToK(*PrimitiveCount), 
-									NumberToK(Category.PrimitiveBudget), 
+							uint64 Count = bUseVertexCount ? DrawData->VertexCount : DrawData->PrimitiveCount;
+							if (Count > 0)
+							{
+								uint64 Budget = bUseVertexCount ? Category.DrawBudgets.VertexCount : Category.DrawBudgets.PrimitiveCount;
+								FCoreDelegates::EOnScreenMessageSeverity Severity = Budget < Count ? FCoreDelegates::EOnScreenMessageSeverity::Warning : FCoreDelegates::EOnScreenMessageSeverity::Info;
+								OutMessages.Add(Severity, FText::FromString(FString::Printf(TEXT("%5" UINT64_FMT "K / %5dK - %s (%s)"),
+									NumberToK(Count),
+									NumberToK(Budget),
 									*(Category.Name.ToString()),
 									*Category.PassFriendlyName
-							)));
+									)));
 
-							TotalPrimitivesTracked += *PrimitiveCount;
+								TotalTracked += Count;
+							}
 						}
 					}
 
-					if (TotalPrimitivesTracked > 0)
+					if (TotalTracked > 0)
 					{
-						int32 CollectionBudget = StatCollection ? StatCollection->PrimitiveBudget : INT32_MAX;
+						int32 CollectionBudget = StatCollection ? (bUseVertexCount ? StatCollection->DrawBudgets.VertexCount : StatCollection->DrawBudgets.PrimitiveCount) : INT32_MAX;
 
-						FCoreDelegates::EOnScreenMessageSeverity Severity = TotalPrimitivesTracked < CollectionBudget ? FCoreDelegates::EOnScreenMessageSeverity::Info : FCoreDelegates::EOnScreenMessageSeverity::Warning;
-						OutMessages.Add(Severity, FText::FromString(FString::Printf(TEXT("%5dK - Total Tracked"), NumberToK(TotalPrimitivesTracked))));
+						FCoreDelegates::EOnScreenMessageSeverity Severity = TotalTracked < CollectionBudget ? FCoreDelegates::EOnScreenMessageSeverity::Info : FCoreDelegates::EOnScreenMessageSeverity::Warning;
+						OutMessages.Add(Severity, FText::FromString(FString::Printf(TEXT("%5dK - Total Tracked"), NumberToK(TotalTracked))));
 					}
 				}
 
 				// Categories which are not tracked by specific budgets
-				int32 TotalPrimitivesUntracked = 0;
+				int32 TotalUntracked = 0;
 				
 				FString PassFriendlyNames = (StatCollection && StatCollection->Untracked.Passes.Num()) ? FString::Printf(TEXT(" (%s)"), *StatCollection->Untracked.PassFriendlyName) : TEXT(" (All Passes)");
-				for (const TPair<FName, uint64>& Pair : UntrackedPrimitives)
+				for (const TPair<FName, FDrawData>& Pair : UntrackedDrawData)
 				{
 					const FName& Name = Pair.Key;
-					uint64 PrimitiveCount = Pair.Value;
-
-					if (PrimitiveCount > 0)
+					uint64 Count = bUseVertexCount ? Pair.Value.VertexCount : Pair.Value.PrimitiveCount;
+					if (Count > 0)
 					{
-						OutMessages.Add(FCoreDelegates::EOnScreenMessageSeverity::Warning, FText::FromString(FString::Printf(TEXT("\t%5dK - %s%s"), NumberToK(PrimitiveCount), *(Name.ToString()), *PassFriendlyNames)));
-					}
-
-					TotalPrimitivesUntracked += PrimitiveCount;
+						OutMessages.Add(FCoreDelegates::EOnScreenMessageSeverity::Warning, FText::FromString(FString::Printf(TEXT("\t%5dK - %s%s"), NumberToK(Count), *(Name.ToString()), *PassFriendlyNames)));
+						TotalUntracked += Count;
+					}					
 				}
 
-				if (TotalPrimitivesUntracked > 0 && !bShowPassNameStats)
+				if (TotalUntracked > 0 && !bShowPassNameStats)
 				{
-					OutMessages.Add(FCoreDelegates::EOnScreenMessageSeverity::Warning, FText::FromString(FString::Printf(TEXT("%5dK - Total Untracked"), NumberToK(TotalPrimitivesUntracked))));
+					OutMessages.Add(FCoreDelegates::EOnScreenMessageSeverity::Warning, FText::FromString(FString::Printf(TEXT("%5dK - Total Untracked"), NumberToK(TotalUntracked))));
 				}
 
-				int32 PrimitiveBudget = StatCollection ? StatCollection->PrimitiveBudget : 0;
-				if (PrimitiveBudget)
+				int32 TotalBudget = StatCollection ? (bUseVertexCount ? StatCollection->DrawBudgets.VertexCount : StatCollection->DrawBudgets.PrimitiveCount) : 0;
+				int32 TotalCount = bUseVertexCount ? Stats.TotalDrawData.VertexCount : Stats.TotalDrawData.PrimitiveCount;
+				if (TotalBudget)
 				{
-					FCoreDelegates::EOnScreenMessageSeverity Severity = PrimitiveBudget < Stats.TotalPrimitives ? FCoreDelegates::EOnScreenMessageSeverity::Warning : FCoreDelegates::EOnScreenMessageSeverity::Info;
-					OutMessages.Add(Severity, FText::FromString(FString::Printf(TEXT("%5dK / %5dK - Total"), NumberToK(TotalPrimitivesTracked + TotalPrimitivesUntracked), NumberToK(PrimitiveBudget))));
+					FCoreDelegates::EOnScreenMessageSeverity Severity = TotalBudget < TotalCount ? FCoreDelegates::EOnScreenMessageSeverity::Warning : FCoreDelegates::EOnScreenMessageSeverity::Info;
+					OutMessages.Add(Severity, FText::FromString(FString::Printf(TEXT("%5dK / %5dK - Total"), NumberToK(TotalTracked + TotalUntracked), NumberToK(TotalBudget))));
 				}
 				else
 				{
-					OutMessages.Add(FCoreDelegates::EOnScreenMessageSeverity::Info, FText::FromString(FString::Printf(TEXT("\t%5dK - Total"), NumberToK(Stats.TotalPrimitives))));
+					OutMessages.Add(FCoreDelegates::EOnScreenMessageSeverity::Info, FText::FromString(FString::Printf(TEXT("\t%5dK - Total"), NumberToK(TotalCount))));
 				}
 			}
 		});
@@ -321,7 +344,7 @@ void FMeshDrawCommandStatsManager::Update()
 					CustomArgsBufferResult.DrawIndexedIndirectParameters = reinterpret_cast<const FRHIDrawIndexedIndirectParameters*>(CustomArgsBufferResult.GPUBufferReadback->Lock(CustomArgsBufferResult.GPUBufferReadback->GetGPUSizeBytes()));
 				}
 
-				using PassCategoryStats = TMap<FName, uint64>;
+				using PassCategoryStats = TMap<FName, FDrawData>;
 				TMap<FName, PassCategoryStats> Passes;
 
 				for (FMeshDrawCommandPassStats* PassStats : FrameData->PassData)
@@ -340,6 +363,7 @@ void FMeshDrawCommandStatsManager::Update()
 					for (int32 CmdIndex = 0; CmdIndex < PassStats->DrawData.Num(); ++CmdIndex)
 					{
 						FVisibleMeshDrawCommandStatsData& DrawData = PassStats->DrawData[CmdIndex];
+						check(DrawData.StatsData.NumVertices > 0);
 						int32 IndirectCommandIndex = DrawData.IndirectArgsOffset / (FInstanceCullingContext::IndirectArgsNumWords * sizeof(uint32));
 						if (DrawData.CustomIndirectArgsBuffer)
 						{
@@ -354,7 +378,8 @@ void FMeshDrawCommandStatsManager::Update()
 								DrawData.TotalInstanceCount = FMath::Max(DrawData.TotalInstanceCount, DrawData.VisibleInstanceCount);
 
 								Stats.CustomIndirectInstances += DrawData.VisibleInstanceCount;
-								Stats.CustomIndirectPrimitives += DrawData.VisibleInstanceCount * DrawData.PrimitiveCount;
+								Stats.CustomIndirectDrawData.VertexCount += DrawData.VisibleInstanceCount * DrawData.StatsData.NumVertices;
+								Stats.CustomIndirectDrawData.PrimitiveCount += DrawData.VisibleInstanceCount * DrawData.PrimitiveCount;
 							}
 						}
 						else if (IndirectArgsPtr && DrawData.UseInstantCullingIndirectBuffer > 0 && InstanceCullingReadBackData)
@@ -364,15 +389,18 @@ void FMeshDrawCommandStatsManager::Update()
 							DrawData.VisibleInstanceCount = IndirectArgs.InstanceCount;
 							ensure(DrawData.VisibleInstanceCount <= DrawData.TotalInstanceCount);
 							Stats.InstanceCullingIndirectInstances += DrawData.VisibleInstanceCount;
-							Stats.InstanceCullingIndirectPrimitives += DrawData.VisibleInstanceCount * DrawData.PrimitiveCount;
+							Stats.InstanceCullingIndirectDrawData.VertexCount += DrawData.VisibleInstanceCount * DrawData.StatsData.NumVertices;
+							Stats.InstanceCullingIndirectDrawData.PrimitiveCount += DrawData.VisibleInstanceCount * DrawData.PrimitiveCount;
 						}
 
 						Stats.TotalInstances += DrawData.VisibleInstanceCount;
-						Stats.TotalPrimitives += DrawData.VisibleInstanceCount * DrawData.PrimitiveCount;
+						Stats.TotalDrawData.VertexCount += DrawData.VisibleInstanceCount * DrawData.StatsData.NumVertices;
+						Stats.TotalDrawData.PrimitiveCount += DrawData.VisibleInstanceCount * DrawData.PrimitiveCount;
 
 						FName StatName = bShowPassNameStats ? PassStats->PassName : DrawData.StatsData.CategoryName;
-						uint64& TotalCount = CategoryStats.FindOrAdd(StatName);
-						TotalCount += DrawData.VisibleInstanceCount * DrawData.PrimitiveCount;
+						FDrawData& CategoryDrawData = CategoryStats.FindOrAdd(StatName);
+						CategoryDrawData.PrimitiveCount += DrawData.VisibleInstanceCount * DrawData.PrimitiveCount;
+						CategoryDrawData.VertexCount += DrawData.VisibleInstanceCount * DrawData.StatsData.NumVertices;
 					}
 
 					if (IndirectArgsPtr)
@@ -391,7 +419,6 @@ void FMeshDrawCommandStatsManager::Update()
 				for (auto PassIter = Passes.CreateConstIterator(); PassIter; ++PassIter)
 				{
 					PassCategoryStats CatMap = PassIter.Value();
-
 					for (auto CatIter = CatMap.CreateConstIterator(); CatIter; ++CatIter)
 					{
 						Stats.CategoryStats.Add(FStats::FCategoryStats(PassIter.Key(), CatIter.Key(), CatIter.Value()));
@@ -420,11 +447,14 @@ void FMeshDrawCommandStatsManager::Update()
 
 	// We keep and set the value from the previous frame in case there are no readback, this avoids alternating values if for example two queries were consumed in one frame
 	// might be able to do this better perhaps.
-	SET_DWORD_STAT(STAT_Culling_TotalNumPrimitives, Stats.TotalPrimitives);
+	SET_DWORD_STAT(STAT_Culling_TotalNumPrimitives, Stats.TotalDrawData.PrimitiveCount);
+	SET_DWORD_STAT(STAT_Culling_TotalNumVertices, Stats.TotalDrawData.VertexCount);
 	SET_DWORD_STAT(STAT_Culling_TotalNumInstances, Stats.TotalInstances);
-	SET_DWORD_STAT(STAT_Culling_InstanceCullingIndirectNumPrimitives, Stats.InstanceCullingIndirectPrimitives);
+	SET_DWORD_STAT(STAT_Culling_InstanceCullingIndirectNumPrimitives, Stats.InstanceCullingIndirectDrawData.PrimitiveCount);
+	SET_DWORD_STAT(STAT_Culling_InstanceCullingIndirectNumVertices, Stats.InstanceCullingIndirectDrawData.VertexCount);
 	SET_DWORD_STAT(STAT_Culling_InstanceCullingIndirectNumInstances, Stats.InstanceCullingIndirectInstances);
-	SET_DWORD_STAT(STAT_Culling_CustomIndirectNumPrimitives, Stats.CustomIndirectPrimitives);
+	SET_DWORD_STAT(STAT_Culling_CustomIndirectNumPrimitives, Stats.CustomIndirectDrawData.PrimitiveCount);
+	SET_DWORD_STAT(STAT_Culling_CustomIndirectNumVertices, Stats.CustomIndirectDrawData.VertexCount);
 	SET_DWORD_STAT(STAT_Culling_CustomIndirectNumInstances, Stats.CustomIndirectInstances);
 
 	// Collect stats during the next frame (check if STATGROUP_Culling is also visible somehow)
@@ -457,7 +487,8 @@ void FMeshDrawCommandStatsManager::Update()
 					Category = &Collection.Categories.AddDefaulted_GetRef();
 				}
 
-				Category->PrimitiveBudget = CategoryBudget.PrimitiveBudget;
+				Category->DrawBudgets.PrimitiveCount = CategoryBudget.PrimitiveBudget;
+				Category->DrawBudgets.VertexCount = CategoryBudget.VertexBudget;
 				Category->Name = CategoryBudget.CategoryName;
 
 				FString& FriendlyNames = Category->PassFriendlyName;
@@ -491,7 +522,8 @@ void FMeshDrawCommandStatsManager::Update()
 				FStatCollection* Collection = StatCollections.Find(BudgetTotal.Collection);
 				if (Collection)
 				{
-					Collection->PrimitiveBudget = BudgetTotal.PrimitiveBudget;
+					Collection->DrawBudgets.PrimitiveCount = BudgetTotal.PrimitiveBudget;
+					Collection->DrawBudgets.VertexCount = BudgetTotal.VertexBudget;
 				}
 			}
 			
@@ -502,8 +534,8 @@ void FMeshDrawCommandStatsManager::Update()
 		}
 
 		// Total up Primitive across stats to their respective Budgets
-		BudgetedPrimitives.Reset();
-		UntrackedPrimitives.Reset();
+		BudgetedDrawData.Reset();
+		UntrackedDrawData.Reset();
 
 		int CollectionIdx = CVarMeshDrawCommandStats->GetInt();
 
@@ -535,16 +567,18 @@ void FMeshDrawCommandStatsManager::Update()
 
 						if (!Category.Passes.Num() || Category.Passes.Contains(CategoryStat.PassName))
 						{
-							uint64& Count = BudgetedPrimitives.FindOrAdd(Category.Name);
-							Count += CategoryStat.PrimitiveCount;
+							FDrawData& DrawData = BudgetedDrawData.FindOrAdd(Category.Name);
+							DrawData.PrimitiveCount += CategoryStat.DrawData.PrimitiveCount;
+							DrawData.VertexCount += CategoryStat.DrawData.VertexCount;
 						}
 					}
 				}
 				// Collect untracked from a specific pass or all passes if there is no collection or the Untracked collection doesn't exist
 				else if (!Collection || !Collection->Untracked.Passes.Num() || Collection->Untracked.Passes.Contains(CategoryStat.PassName))
 				{
-					uint64& Count = UntrackedPrimitives.FindOrAdd(CategoryStat.CategoryName);
-					Count += CategoryStat.PrimitiveCount;
+					FDrawData& DrawData = UntrackedDrawData.FindOrAdd(CategoryStat.CategoryName);
+					DrawData.PrimitiveCount += CategoryStat.DrawData.PrimitiveCount;
+					DrawData.VertexCount += CategoryStat.DrawData.VertexCount;
 				}
 			}
 		}
@@ -554,18 +588,18 @@ void FMeshDrawCommandStatsManager::Update()
 	if (bCsvExport)
 	{
 		// Output Budget totals
-		for (const TPair<FName, uint64>& Pair : BudgetedPrimitives)
+		for (const TPair<FName, FDrawData>& Pair : BudgetedDrawData)
 		{
 			TRACE_CSV_PROFILER_INLINE_STAT(TCHAR_TO_ANSI(*Pair.Key.ToString()), CSV_CATEGORY_INDEX(MeshDrawCommandStats));
-			FCsvProfiler::RecordCustomStat(Pair.Key, CSV_CATEGORY_INDEX(MeshDrawCommandStats), IntCastChecked<int32>(Pair.Value), ECsvCustomStatOp::Set);
+			FCsvProfiler::RecordCustomStat(Pair.Key, CSV_CATEGORY_INDEX(MeshDrawCommandStats), IntCastChecked<int32>(Pair.Value.PrimitiveCount), ECsvCustomStatOp::Set);
 		}
 
 		// Output Untracked totals as a single bucket
 		uint64 TotalUntracked = 0;
 
-		for (const TPair<FName, uint64>& Pair : UntrackedPrimitives)
+		for (const TPair<FName, FDrawData>& Pair : UntrackedDrawData)
 		{
-			TotalUntracked += Pair.Value;
+			TotalUntracked += Pair.Value.PrimitiveCount;
 		}
 
 		const char* Name = "Untracked";
@@ -589,6 +623,7 @@ void FMeshDrawCommandStatsManager::DumpStats(FFrameData* FrameData)
 	{
 		FName PassName;
 		int32 VisibilePrimitiveCount = 0;
+		int32 VisibileVertexCount = 0;
 		int32 VisibleInstance = 0;
 		FName CategoryName;
 		FName ResourceName;
@@ -596,8 +631,10 @@ void FMeshDrawCommandStatsManager::DumpStats(FFrameData* FrameData)
 		int32 SegmentIndex = 0;
 		FString MaterialName;
 		int32 PrimitiveCount = 0;
+		int32 VertexCount = 0;
 		int32 TotalInstanceCount = 0;
 		int32 TotalPrimitiveCount = 0;
+		int32 TotalVertexCount = 0;
 	};
 	TArray<FStatEntry> StatEntries;
 
@@ -610,10 +647,13 @@ void FMeshDrawCommandStatsManager::DumpStats(FFrameData* FrameData)
 				FStatEntry& StatEntry = StatEntries.Add_GetRef(FStatEntry());
 				StatEntry.PassName = PassStats->PassName;
 				StatEntry.VisibilePrimitiveCount = DrawData.VisibleInstanceCount * DrawData.PrimitiveCount;
+				StatEntry.VisibileVertexCount = DrawData.VisibleInstanceCount * DrawData.StatsData.NumVertices;
 				StatEntry.VisibleInstance = DrawData.VisibleInstanceCount;
 				StatEntry.PrimitiveCount = DrawData.PrimitiveCount;
+				StatEntry.VertexCount = DrawData.StatsData.NumVertices;
 				StatEntry.TotalInstanceCount = DrawData.TotalInstanceCount;
 				StatEntry.TotalPrimitiveCount = DrawData.TotalInstanceCount * DrawData.PrimitiveCount;
+				StatEntry.TotalVertexCount = DrawData.TotalInstanceCount * DrawData.StatsData.NumVertices;
 				StatEntry.CategoryName = DrawData.StatsData.CategoryName;
 			#if MESH_DRAW_COMMAND_DEBUG_DATA
 				StatEntry.LODIndex = DrawData.LODIndex;
@@ -637,7 +677,7 @@ void FMeshDrawCommandStatsManager::DumpStats(FFrameData* FrameData)
 			return LHS.VisibilePrimitiveCount > RHS.VisibilePrimitiveCount;
 		});
 
-	const TCHAR* Header = TEXT("Pass,VisiblePrimitiveCount,VisibleInstances,Category,ResourceName,LODIndex,SegmentIndex,MaterialName,PrimitiveCount,TotalInstanceCount,TotalPrimitiveCount\n");
+	const TCHAR* Header = TEXT("Pass,VisiblePrimitiveCount,VisibleVertices,VisibleInstances,Category,ResourceName,LODIndex,SegmentIndex,MaterialName,PrimitiveCount,VertexCount,TotalInstanceCount,TotalPrimitiveCount,TotalVertexCount\n");
 	CSVFile->Serialize(TCHAR_TO_ANSI(Header), FPlatformString::Strlen(Header));
 
 	TCHAR PassNameBuffer[FName::StringBufferSize];
@@ -650,9 +690,10 @@ void FMeshDrawCommandStatsManager::DumpStats(FFrameData* FrameData)
 		StatEntry.CategoryName.ToString(CategoryBuffer);
 		StatEntry.ResourceName.ToString(ResourceNameBuffer);
 
-		FString Row = FString::Printf(TEXT("%s,%d,%d,%s,%s,%d,%d,%s,%d,%d,%d\n"),
+		FString Row = FString::Printf(TEXT("%s,%d,%d,%d,%s,%s,%d,%d,%s,%d,%d,%d,%d,%d\n"),
 			PassNameBuffer,
 			StatEntry.VisibilePrimitiveCount,
+			StatEntry.VisibileVertexCount,
 			StatEntry.VisibleInstance,
 			CategoryBuffer,
 			ResourceNameBuffer,
@@ -660,8 +701,10 @@ void FMeshDrawCommandStatsManager::DumpStats(FFrameData* FrameData)
 			StatEntry.SegmentIndex,
 			*StatEntry.MaterialName,
 			StatEntry.PrimitiveCount,
+			StatEntry.VertexCount,
 			StatEntry.TotalInstanceCount,
-			StatEntry.TotalPrimitiveCount);
+			StatEntry.TotalPrimitiveCount,
+			StatEntry.TotalVertexCount);
 		CSVFile->Serialize(TCHAR_TO_ANSI(*Row), Row.Len());
 	}
 
diff --git a/Engine/Source/Runtime/Renderer/Private/MeshDrawCommandStats.h b/Engine/Source/Runtime/Renderer/Private/MeshDrawCommandStats.h
--- a/Engine/Source/Runtime/Renderer/Private/MeshDrawCommandStats.h
+++ b/Engine/Source/Runtime/Renderer/Private/MeshDrawCommandStats.h
@@ -114,6 +114,21 @@ private:
 		bool bIndirectArgReadbackRequested = false;		
 	};
 
+	struct FDrawData
+	{
+		FDrawData() = default;
+		FDrawData(uint64 InPrimitiveCount, uint64 InVertexCount) : PrimitiveCount(InPrimitiveCount), VertexCount(InVertexCount) {}
+
+		void Reset()
+		{
+			PrimitiveCount = 0;
+			VertexCount = 0;
+		}
+
+		uint64 PrimitiveCount = 0;
+		uint64 VertexCount = 0;
+	};
+
 	/**
 	 * Last updated per frame stats
 	 */
@@ -121,25 +136,25 @@ private:
 	{
 		void Reset()
 		{
-			TotalPrimitives = 0;
+			TotalDrawData.Reset();
 			TotalInstances = 0;
 
-			InstanceCullingIndirectPrimitives = 0;
+			InstanceCullingIndirectDrawData.Reset();
 			InstanceCullingIndirectInstances = 0;
 
-			CustomIndirectPrimitives = 0;
+			CustomIndirectDrawData.Reset();
 			CustomIndirectInstances = 0;
 
 			CategoryStats.Empty();
 		}
 
-		int32 TotalPrimitives = 0;					
+		FDrawData TotalDrawData;
 		int32 TotalInstances = 0;					
 
-		int32 InstanceCullingIndirectPrimitives = 0;
+		FDrawData InstanceCullingIndirectDrawData;
 		int32 InstanceCullingIndirectInstances = 0;		
 
-		int32 CustomIndirectPrimitives = 0;
+		FDrawData CustomIndirectDrawData;
 		int32 CustomIndirectInstances = 0;
 
 		/**
@@ -147,11 +162,11 @@ private:
 		 */
 		struct FCategoryStats
 		{
-			FCategoryStats(FName InPassName, FName InCategoryName, uint64 InPrimitiveCount)
-				: PassName(InPassName), CategoryName(InCategoryName), PrimitiveCount(InPrimitiveCount) {}
+			FCategoryStats(FName InPassName, FName InCategoryName, FDrawData InDrawData)
+				: PassName(InPassName), CategoryName(InCategoryName), DrawData(InDrawData) {}
 			FName PassName;
 			FName CategoryName;
-			uint64 PrimitiveCount;
+			FDrawData DrawData;
 		};
 		TArray<FCategoryStats> CategoryStats;
 	};
@@ -176,7 +191,7 @@ private:
 		FString PassFriendlyName;	// Pipe delimited names of the passes for this category (FString because append to it as it's being created)
 		TSet<FName> Passes;			// Passes this category cares about
 		TSet<FName> LinkedNames;	// LinkedStats this collection cares about
-		int32 PrimitiveBudget;		// Budget for this category
+		FDrawData DrawBudgets; 		// Draw budget for this category
 	};
 
 	struct FStatCollection
@@ -200,7 +215,7 @@ private:
 		TArray<FCollectionCategory> Categories;
 		TMap<FName, TArray<int>> StatToCategoryIndices;
 		FCollectionCategory Untracked;
-		int32 PrimitiveBudget;
+		FDrawData DrawBudgets;
 	};
 
 	using FStatCollectionMap = TMap<int32, FStatCollection>;
@@ -208,19 +223,20 @@ private:
 	/** Dump given frame data stats to csv file on disc in profiling folder */
 	void DumpStats(FFrameData* FrameData);
 
-	int32 CurrentFrameNumber = 0;			//< Ever incrementing frame number
-	bool bRequestDumpStats = false;			//< Dump stats requested via command?
-	bool bCollectStats = false;				//< Collect stats during next frame?
+	int32 CurrentFrameNumber = 0;				//< Ever incrementing frame number
+	bool bRequestDumpStats = false;				//< Dump stats requested via command?
+	bool bCollectStats = false;					//< Collect stats during next frame?
 
-	FString OptionalCategory;				//< Optional category to append to stats filename. Empty string if no category.
+	FString OptionalCategory;					//< Optional category to append to stats filename. Empty string if no category.
 	FCriticalSection FrameDataCS;			
-	TArray<FFrameData*> Frames;						//< All active frames (contains the frame for which we are collecting stats now and all frames waiting for GPU readback)
-	FStats Stats;									//< Last updated frame stats
-	FStatCollectionMap StatCollections; 			//< Per Collection LinkedStatName to Budget CategoryName
-	TMap<FName, uint64> BudgetedPrimitives;			//< Budget CategoryName to Total Primitive Count
-	TMap<FName, uint64> UntrackedPrimitives;		//< Primitives which aren't tracked by any Budgets
+	TArray<FFrameData*> Frames;					//< All active frames (contains the frame for which we are collecting stats now and all frames waiting for GPU readback)
+	FStats Stats;								//< Last updated frame stats
+	FStatCollectionMap StatCollections; 		//< Per Collection LinkedStatName to Budget CategoryName
+
+	TMap<FName, FDrawData> BudgetedDrawData; 	//< Budget CategoryName to Total Primitive Count
+	TMap<FName, FDrawData> UntrackedDrawData;  	//< Primitives which aren't tracked by any Budgets
 
-	FDelegateHandle ScreenMessageDelegate;	//< Delegate used to render optional screen stats
+	FDelegateHandle ScreenMessageDelegate;		//< Delegate used to render optional screen stats
 
 	static FMeshDrawCommandStatsManager* Instance;
 };
diff --git a/Engine/Source/Runtime/Renderer/Private/MeshDrawCommands.cpp b/Engine/Source/Runtime/Renderer/Private/MeshDrawCommands.cpp
--- a/Engine/Source/Runtime/Renderer/Private/MeshDrawCommands.cpp
+++ b/Engine/Source/Runtime/Renderer/Private/MeshDrawCommands.cpp
@@ -913,7 +913,6 @@ void CollectMeshDrawCommandPassStats(
 
 		FVisibleMeshDrawCommandStatsData& DrawData = PassStats->DrawData[DrawCommandIndex];
 		MeshDrawCommand->GetStatsData(DrawData);
-		DrawData.PrimitiveCount		= MeshDrawCommand->NumPrimitives;
 
 		bool bUseIndirect = MeshDrawCommandInfo.bUseIndirect;
 
diff --git a/Engine/Source/Runtime/Renderer/Private/MeshPassProcessor.cpp b/Engine/Source/Runtime/Renderer/Private/MeshPassProcessor.cpp
--- a/Engine/Source/Runtime/Renderer/Private/MeshPassProcessor.cpp
+++ b/Engine/Source/Runtime/Renderer/Private/MeshPassProcessor.cpp
@@ -981,6 +981,12 @@ void FMeshDrawCommand::SetDrawParametersAndFinalize(
 		IndirectArgs.Offset = BatchElement.IndirectArgsOffset;
 	}
 
+#if MESH_DRAW_COMMAND_STATS
+	// Store the num vertices for the MDC stats - with indirect args they are otherwise lost because not really needed for draw but handy for stat tracking
+	// (can't be retrieved from draw data and indirect args)
+	StatsData.NumVertices = BatchElement.MaxVertexIndex - BatchElement.MinVertexIndex + 1;
+#endif
+
 	Finalize(PipelineId, ShadersForDebugging);
 }
 
diff --git a/Engine/Source/Runtime/Renderer/Public/MeshPassProcessor.h b/Engine/Source/Runtime/Renderer/Public/MeshPassProcessor.h
--- a/Engine/Source/Runtime/Renderer/Public/MeshPassProcessor.h
+++ b/Engine/Source/Runtime/Renderer/Public/MeshPassProcessor.h
@@ -838,6 +838,8 @@ struct FMeshDrawCommandStatsData
 #if MESH_DRAW_COMMAND_STATS
 	/** ID used to retrieve the component data during stat collection. */
 	FName CategoryName;
+	/** Total number of vertices for the mesh draw command - optional draw data not always needed. */
+	uint32 NumVertices;
 #endif
 };
 
