BugHawk fixes for comment and uint64 casting

diff --git a/Engine/Source/Runtime/Renderer/Private/MeshDrawCommandStats.cpp b/Engine/Source/Runtime/Renderer/Private/MeshDrawCommandStats.cpp
--- a/Engine/Source/Runtime/Renderer/Private/MeshDrawCommandStats.cpp
+++ b/Engine/Source/Runtime/Renderer/Private/MeshDrawCommandStats.cpp
@@ -157,7 +157,7 @@ FMeshDrawCommandStatsManager::FMeshDrawCommandStatsManager()
 	ScreenMessageDelegate = FRendererOnScreenNotification::Get().AddLambda([this](TMultiMap<FCoreDelegates::EOnScreenMessageSeverity, FText >& OutMessages)
 		{	
 			const bool bShowPassNameStats = CVarMeshDrawCommandStats->GetInt() == (int)MeshDrawStatsCollection::Pass;
-			auto NumberToK = [](int32 InV) -> int32 { return InV / 1000; };
+			auto NumberToK = [](uint64 InV) -> uint64 { return InV / 1000; };
 
 			int64 TotalTracked = 0;
 			
diff --git a/Engine/Source/Runtime/Renderer/Private/MeshDrawCommandStats.h b/Engine/Source/Runtime/Renderer/Private/MeshDrawCommandStats.h
--- a/Engine/Source/Runtime/Renderer/Private/MeshDrawCommandStats.h
+++ b/Engine/Source/Runtime/Renderer/Private/MeshDrawCommandStats.h
@@ -233,8 +233,8 @@ private:
 	FStats Stats;								//< Last updated frame stats
 	FStatCollectionMap StatCollections; 		//< Per Collection LinkedStatName to Budget CategoryName
 
-	TMap<FName, FDrawData> BudgetedDrawData; 	//< Budget CategoryName to Total Primitive Count
-	TMap<FName, FDrawData> UntrackedDrawData;  	//< Primitives which aren't tracked by any Budgets
+	TMap<FName, FDrawData> BudgetedDrawData; 	//< Budget CategoryName to Total Draw data Count
+	TMap<FName, FDrawData> UntrackedDrawData;  	//< Draw data stats which aren't tracked by any Budgets
 
 	FDelegateHandle ScreenMessageDelegate;		//< Delegate used to render optional screen stats
 
diff --git a/Engine/Source/Runtime/Renderer/Public/MeshPassProcessor.h b/Engine/Source/Runtime/Renderer/Public/MeshPassProcessor.h
--- a/Engine/Source/Runtime/Renderer/Public/MeshPassProcessor.h
+++ b/Engine/Source/Runtime/Renderer/Public/MeshPassProcessor.h
@@ -838,7 +838,7 @@ struct FMeshDrawCommandStatsData
 #if MESH_DRAW_COMMAND_STATS
 	/** ID used to retrieve the component data during stat collection. */
 	FName CategoryName;
-	/** Total number of vertices for the mesh draw command - optional draw data not always needed. */
+	/** Total number of vertices for the mesh draw command */
 	uint32 NumVertices;
 #endif
 };
