[HWRT] Flush SBT allocations to clear in a separate RDG pass.
- It's currently done inside the pass where we set ray tracing bindings for a view, but we are adding support for multiple views so we need to move it to it's own pass that is view-indepedent.

diff --git a/Engine/Source/Runtime/Renderer/Private/DeferredShadingRenderer.cpp b/Engine/Source/Runtime/Renderer/Private/DeferredShadingRenderer.cpp
--- a/Engine/Source/Runtime/Renderer/Private/DeferredShadingRenderer.cpp
+++ b/Engine/Source/Runtime/Renderer/Private/DeferredShadingRenderer.cpp
@@ -1030,6 +1030,15 @@ void FDeferredShadingSceneRenderer::SetupRayTracingRenderingData(FRDGBuilder& Gr
 		ReferenceView.DirtyPersistentRayTracingShaderBindings = Scene->RayTracingSBT.GetDirtyBindings(ReferenceView.VisibleRayTracingShaderBindings, bForceAllDirty);
 	}
 
+	GraphBuilder.AddPass(RDG_EVENT_NAME("FlushRayTracingSBTAllocationsToClear"), ERDGPassFlags::NeverCull,
+		[Scene = Scene](FRDGAsyncTask, FRHICommandList& RHICmdList)
+		{
+			TRACE_CPUPROFILER_EVENT_SCOPE(FlushRayTracingSBTAllocationsToClear);
+
+			Scene->RayTracingSBT.FlushAllocationsToClear(RHICmdList);
+		}
+	);
+
 	// Prepare the local ray tracing shader binding data to update on RHI timeline for Material, Lumen and InlineLumen
 	if (ReferenceView.MaterialRayTracingData.ShaderBindingTable)
 	{
@@ -1091,8 +1100,6 @@ void FDeferredShadingSceneRenderer::SetupRayTracingRenderingData(FRDGBuilder& Gr
 
 		check(ReferenceView.MaterialRayTracingData.PipelineState || ReferenceView.MaterialRayTracingData.MaterialBindings.Num() == 0);
 
-		Scene->RayTracingSBT.FlushAllocationsToClear(RHICmdList);
-
 		if (ReferenceView.MaterialRayTracingData.PipelineState && (ReferenceView.MaterialRayTracingData.MaterialBindings.Num() || ReferenceView.MaterialRayTracingData.CallableBindings.Num()))
 		{
 			SetRaytracingShaderBindings(RHICmdList, Allocator, ReferenceView.MaterialRayTracingData);
