Fix crash when enabling path tracing while Lumen visualization is also enabled.

diff --git a/Engine/Source/Runtime/Renderer/Private/PostProcess/PostProcessing.cpp b/Engine/Source/Runtime/Renderer/Private/PostProcess/PostProcessing.cpp
--- a/Engine/Source/Runtime/Renderer/Private/PostProcess/PostProcessing.cpp
+++ b/Engine/Source/Runtime/Renderer/Private/PostProcess/PostProcessing.cpp
@@ -555,7 +555,7 @@ void AddPostProcessingPasses(
 #else
 	PassSequence.SetEnabled(EPass::VisualizePostProcessStack, false);
 #endif
-	PassSequence.SetEnabled(EPass::VisualizeLumenScene, LumenVisualizeMode >= 0 && LumenVisualizeMode != VISUALIZE_MODE_OVERVIEW && LumenVisualizeMode != VISUALIZE_MODE_PERFORMANCE_OVERVIEW && bPostProcessingEnabled);
+	PassSequence.SetEnabled(EPass::VisualizeLumenScene, LumenVisualizeMode >= 0 && LumenVisualizeMode != VISUALIZE_MODE_OVERVIEW && LumenVisualizeMode != VISUALIZE_MODE_PERFORMANCE_OVERVIEW && bPostProcessingEnabled && !EngineShowFlags.PathTracing);
 	PassSequence.SetEnabled(EPass::VisualizeSubstrate, Substrate::ShouldRenderSubstrateDebugPasses(View));
 	PassSequence.SetEnabled(EPass::VisualizeLightGrid, View.Family && ShouldVisualizeLightGrid(View.Family->GetShaderPlatform()));
 
@@ -588,7 +588,7 @@ void AddPostProcessingPasses(
 	PassSequence.SetEnabled(EPass::VisualizeGBufferHints, EngineShowFlags.GBufferHints);
 	PassSequence.SetEnabled(EPass::VisualizeSubsurface, EngineShowFlags.VisualizeSSS);
 	PassSequence.SetEnabled(EPass::VisualizeGBufferOverview, bVisualizeGBufferOverview || bVisualizeGBufferDumpToFile || bVisualizeGBufferDumpToPIpe);
-	PassSequence.SetEnabled(EPass::VisualizeLumenSceneOverview, (LumenVisualizeMode == VISUALIZE_MODE_OVERVIEW || LumenVisualizeMode == VISUALIZE_MODE_PERFORMANCE_OVERVIEW) && bPostProcessingEnabled);
+	PassSequence.SetEnabled(EPass::VisualizeLumenSceneOverview, (LumenVisualizeMode == VISUALIZE_MODE_OVERVIEW || LumenVisualizeMode == VISUALIZE_MODE_PERFORMANCE_OVERVIEW) && bPostProcessingEnabled && !EngineShowFlags.PathTracing);
 	PassSequence.SetEnabled(EPass::VisualizeHDR, bVisualizeHDR);
 	PassSequence.SetEnabled(EPass::VisualizeMotionVectors, EngineShowFlags.VisualizeMotionVectors || EngineShowFlags.VisualizeReprojection);
 	PassSequence.SetEnabled(EPass::VisualizeTemporalUpscaler, EngineShowFlags.VisualizeTemporalUpscaler);
