[HWRT] Decouple inline ray tracing SBT from View specific code.

diff --git a/Engine/Source/Runtime/Renderer/Private/DeferredShadingRenderer.cpp b/Engine/Source/Runtime/Renderer/Private/DeferredShadingRenderer.cpp
--- a/Engine/Source/Runtime/Renderer/Private/DeferredShadingRenderer.cpp
+++ b/Engine/Source/Runtime/Renderer/Private/DeferredShadingRenderer.cpp
@@ -630,7 +630,9 @@ BEGIN_SHADER_PARAMETER_STRUCT(FSetRayTracingBindingsPassParams, )
 	SHADER_PARAMETER_RDG_UNIFORM_BUFFER(FSceneUniformParameters, Scene)
 	SHADER_PARAMETER_RDG_UNIFORM_BUFFER(FRayTracingLightGrid, LightGridPacked)
 	SHADER_PARAMETER_STRUCT_REF(FLumenHardwareRayTracingUniformBufferParameters, LumenHardwareRayTracingUniformBuffer)
+END_SHADER_PARAMETER_STRUCT()
 
+BEGIN_SHADER_PARAMETER_STRUCT(FSetRayTracingBindingsInlinePassParams, )
 	RDG_BUFFER_ACCESS(InlineRayTracingBindingDataBuffer, ERHIAccess::CopyDest)
 END_SHADER_PARAMETER_STRUCT()
 
@@ -835,13 +837,12 @@ bool FDeferredShadingSceneRenderer::SetupRayTracingPipelineStatesAndSBT(FRDGBuil
 	}
 
 	// Send SBTs to all views since they all share the same ones
-	EnumerateLinkedViews([MaterialSBT, LumenSBT, InlineSBT, InlineBindingDataBuffer] (FViewInfo& View)
+	EnumerateLinkedViews([MaterialSBT, LumenSBT, InlineBindingDataBuffer] (FViewInfo& View)
 	{
 		if (View.bHasAnyRayTracingPass)
 		{
 			View.MaterialRayTracingData.ShaderBindingTable = MaterialSBT;
 			View.LumenRayTracingData.ShaderBindingTable = LumenSBT;
-			View.InlineRayTracingData.ShaderBindingTable = InlineSBT;
 
 			View.InlineRayTracingBindingDataBuffer = InlineBindingDataBuffer;
 		}
@@ -1022,7 +1023,7 @@ void FDeferredShadingSceneRenderer::SetupRayTracingRenderingData(FRDGBuilder& Gr
 	
 	// Build the dirty persistent shader bindings from the visible shader bindings (SBT version is updated after the PSO and SBTs have been created)
 	// Inline lumen buffer bindings don't use the dirty bindings yet - so don't need to build it here (uses the visible binding set)
-	bool bRequireBindingsUpdate = ReferenceView.MaterialRayTracingData.ShaderBindingTable || ReferenceView.LumenRayTracingData.ShaderBindingTable || ReferenceView.InlineRayTracingData.ShaderBindingTable;
+	bool bRequireBindingsUpdate = ReferenceView.MaterialRayTracingData.ShaderBindingTable || ReferenceView.LumenRayTracingData.ShaderBindingTable || Scene->InlineRayTracingSBTID != INDEX_NONE;
 	if (bRequireBindingsUpdate)
 	{
 		// If fallback RTPSO then mark all bindings as dirty because they need to bound again when final RTPSO is ready (Shader identifier could have changed)
@@ -1048,27 +1049,6 @@ void FDeferredShadingSceneRenderer::SetupRayTracingRenderingData(FRDGBuilder& Gr
 	{
 		SetupLumenHardwareRayTracingHitGroupBindings(GraphBuilder, ReferenceView);	
 	}
-	if (ReferenceView.InlineRayTracingData.ShaderBindingTable)
-	{		
-		TRACE_CPUPROFILER_EVENT_SCOPE(FDeferredShadingSceneRenderer::SetupInlineHardwareRaytracingHitGroupBindings);
-
-		uint32 ShaderSlotsPerSegment = Scene->RayTracingSBT.GetNumShaderSlotsPerSegment();
-		AddRayTracingLocalShaderBindingWriterTasks(GraphBuilder, ReferenceView.DirtyPersistentRayTracingShaderBindings, ReferenceView.InlineRayTracingData.MaterialBindings, 
-		                                           [ShaderSlotsPerSegment, &RayTracingMeshCommands = Scene->CachedRayTracingMeshCommands](const FRayTracingShaderBindingData& RTShaderBindingData, FRayTracingLocalShaderBindingWriter* BindingWriter)
-		{
-			const FRayTracingMeshCommand& MeshCommand = RTShaderBindingData.GetRayTracingMeshCommand(RayTracingMeshCommands);
-
-			for (uint32 SlotIndex = 0; SlotIndex < ShaderSlotsPerSegment; ++SlotIndex)
-			{
-				FRayTracingLocalShaderBindings& Binding = BindingWriter->AddWithExternalParameters();
-				Binding.RecordIndex = RTShaderBindingData.SBTRecordIndex + SlotIndex;
-				Binding.Geometry = RTShaderBindingData.RayTracingGeometry;
-				Binding.SegmentIndex = MeshCommand.GeometrySegmentIndex;							
-				Binding.BindingType = RTShaderBindingData.BindingType;
-				Binding.UserData = 0;
-			}
-		});
-	}
 
 	FRDGBufferRef LumenHardwareRayTracingHitDataBuffer = nullptr;
 	if (bAnyInlineHardwareRayTracingPassEnabled)
@@ -1086,67 +1066,99 @@ void FDeferredShadingSceneRenderer::SetupRayTracingRenderingData(FRDGBuilder& Gr
 
 	const bool bIsPathTracing = ViewFamily.EngineShowFlags.PathTracing;
 
-	FSetRayTracingBindingsPassParams* PassParams = GraphBuilder.AllocParameters<FSetRayTracingBindingsPassParams>();
-	PassParams->Scene = GetSceneUniformBufferRef(GraphBuilder);
-	PassParams->LightGridPacked = bIsPathTracing ? nullptr : ReferenceView.RayTracingLightGridUniformBuffer; // accessed by FRayTracingLightingMS // Is this needed for anything?
-	PassParams->LumenHardwareRayTracingUniformBuffer = ReferenceView.LumenHardwareRayTracingUniformBuffer;
-	PassParams->InlineRayTracingBindingDataBuffer = ReferenceView.InlineRayTracingBindingDataBuffer;
-
-	const FRayTracingLightFunctionMap* RayTracingLightFunctionMap = GraphBuilder.Blackboard.Get<FRayTracingLightFunctionMap>();
-	GraphBuilder.AddPass(RDG_EVENT_NAME("SetRayTracingBindings"), PassParams, ERDGPassFlags::Copy | ERDGPassFlags::Compute | ERDGPassFlags::NeverCull,
-		[this, PassParams, bIsPathTracing, &ReferenceView, RayTracingLightFunctionMap](FRDGAsyncTask, FRHICommandList& RHICmdList)
 	{
-		TRACE_CPUPROFILER_EVENT_SCOPE(SetRayTracingBindings);
+		const FRayTracingLightFunctionMap* RayTracingLightFunctionMap = GraphBuilder.Blackboard.Get<FRayTracingLightFunctionMap>();
 
-		check(ReferenceView.MaterialRayTracingData.PipelineState || ReferenceView.MaterialRayTracingData.MaterialBindings.Num() == 0);
-
-		if (ReferenceView.MaterialRayTracingData.PipelineState && (ReferenceView.MaterialRayTracingData.MaterialBindings.Num() || ReferenceView.MaterialRayTracingData.CallableBindings.Num()))
-		{
-			SetRaytracingShaderBindings(RHICmdList, Allocator, ReferenceView.MaterialRayTracingData);
+		FSetRayTracingBindingsPassParams* PassParams = GraphBuilder.AllocParameters<FSetRayTracingBindingsPassParams>();
+		PassParams->Scene = GetSceneUniformBufferRef(GraphBuilder);
+		PassParams->LightGridPacked = bIsPathTracing ? nullptr : ReferenceView.RayTracingLightGridUniformBuffer; // accessed by FRayTracingLightingMS // Is this needed for anything?
+		PassParams->LumenHardwareRayTracingUniformBuffer = ReferenceView.LumenHardwareRayTracingUniformBuffer;
 
-			if (bIsPathTracing)
+		GraphBuilder.AddPass(RDG_EVENT_NAME("SetRayTracingBindings"), PassParams, ERDGPassFlags::Copy | ERDGPassFlags::Compute | ERDGPassFlags::NeverCull,
+			[this, PassParams, bIsPathTracing, &ReferenceView, RayTracingLightFunctionMap](FRDGAsyncTask, FRHICommandList& RHICmdList)
 			{
-				SetupPathTracingDefaultMissShader(RHICmdList, ReferenceView);
+				TRACE_CPUPROFILER_EVENT_SCOPE(SetRayTracingBindings);
 
-				BindLightFunctionShadersPathTracing(RHICmdList, Scene, RayTracingLightFunctionMap, ReferenceView);
-			}
-			else
-			{
-				SetupRayTracingDefaultMissShader(RHICmdList, ReferenceView);
-				SetupRayTracingLightingMissShader(RHICmdList, ReferenceView);
+				check(ReferenceView.MaterialRayTracingData.PipelineState || ReferenceView.MaterialRayTracingData.MaterialBindings.Num() == 0);
 
-				BindLightFunctionShaders(RHICmdList, Scene, RayTracingLightFunctionMap, ReferenceView);
-			}
-			
-			RHICmdList.CommitShaderBindingTable(ReferenceView.MaterialRayTracingData.ShaderBindingTable);			
-		}
-
-		if (!bIsPathTracing)
-		{
-			if (GRHISupportsRayTracingShaders || GRHISupportsInlineRayTracing)
-			{
-				if (ReferenceView.LumenRayTracingData.PipelineState && ReferenceView.LumenRayTracingData.ShaderBindingTable)
+				if (ReferenceView.MaterialRayTracingData.PipelineState && (ReferenceView.MaterialRayTracingData.MaterialBindings.Num() || ReferenceView.MaterialRayTracingData.CallableBindings.Num()))
 				{
-					RHICmdList.SetRayTracingMissShader(ReferenceView.LumenRayTracingData.ShaderBindingTable, RAY_TRACING_MISS_SHADER_SLOT_DEFAULT, ReferenceView.LumenRayTracingData.PipelineState, 0 /* MissShaderPipelineIndex */, 0, nullptr, 0);
+					SetRaytracingShaderBindings(RHICmdList, Allocator, ReferenceView.MaterialRayTracingData);
+
+					if (bIsPathTracing)
+					{
+						SetupPathTracingDefaultMissShader(RHICmdList, ReferenceView);
+
+						BindLightFunctionShadersPathTracing(RHICmdList, Scene, RayTracingLightFunctionMap, ReferenceView);
+					}
+					else
+					{
+						SetupRayTracingDefaultMissShader(RHICmdList, ReferenceView);
+						SetupRayTracingLightingMissShader(RHICmdList, ReferenceView);
+
+						BindLightFunctionShaders(RHICmdList, Scene, RayTracingLightFunctionMap, ReferenceView);
+					}
+
+					RHICmdList.CommitShaderBindingTable(ReferenceView.MaterialRayTracingData.ShaderBindingTable);
 				}
 
-				if (ReferenceView.LumenRayTracingData.ShaderBindingTable)
+				if (!bIsPathTracing)
 				{
-					SetRaytracingShaderBindings(RHICmdList, Allocator, ReferenceView.LumenRayTracingData);
-					RHICmdList.CommitShaderBindingTable(ReferenceView.LumenRayTracingData.ShaderBindingTable);
+					if (ReferenceView.LumenRayTracingData.PipelineState && ReferenceView.LumenRayTracingData.ShaderBindingTable)
+					{
+						RHICmdList.SetRayTracingMissShader(ReferenceView.LumenRayTracingData.ShaderBindingTable, RAY_TRACING_MISS_SHADER_SLOT_DEFAULT, ReferenceView.LumenRayTracingData.PipelineState, 0 /* MissShaderPipelineIndex */, 0, nullptr, 0);
+					}
+
+					if (ReferenceView.LumenRayTracingData.ShaderBindingTable)
+					{
+						SetRaytracingShaderBindings(RHICmdList, Allocator, ReferenceView.LumenRayTracingData);
+						RHICmdList.CommitShaderBindingTable(ReferenceView.LumenRayTracingData.ShaderBindingTable);
+					}
 				}
+			});
+	}
+
+	if (!bIsPathTracing && Scene->InlineRayTracingSBTID != INDEX_NONE)
+	{
+		FViewInfo::FRayTracingData InlineRayTracingData;
+		InlineRayTracingData.ShaderBindingTable = Scene->RayTracingSBT.GetPersistentSBT(Scene->InlineRayTracingSBTID);
+
+		{
+			TRACE_CPUPROFILER_EVENT_SCOPE(FDeferredShadingSceneRenderer::SetupInlineHardwareRaytracingHitGroupBindings);
 
-				if (ReferenceView.InlineRayTracingData.ShaderBindingTable)
+			const uint32 ShaderSlotsPerSegment = Scene->RayTracingSBT.GetNumShaderSlotsPerSegment();
+			AddRayTracingLocalShaderBindingWriterTasks(GraphBuilder, ReferenceView.DirtyPersistentRayTracingShaderBindings, InlineRayTracingData.MaterialBindings,
+				[ShaderSlotsPerSegment, &RayTracingMeshCommands = Scene->CachedRayTracingMeshCommands](const FRayTracingShaderBindingData& RTShaderBindingData, FRayTracingLocalShaderBindingWriter* BindingWriter)
 				{
-					check(ReferenceView.InlineRayTracingBindingDataBuffer);
+					const FRayTracingMeshCommand& MeshCommand = RTShaderBindingData.GetRayTracingMeshCommand(RayTracingMeshCommands);
 
-					SetRaytracingShaderBindings(RHICmdList, Allocator, ReferenceView.InlineRayTracingData);
-					
-					RHICmdList.CommitShaderBindingTable(ReferenceView.InlineRayTracingData.ShaderBindingTable, PassParams->InlineRayTracingBindingDataBuffer->GetRHI());
-				}
-			}
+					for (uint32 SlotIndex = 0; SlotIndex < ShaderSlotsPerSegment; ++SlotIndex)
+					{
+						FRayTracingLocalShaderBindings& Binding = BindingWriter->AddWithExternalParameters();
+						Binding.RecordIndex = RTShaderBindingData.SBTRecordIndex + SlotIndex;
+						Binding.Geometry = RTShaderBindingData.RayTracingGeometry;
+						Binding.SegmentIndex = MeshCommand.GeometrySegmentIndex;
+						Binding.BindingType = RTShaderBindingData.BindingType;
+						Binding.UserData = 0;
+					}
+				});
 		}
-	});
+
+		{
+			FSetRayTracingBindingsInlinePassParams* PassParams = GraphBuilder.AllocParameters<FSetRayTracingBindingsInlinePassParams>();
+			PassParams->InlineRayTracingBindingDataBuffer = Scene->RayTracingSBT.GetPersistentInlineBindingDataBuffer(GraphBuilder, Scene->InlineRayTracingSBTID);
+
+			GraphBuilder.AddPass(RDG_EVENT_NAME("SetRayTracingBindingsInline"), PassParams, ERDGPassFlags::Copy | ERDGPassFlags::Compute | ERDGPassFlags::NeverCull,
+				[PassParams, &Allocator = Allocator, InlineRayTracingData = MoveTemp(InlineRayTracingData)](FRDGAsyncTask, FRHICommandList& RHICmdList) mutable
+				{
+					TRACE_CPUPROFILER_EVENT_SCOPE(SetRayTracingBindingsInline);
+
+					SetRaytracingShaderBindings(RHICmdList, Allocator, InlineRayTracingData);
+
+					RHICmdList.CommitShaderBindingTable(InlineRayTracingData.ShaderBindingTable, PassParams->InlineRayTracingBindingDataBuffer->GetRHI());
+				});
+		}
+	}
 }
 
 #endif // RHI_RAYTRACING
diff --git a/Engine/Source/Runtime/Renderer/Private/SceneRendering.h b/Engine/Source/Runtime/Renderer/Private/SceneRendering.h
--- a/Engine/Source/Runtime/Renderer/Private/SceneRendering.h
+++ b/Engine/Source/Runtime/Renderer/Private/SceneRendering.h
@@ -1555,7 +1555,6 @@ public:
 
 	FRayTracingData MaterialRayTracingData;
 	FRayTracingData LumenRayTracingData;
-	FRayTracingData InlineRayTracingData;	
 
 	// Buffer to the shader binding data used for inline raytracing - only valid when inline raytracing is enabled
 	FRDGBufferRef InlineRayTracingBindingDataBuffer = nullptr;
