Automatically disable Software VRS when Nanite is not enabled.
Previously, it was possible for the VRS image manager to request Software Shading Rate Images that would not be used anywhere if r.Nanite=0 but r.Nanite.SoftwareVRS=1.

diff --git a/Engine/Source/Runtime/Renderer/Private/VariableRateShading/VariableRateShadingImageManager.cpp b/Engine/Source/Runtime/Renderer/Private/VariableRateShading/VariableRateShadingImageManager.cpp
--- a/Engine/Source/Runtime/Renderer/Private/VariableRateShading/VariableRateShadingImageManager.cpp
+++ b/Engine/Source/Runtime/Renderer/Private/VariableRateShading/VariableRateShadingImageManager.cpp
@@ -523,6 +523,7 @@ void FVariableRateShadingImageManager::PrepareImageBasedVRS(FRDGBuilder& GraphBu
 
 	bSoftwareVRSEnabledForFrame = 
 		CVarEnableVRSSoftwareImage.GetValueOnRenderThread() > 0 &&
+		UseNanite(ShaderPlatform) &&
 		CVarLocalNaniteSoftwareVRS->GetInt() > 0 &&
 		IsFeatureLevelSupported(ShaderPlatform, ERHIFeatureLevel::SM6);
 
