Fixed: Issues wtih a PSO pre-caching for a UI and LightFunction materials by a mobile base pass

diff --git a/Engine/Source/Runtime/Renderer/Private/MobileBasePass.cpp b/Engine/Source/Runtime/Renderer/Private/MobileBasePass.cpp
--- a/Engine/Source/Runtime/Renderer/Private/MobileBasePass.cpp
+++ b/Engine/Source/Runtime/Renderer/Private/MobileBasePass.cpp
@@ -1014,6 +1014,14 @@ static void SetupMultiViewInfo(FGraphicsPipelineRenderTargetsInfo& RenderTargets
 
 void FMobileBasePassMeshProcessor::CollectPSOInitializers(const FSceneTexturesConfig& SceneTexturesConfig, const FMaterial& Material, const FPSOPrecacheVertexFactoryData& VertexFactoryData, const FPSOPrecacheParams& PreCacheParams, TArray<FPSOPrecacheData>& PSOInitializers)
 {
+	EMaterialDomain MaterialDomain = Material.GetMaterialDomain();
+	// Precache only Surface and Decal materials (other domains are handled separately) for BasePass rendering
+	const bool bPrecacheMaterialDomain = (MaterialDomain == MD_Surface || MaterialDomain == MD_DeferredDecal);
+	if (!bPrecacheMaterialDomain)
+	{
+		return;
+	}
+		
 	if (bTranslucentBasePass)
 	{
 		static IConsoleVariable* PSOPrecacheTranslucencyAllPass = IConsoleManager::Get().FindConsoleVariable(TEXT("r.PSOPrecache.TranslucencyAllPass"));
