Export FPersistentBuffer/FBufferScatterUploader methods for access outside of the Renderer module

#rnx

diff --git a/Engine/Source/Runtime/Renderer/Private/RendererPrivateUtils.h b/Engine/Source/Runtime/Renderer/Private/RendererPrivateUtils.h
--- a/Engine/Source/Runtime/Renderer/Private/RendererPrivateUtils.h
+++ b/Engine/Source/Runtime/Renderer/Private/RendererPrivateUtils.h
@@ -110,8 +110,8 @@ public:
 		int32 NumScatters;
 	};
 
-	void UploadTo(FRDGBuilder& GraphBuilder, FRDGBuffer *DestBuffer, FRDGBuffer *ScatterOffsets, FRDGBuffer *Values, uint32 NumScatters, uint32 NumBytesPerElement, int32 NumValuesPerScatter);
-	void UploadTo(FRDGBuilder& GraphBuilder, FRDGBuffer* DestBuffer, FRDGBuffer* ScatterOffsets, FRDGBuffer* Values, TFunction<uint64()>&& GetNumScatters, uint32 NumBytesPerElement, int32 NumValuesPerScatter);
+	RENDERER_API void UploadTo(FRDGBuilder& GraphBuilder, FRDGBuffer *DestBuffer, FRDGBuffer *ScatterOffsets, FRDGBuffer *Values, uint32 NumScatters, uint32 NumBytesPerElement, int32 NumValuesPerScatter);
+	RENDERER_API void UploadTo(FRDGBuilder& GraphBuilder, FRDGBuffer* DestBuffer, FRDGBuffer* ScatterOffsets, FRDGBuffer* Values, TFunction<uint64()>&& GetNumScatters, uint32 NumBytesPerElement, int32 NumValuesPerScatter);
 };
 
 namespace UE::RendererPrivateUtils::Implementation
@@ -123,17 +123,17 @@ namespace UE::RendererPrivateUtils::Implementation
 class FPersistentBuffer
 {
 public:
-	FPersistentBuffer(int32 InMinimumNumElementsReserved, const TCHAR *InName, bool bInRoundUpToPOT = true);
+	RENDERER_API FPersistentBuffer(int32 InMinimumNumElementsReserved, const TCHAR *InName, bool bInRoundUpToPOT = true);
 
-	FRDGBuffer* Register(FRDGBuilder& GraphBuilder) const;
+	RENDERER_API FRDGBuffer* Register(FRDGBuilder& GraphBuilder) const;
 
-	void Empty();
+	RENDERER_API void Empty();
 
 protected:
 
-	FRDGBuffer* ResizeBufferIfNeeded(FRDGBuilder& GraphBuilder, const FRDGBufferDesc& BufferDesc);
+	RENDERER_API FRDGBuffer* ResizeBufferIfNeeded(FRDGBuilder& GraphBuilder, const FRDGBufferDesc& BufferDesc);
 
-	FRDGBuffer* ResizeAndClearBufferIfNeeded(FRDGBuilder& GraphBuilder, const FRDGBufferDesc& BufferDesc);
+	RENDERER_API FRDGBuffer* ResizeAndClearBufferIfNeeded(FRDGBuilder& GraphBuilder, const FRDGBufferDesc& BufferDesc);
 
 	int32 MinimumNumElementsReserved = 0;
 	const TCHAR *Name = nullptr;
