[HWRT] Use one SBT per view when shader binding layout is not supported.
- in this case we need to store ViewUniformBuffer in SBT entries so can't share a single SBT for all views.

diff --git a/Engine/Source/Runtime/Renderer/Private/DeferredShadingRenderer.cpp b/Engine/Source/Runtime/Renderer/Private/DeferredShadingRenderer.cpp
--- a/Engine/Source/Runtime/Renderer/Private/DeferredShadingRenderer.cpp
+++ b/Engine/Source/Runtime/Renderer/Private/DeferredShadingRenderer.cpp
@@ -728,76 +728,140 @@ bool FDeferredShadingSceneRenderer::SetupRayTracingPipelineStatesAndSBT(FRDGBuil
 			// Need RTPSO
 			EnumAddFlags(ShaderBindingMode, ERayTracingShaderBindingMode::RTPSO);
 
-			// Allocate, or check if the persistent SBT RHI is still valid
-			if (Scene->MaterialRayTracingSBTID == INDEX_NONE)
+			// Allocate persistent SBT RHI if necessary
+			// One SBT per View if ShaderBindingLayout is not supported, otherwise a shared SBT stored in the Scene
+			if (RHIGetStaticShaderBindingLayoutSupport(ShaderPlatform) == ERHIStaticShaderBindingLayoutSupport::Unsupported)
+			{
+				for (FViewInfo& View : Views)
+				{
+					if (View.ViewState != nullptr)
+					{
+						FScene::FViewRayTracingData& ViewRayTracingData = Scene->GetViewRayTracingData(View);
+
+						if (ViewRayTracingData.MaterialRayTracingSBTID == INDEX_NONE)
+						{
+							ViewRayTracingData.MaterialRayTracingSBTID = Scene->RayTracingSBT.AllocatePersistentSBTID(GraphBuilder.RHICmdList, ERayTracingShaderBindingMode::RTPSO);
+						}
+					}
+				}
+			}
+			else if (Scene->MaterialRayTracingSBTID == INDEX_NONE)
 			{
 				Scene->MaterialRayTracingSBTID = Scene->RayTracingSBT.AllocatePersistentSBTID(GraphBuilder.RHICmdList, ERayTracingShaderBindingMode::RTPSO);
 			}
 		}
-		else if (Scene->MaterialRayTracingSBTID != INDEX_NONE)
+		else
 		{
-			Scene->RayTracingSBT.ReleasePersistentSBT(Scene->MaterialRayTracingSBTID);
-			Scene->MaterialRayTracingSBTID = INDEX_NONE;
-		}
-	}
+			// release material persistent SBT when not needed
 
-	// Create Lumen hardware ray tracing SBT and material pipeline
-	{
-		ERayTracingShaderBindingMode LumenShaderBindingMode = ERayTracingShaderBindingMode::Disabled;
-
-		if (!bIsPathTracing)
-		{
-			TArray<FRHIRayTracingShader*> LumenHardwareRayTracingRayGenShaders;
+			if (Scene->MaterialRayTracingSBTID != INDEX_NONE)
+			{
+				Scene->RayTracingSBT.ReleasePersistentSBT(Scene->MaterialRayTracingSBTID);
+				Scene->MaterialRayTracingSBTID = INDEX_NONE;
+			}
 
-			if (GRHISupportsRayTracingShaders)
+			for (FViewInfo& View : Views)
 			{
-				if (DoesPlatformSupportLumenGI(ShaderPlatform))
+				if (View.ViewState != nullptr)
 				{
-					for (const FViewInfo& View : Views)
+					FScene::FViewRayTracingData& ViewRayTracingData = Scene->GetViewRayTracingData(View);
+
+					if (ViewRayTracingData.MaterialRayTracingSBTID != INDEX_NONE)
 					{
-						PrepareLumenHardwareRayTracingVisualizeLumenMaterial(View, LumenHardwareRayTracingRayGenShaders);
-						PrepareLumenHardwareRayTracingRadianceCacheLumenMaterial(View, LumenHardwareRayTracingRayGenShaders);
-						PrepareLumenHardwareRayTracingTranslucencyVolumeLumenMaterial(View, LumenHardwareRayTracingRayGenShaders);
-						PrepareLumenHardwareRayTracingRadiosityLumenMaterial(View, LumenHardwareRayTracingRayGenShaders);
-						PrepareLumenHardwareRayTracingReflectionsLumenMaterial(View, LumenHardwareRayTracingRayGenShaders);
-						PrepareLumenHardwareRayTracingReSTIRLumenMaterial(View, LumenHardwareRayTracingRayGenShaders);
-						PrepareLumenHardwareRayTracingScreenProbeGatherLumenMaterial(View, LumenHardwareRayTracingRayGenShaders);
-						PrepareLumenHardwareRayTracingDirectLightingLumenMaterial(View, LumenHardwareRayTracingRayGenShaders);
+						Scene->RayTracingSBT.ReleasePersistentSBT(ViewRayTracingData.MaterialRayTracingSBTID);
+						ViewRayTracingData.MaterialRayTracingSBTID = INDEX_NONE;
 					}
+
 				}
+			}
+		}
+	}
+
+	if (GRHISupportsRayTracingShaders)
+	{
+		// Create Lumen hardware ray tracing SBT and material pipeline
+
+		TArray<FRHIRayTracingShader*> LumenHardwareRayTracingRayGenShaders;
 
+		if (!bIsPathTracing)
+		{
+			if (DoesPlatformSupportLumenGI(ShaderPlatform))
+			{
 				for (const FViewInfo& View : Views)
 				{
-					PrepareMegaLightsHardwareRayTracingLumenMaterial(View, *Scene, LumenHardwareRayTracingRayGenShaders);
+					PrepareLumenHardwareRayTracingVisualizeLumenMaterial(View, LumenHardwareRayTracingRayGenShaders);
+					PrepareLumenHardwareRayTracingRadianceCacheLumenMaterial(View, LumenHardwareRayTracingRayGenShaders);
+					PrepareLumenHardwareRayTracingTranslucencyVolumeLumenMaterial(View, LumenHardwareRayTracingRayGenShaders);
+					PrepareLumenHardwareRayTracingRadiosityLumenMaterial(View, LumenHardwareRayTracingRayGenShaders);
+					PrepareLumenHardwareRayTracingReflectionsLumenMaterial(View, LumenHardwareRayTracingRayGenShaders);
+					PrepareLumenHardwareRayTracingReSTIRLumenMaterial(View, LumenHardwareRayTracingRayGenShaders);
+					PrepareLumenHardwareRayTracingScreenProbeGatherLumenMaterial(View, LumenHardwareRayTracingRayGenShaders);
+					PrepareLumenHardwareRayTracingDirectLightingLumenMaterial(View, LumenHardwareRayTracingRayGenShaders);
 				}
-
-				DeduplicateRayGenerationShaders(LumenHardwareRayTracingRayGenShaders);
 			}
 
-			if (LumenHardwareRayTracingRayGenShaders.Num())
+			for (const FViewInfo& View : Views)
 			{
-				CreateLumenHardwareRayTracingMaterialPipeline(GraphBuilder, LumenHardwareRayTracingRayGenShaders, MaxLocalBindingDataSize);
-				EnumAddFlags(LumenShaderBindingMode, ERayTracingShaderBindingMode::RTPSO);
+				PrepareMegaLightsHardwareRayTracingLumenMaterial(View, *Scene, LumenHardwareRayTracingRayGenShaders);
 			}
 		}
 
-		if (LumenShaderBindingMode != ERayTracingShaderBindingMode::Disabled)
+		DeduplicateRayGenerationShaders(LumenHardwareRayTracingRayGenShaders);
+
+		if (LumenHardwareRayTracingRayGenShaders.Num())
 		{
-			// Allocate, or check if the persistent SBT RHI is still valid
-			if (Scene->LumenRayTracingSBTID == INDEX_NONE)
+			CreateLumenHardwareRayTracingMaterialPipeline(GraphBuilder, LumenHardwareRayTracingRayGenShaders, MaxLocalBindingDataSize);
+
+			// Need RTPSO
+			EnumAddFlags(ShaderBindingMode, ERayTracingShaderBindingMode::RTPSO);
+
+			// Allocate persistent SBT RHI if necessary
+			// One SBT per View if ShaderBindingLayout is not supported, otherwise a shared SBT stored in the Scene
+			if (RHIGetStaticShaderBindingLayoutSupport(ShaderPlatform) == ERHIStaticShaderBindingLayoutSupport::Unsupported)
+			{
+				for (FViewInfo& View : Views)
+				{
+					if (View.ViewState != nullptr)
+					{
+						FScene::FViewRayTracingData& ViewRayTracingData = Scene->GetViewRayTracingData(View);
+
+						if (ViewRayTracingData.LumenRayTracingSBTID == INDEX_NONE)
+						{
+							ViewRayTracingData.LumenRayTracingSBTID = Scene->RayTracingSBT.AllocatePersistentSBTID(GraphBuilder.RHICmdList, ERayTracingShaderBindingMode::RTPSO);
+						}
+					}
+				}
+			}
+			else if (Scene->LumenRayTracingSBTID == INDEX_NONE)
 			{
-				Scene->LumenRayTracingSBTID = Scene->RayTracingSBT.AllocatePersistentSBTID(GraphBuilder.RHICmdList, LumenShaderBindingMode);
+				Scene->LumenRayTracingSBTID = Scene->RayTracingSBT.AllocatePersistentSBTID(GraphBuilder.RHICmdList, ERayTracingShaderBindingMode::RTPSO);
 			}
 		}
-		else if (Scene->LumenRayTracingSBTID != INDEX_NONE)
+		else
 		{
 			// release Lumen persistent SBT when not needed
-			Scene->RayTracingSBT.ReleasePersistentSBT(Scene->LumenRayTracingSBTID);
-			Scene->LumenRayTracingSBTID = INDEX_NONE;
-		}
 
-		// Merge Lumen shader binding mode with the shared shader binding mode
-		EnumAddFlags(ShaderBindingMode, LumenShaderBindingMode);
+			if (Scene->LumenRayTracingSBTID != INDEX_NONE)
+			{
+				Scene->RayTracingSBT.ReleasePersistentSBT(Scene->LumenRayTracingSBTID);
+				Scene->LumenRayTracingSBTID = INDEX_NONE;
+			}
+
+			for (FViewInfo& View : Views)
+			{
+				if (View.ViewState != nullptr)
+				{
+					FScene::FViewRayTracingData& ViewRayTracingData = Scene->GetViewRayTracingData(View);
+
+					if (ViewRayTracingData.LumenRayTracingSBTID != INDEX_NONE)
+					{
+						Scene->RayTracingSBT.ReleasePersistentSBT(ViewRayTracingData.LumenRayTracingSBTID);
+						ViewRayTracingData.LumenRayTracingSBTID = INDEX_NONE;
+					}
+
+				}
+			}
+		}
 	}
 
 	// Check if inline SBT is needed or not			
@@ -816,38 +880,71 @@ bool FDeferredShadingSceneRenderer::SetupRayTracingPipelineStatesAndSBT(FRDGBuil
 		Scene->InlineRayTracingSBTID = INDEX_NONE;
 	}
 
-	FRHIShaderBindingTable* MaterialSBT = nullptr;
-	FRHIShaderBindingTable* LumenSBT = nullptr;
-	FRHIShaderBindingTable* InlineSBT = nullptr;
-
-	FRDGBufferRef InlineBindingDataBuffer = nullptr;
-
 	if (ShaderBindingMode != ERayTracingShaderBindingMode::Disabled)
 	{
 		Scene->RayTracingSBT.CheckPersistentRHI(GraphBuilder.RHICmdList, MaxLocalBindingDataSize);
 
-		MaterialSBT = Scene->MaterialRayTracingSBTID != INDEX_NONE ? Scene->RayTracingSBT.GetPersistentSBT(Scene->MaterialRayTracingSBTID) : nullptr;
-		LumenSBT = Scene->LumenRayTracingSBTID != INDEX_NONE ? Scene->RayTracingSBT.GetPersistentSBT(Scene->LumenRayTracingSBTID) : nullptr;
+		if (RHIGetStaticShaderBindingLayoutSupport(ShaderPlatform) != ERHIStaticShaderBindingLayoutSupport::Unsupported)
+		{
+			// share SBT across all views
 
-		if (Scene->InlineRayTracingSBTID != INDEX_NONE)
+			FRHIShaderBindingTable* MaterialSBT = Scene->MaterialRayTracingSBTID != INDEX_NONE ? Scene->RayTracingSBT.GetPersistentSBT(Scene->MaterialRayTracingSBTID) : nullptr;
+			FRHIShaderBindingTable* LumenSBT = Scene->LumenRayTracingSBTID != INDEX_NONE ? Scene->RayTracingSBT.GetPersistentSBT(Scene->LumenRayTracingSBTID) : nullptr;
+
+			bool bFirstViewWithRayTracingPass = true;
+
+			for (int32 ViewIndex = 0; ViewIndex < Views.Num(); ++ViewIndex)
+			{
+				FViewInfo& View = Views[ViewIndex];
+
+				if (View.bHasAnyRayTracingPass)
+				{
+					View.MaterialRayTracingData.ShaderBindingTable = MaterialSBT;
+					View.LumenRayTracingData.ShaderBindingTable = LumenSBT;
+					View.bOwnsShaderBindingTables = bFirstViewWithRayTracingPass;
+
+					bFirstViewWithRayTracingPass = false;
+				}
+			}
+		}
+		else
 		{
-			InlineSBT = Scene->RayTracingSBT.GetPersistentSBT(Scene->InlineRayTracingSBTID);
-			InlineBindingDataBuffer = Scene->RayTracingSBT.GetPersistentInlineBindingDataBuffer(GraphBuilder, Scene->InlineRayTracingSBTID);
+			// need one SBT per View
+			for (FViewInfo& View : Views)
+			{
+				if (!View.bHasAnyRayTracingPass)
+				{
+					continue;
+				}
+
+				if (View.ViewState != nullptr)
+				{
+					FScene::FViewRayTracingData& ViewRayTracingData = Scene->GetViewRayTracingData(View);
+
+					View.MaterialRayTracingData.ShaderBindingTable = ViewRayTracingData.MaterialRayTracingSBTID != INDEX_NONE ? Scene->RayTracingSBT.GetPersistentSBT(ViewRayTracingData.MaterialRayTracingSBTID) : nullptr;
+					View.LumenRayTracingData.ShaderBindingTable = ViewRayTracingData.LumenRayTracingSBTID != INDEX_NONE ? Scene->RayTracingSBT.GetPersistentSBT(ViewRayTracingData.LumenRayTracingSBTID) : nullptr;
+				}
+				else
+				{
+					ensureMsgf(false, TEXT("Valid ViewState is required when using view-specific SBT. Should use a transient SBT instead?"));
+				}
+			}
 		}
-	}
 
-	// Send SBTs to all views since they all share the same ones
-	EnumerateLinkedViews([MaterialSBT, LumenSBT, InlineBindingDataBuffer] (FViewInfo& View)
-	{
-		if (View.bHasAnyRayTracingPass)
+		if (Scene->InlineRayTracingSBTID != INDEX_NONE)
 		{
-			View.MaterialRayTracingData.ShaderBindingTable = MaterialSBT;
-			View.LumenRayTracingData.ShaderBindingTable = LumenSBT;
+			// set the same InlineBindingDataBuffer on all views
+			FRDGBufferRef InlineBindingDataBuffer = Scene->RayTracingSBT.GetPersistentInlineBindingDataBuffer(GraphBuilder, Scene->InlineRayTracingSBTID);
 
-			View.InlineRayTracingBindingDataBuffer = InlineBindingDataBuffer;
+			for (FViewInfo& View : Views)
+			{
+				if (View.bHasAnyRayTracingPass)
+				{
+					View.InlineRayTracingBindingDataBuffer = InlineBindingDataBuffer;
+				}
+			}
 		}
-		return true;
-	});
+	}
 
 	return true;
 }
@@ -963,6 +1060,8 @@ void FDeferredShadingSceneRenderer::SetupRayTracingRenderingData(FRDGBuilder& Gr
 	// Keep mask the same as what's already set (which will be the view mask) if TLAS updates should be masked to the view
 	RDG_GPU_MASK_SCOPE(GraphBuilder, GRayTracingMultiGpuTLASMask ? GraphBuilder.RHICmdList.GetGPUMask() : FRHIGPUMask::All());
 
+	const bool bIsPathTracing = ViewFamily.EngineShowFlags.PathTracing;
+
 	bool bAnyInlineHardwareRayTracingPassEnabled = false;
 
 	for (FViewInfo& View : Views)
@@ -1020,14 +1119,20 @@ void FDeferredShadingSceneRenderer::SetupRayTracingRenderingData(FRDGBuilder& Gr
 
 		View.AddDynamicRayTracingMeshBatchTaskList.Empty();
 	}
-	
-	// Build the dirty persistent shader bindings from the visible shader bindings (SBT version is updated after the PSO and SBTs have been created)
-	// Inline lumen buffer bindings don't use the dirty bindings yet - so don't need to build it here (uses the visible binding set)
-	bool bRequireBindingsUpdate = ReferenceView.MaterialRayTracingData.ShaderBindingTable || ReferenceView.LumenRayTracingData.ShaderBindingTable || Scene->InlineRayTracingSBTID != INDEX_NONE;
+
+	bool bRequireBindingsUpdate = false;
+
+	for (const FViewInfo& View : Views)
+	{
+		bRequireBindingsUpdate |= View.MaterialRayTracingData.ShaderBindingTable || View.LumenRayTracingData.ShaderBindingTable || Scene->InlineRayTracingSBTID != INDEX_NONE;
+	}
+
 	if (bRequireBindingsUpdate)
 	{
 		// If fallback RTPSO then mark all bindings as dirty because they need to bound again when final RTPSO is ready (Shader identifier could have changed)
-		bool bForceAllDirty = bIsUsingFallbackRTPSO;
+		const bool bForceAllDirty = bIsUsingFallbackRTPSO;
+
+		// Build the dirty persistent shader bindings from the visible shader bindings (SBT version is updated after the PSO and SBTs have been created)
 		ReferenceView.DirtyPersistentRayTracingShaderBindings = Scene->RayTracingSBT.GetDirtyBindings(ReferenceView.VisibleRayTracingShaderBindings, bForceAllDirty);
 	}
 
@@ -1040,79 +1145,79 @@ void FDeferredShadingSceneRenderer::SetupRayTracingRenderingData(FRDGBuilder& Gr
 		}
 	);
 
-	// Prepare the local ray tracing shader binding data to update on RHI timeline for Material, Lumen and InlineLumen
-	if (ReferenceView.MaterialRayTracingData.ShaderBindingTable)
-	{
-		SetupMaterialRayTracingHitGroupBindings(GraphBuilder, ReferenceView);
-	}
-	if (ReferenceView.LumenRayTracingData.ShaderBindingTable)
-	{
-		SetupLumenHardwareRayTracingHitGroupBindings(GraphBuilder, ReferenceView);	
-	}
-
 	FRDGBufferRef LumenHardwareRayTracingHitDataBuffer = nullptr;
 	if (bAnyInlineHardwareRayTracingPassEnabled)
 	{
+		// TODO: Could have a persistent HardwareRayTracingHitDataBuffer and update using DirtyPersistentRayTracingShaderBindings
+		// instead of always recreating the buffer using VisibleRayTracingShaderBindings
 		LumenHardwareRayTracingHitDataBuffer = SetupLumenHardwareRayTracingHitGroupBuffer(GraphBuilder, ReferenceView.VisibleRayTracingShaderBindings);
 	}
 
 	for (FViewInfo& View : Views)
 	{
-		if (View.bHasAnyRayTracingPass)
+		View.LumenHardwareRayTracingHitDataBuffer = LumenHardwareRayTracingHitDataBuffer;
+
+		if (!View.bOwnsShaderBindingTables)
 		{
-			View.LumenHardwareRayTracingHitDataBuffer = LumenHardwareRayTracingHitDataBuffer;
+			continue;
 		}
-	}
 
-	const bool bIsPathTracing = ViewFamily.EngineShowFlags.PathTracing;
+		// Prepare the local ray tracing shader binding data to update on RHI timeline for Material, Lumen
+		if (View.MaterialRayTracingData.ShaderBindingTable)
+		{
+			SetupMaterialRayTracingHitGroupBindings(GraphBuilder, View);
+		}
 
-	{
-		const FRayTracingLightFunctionMap* RayTracingLightFunctionMap = GraphBuilder.Blackboard.Get<FRayTracingLightFunctionMap>();
+		if (View.LumenRayTracingData.ShaderBindingTable)
+		{
+			SetupLumenHardwareRayTracingHitGroupBindings(GraphBuilder, View);
+		}
 
 		FSetRayTracingBindingsPassParams* PassParams = GraphBuilder.AllocParameters<FSetRayTracingBindingsPassParams>();
 		PassParams->Scene = GetSceneUniformBufferRef(GraphBuilder);
-		PassParams->LightGridPacked = bIsPathTracing ? nullptr : ReferenceView.RayTracingLightGridUniformBuffer; // accessed by FRayTracingLightingMS // Is this needed for anything?
-		PassParams->LumenHardwareRayTracingUniformBuffer = ReferenceView.LumenHardwareRayTracingUniformBuffer;
+		PassParams->LightGridPacked = bIsPathTracing ? nullptr : View.RayTracingLightGridUniformBuffer; // accessed by FRayTracingLightingMS // Is this needed for anything?
+		PassParams->LumenHardwareRayTracingUniformBuffer = View.LumenHardwareRayTracingUniformBuffer;
 
+		const FRayTracingLightFunctionMap* RayTracingLightFunctionMap = GraphBuilder.Blackboard.Get<FRayTracingLightFunctionMap>();
 		GraphBuilder.AddPass(RDG_EVENT_NAME("SetRayTracingBindings"), PassParams, ERDGPassFlags::Copy | ERDGPassFlags::Compute | ERDGPassFlags::NeverCull,
-			[this, PassParams, bIsPathTracing, &ReferenceView, RayTracingLightFunctionMap](FRDGAsyncTask, FRHICommandList& RHICmdList)
+			[this, PassParams, bIsPathTracing, &View, RayTracingLightFunctionMap](FRDGAsyncTask, FRHICommandList& RHICmdList)
 			{
 				TRACE_CPUPROFILER_EVENT_SCOPE(SetRayTracingBindings);
 
-				check(ReferenceView.MaterialRayTracingData.PipelineState || ReferenceView.MaterialRayTracingData.MaterialBindings.Num() == 0);
+				check(View.MaterialRayTracingData.PipelineState || View.MaterialRayTracingData.MaterialBindings.Num() == 0);
 
-				if (ReferenceView.MaterialRayTracingData.PipelineState && (ReferenceView.MaterialRayTracingData.MaterialBindings.Num() || ReferenceView.MaterialRayTracingData.CallableBindings.Num()))
+				if (View.MaterialRayTracingData.PipelineState && (View.MaterialRayTracingData.MaterialBindings.Num() || View.MaterialRayTracingData.CallableBindings.Num()))
 				{
-					SetRaytracingShaderBindings(RHICmdList, Allocator, ReferenceView.MaterialRayTracingData);
+					SetRayTracingShaderBindings(RHICmdList, Allocator, View.MaterialRayTracingData);
 
 					if (bIsPathTracing)
 					{
-						SetupPathTracingDefaultMissShader(RHICmdList, ReferenceView);
+						SetupPathTracingDefaultMissShader(RHICmdList, View);
 
-						BindLightFunctionShadersPathTracing(RHICmdList, Scene, RayTracingLightFunctionMap, ReferenceView);
+						BindLightFunctionShadersPathTracing(RHICmdList, Scene, RayTracingLightFunctionMap, View);
 					}
 					else
 					{
-						SetupRayTracingDefaultMissShader(RHICmdList, ReferenceView);
-						SetupRayTracingLightingMissShader(RHICmdList, ReferenceView);
+						SetupRayTracingDefaultMissShader(RHICmdList, View);
+						SetupRayTracingLightingMissShader(RHICmdList, View);
 
-						BindLightFunctionShaders(RHICmdList, Scene, RayTracingLightFunctionMap, ReferenceView);
+						BindLightFunctionShaders(RHICmdList, Scene, RayTracingLightFunctionMap, View);
 					}
 
-					RHICmdList.CommitShaderBindingTable(ReferenceView.MaterialRayTracingData.ShaderBindingTable);
+					RHICmdList.CommitShaderBindingTable(View.MaterialRayTracingData.ShaderBindingTable);
 				}
 
 				if (!bIsPathTracing)
 				{
-					if (ReferenceView.LumenRayTracingData.PipelineState && ReferenceView.LumenRayTracingData.ShaderBindingTable)
+					if (View.LumenRayTracingData.PipelineState && View.LumenRayTracingData.ShaderBindingTable)
 					{
-						RHICmdList.SetRayTracingMissShader(ReferenceView.LumenRayTracingData.ShaderBindingTable, RAY_TRACING_MISS_SHADER_SLOT_DEFAULT, ReferenceView.LumenRayTracingData.PipelineState, 0 /* MissShaderPipelineIndex */, 0, nullptr, 0);
+						RHICmdList.SetRayTracingMissShader(View.LumenRayTracingData.ShaderBindingTable, RAY_TRACING_MISS_SHADER_SLOT_DEFAULT, View.LumenRayTracingData.PipelineState, 0 /* MissShaderPipelineIndex */, 0, nullptr, 0);
 					}
 
-					if (ReferenceView.LumenRayTracingData.ShaderBindingTable)
+					if (View.LumenRayTracingData.ShaderBindingTable)
 					{
-						SetRaytracingShaderBindings(RHICmdList, Allocator, ReferenceView.LumenRayTracingData);
-						RHICmdList.CommitShaderBindingTable(ReferenceView.LumenRayTracingData.ShaderBindingTable);
+						SetRayTracingShaderBindings(RHICmdList, Allocator, View.LumenRayTracingData);
+						RHICmdList.CommitShaderBindingTable(View.LumenRayTracingData.ShaderBindingTable);
 					}
 				}
 			});
@@ -1123,6 +1228,7 @@ void FDeferredShadingSceneRenderer::SetupRayTracingRenderingData(FRDGBuilder& Gr
 		FViewInfo::FRayTracingData InlineRayTracingData;
 		InlineRayTracingData.ShaderBindingTable = Scene->RayTracingSBT.GetPersistentSBT(Scene->InlineRayTracingSBTID);
 
+		// Prepare the local ray tracing shader binding data to update on RHI timeline for Inline SBT
 		{
 			TRACE_CPUPROFILER_EVENT_SCOPE(FDeferredShadingSceneRenderer::SetupInlineHardwareRaytracingHitGroupBindings);
 
@@ -1153,7 +1259,7 @@ void FDeferredShadingSceneRenderer::SetupRayTracingRenderingData(FRDGBuilder& Gr
 				{
 					TRACE_CPUPROFILER_EVENT_SCOPE(SetRayTracingBindingsInline);
 
-					SetRaytracingShaderBindings(RHICmdList, Allocator, InlineRayTracingData);
+					SetRayTracingShaderBindings(RHICmdList, Allocator, InlineRayTracingData);
 
 					RHICmdList.CommitShaderBindingTable(InlineRayTracingData.ShaderBindingTable, PassParams->InlineRayTracingBindingDataBuffer->GetRHI());
 				});
@@ -1840,6 +1946,24 @@ void FDeferredShadingSceneRenderer::Render(FRDGBuilder& GraphBuilder, const FSce
 #if RHI_RAYTRACING
 		if (FamilyPipelineState[&FFamilyPipelineState::bRayTracing])
 		{
+			for (FViewInfo& View : Views)
+			{
+				if (View.ViewState != nullptr)
+				{
+					if (View.ViewState->Scene == nullptr)
+					{
+						// link view state to the scene
+						View.ViewState->Scene = Scene;
+						Scene->ViewStates.Add(View.ViewState);
+					}
+
+					if (!Scene->PerViewRayTracingData.Contains(View.ViewState->GetViewKey()))
+					{
+						Scene->PerViewRayTracingData.Emplace(View.ViewState->GetViewKey());
+					}
+				}
+			}
+
 			InitViewTaskDatas.RayTracingGatherInstances = RayTracing::CreateGatherInstancesTaskData(
 				Allocator,
 				*Scene,
diff --git a/Engine/Source/Runtime/Renderer/Private/RayTracing/RayTracingMaterialHitShaders.cpp b/Engine/Source/Runtime/Renderer/Private/RayTracing/RayTracingMaterialHitShaders.cpp
--- a/Engine/Source/Runtime/Renderer/Private/RayTracing/RayTracingMaterialHitShaders.cpp
+++ b/Engine/Source/Runtime/Renderer/Private/RayTracing/RayTracingMaterialHitShaders.cpp
@@ -1070,7 +1070,7 @@ void MergeAndSetRayTracingBindings(
 		bCopyDataToInlineStorage);
 }
 
-void SetRaytracingShaderBindings(FRHICommandList& RHICmdList, FSceneRenderingBulkObjectAllocator& Allocator, FViewInfo::FRayTracingData& RayTracingData)
+void SetRayTracingShaderBindings(FRHICommandList& RHICmdList, FSceneRenderingBulkObjectAllocator& Allocator, FViewInfo::FRayTracingData& RayTracingData)
 {
 	if (!RayTracingData.MaterialBindings.IsEmpty())
 	{
diff --git a/Engine/Source/Runtime/Renderer/Private/RayTracing/RayTracingMaterialHitShaders.h b/Engine/Source/Runtime/Renderer/Private/RayTracing/RayTracingMaterialHitShaders.h
--- a/Engine/Source/Runtime/Renderer/Private/RayTracing/RayTracingMaterialHitShaders.h
+++ b/Engine/Source/Runtime/Renderer/Private/RayTracing/RayTracingMaterialHitShaders.h
@@ -328,6 +328,6 @@ void AddRayTracingLocalShaderBindingWriterTasks(
 	}
 }
 
-void SetRaytracingShaderBindings(FRHICommandList& RHICmdList, FSceneRenderingBulkObjectAllocator& Allocator, FViewInfo::FRayTracingData& RayTracingData);
+void SetRayTracingShaderBindings(FRHICommandList& RHICmdList, FSceneRenderingBulkObjectAllocator& Allocator, FViewInfo::FRayTracingData& RayTracingData);
 
 #endif // RHI_RAYTRACING
diff --git a/Engine/Source/Runtime/Renderer/Private/RendererScene.cpp b/Engine/Source/Runtime/Renderer/Private/RendererScene.cpp
--- a/Engine/Source/Runtime/Renderer/Private/RendererScene.cpp
+++ b/Engine/Source/Runtime/Renderer/Private/RendererScene.cpp
@@ -452,6 +452,25 @@ void FScene::RemoveViewState_RenderThread(FSceneViewStateInterface* ViewState)
 	}
 
 	RemoveViewLumenSceneData_RenderThread(ViewState);
+
+#if RHI_RAYTRACING
+	{
+		const uint32 ViewKey = ViewState->GetViewKey();
+		FViewRayTracingData const* Found = PerViewRayTracingData.Find(ViewKey);
+		if (Found)
+		{
+			if (Found->MaterialRayTracingSBTID != INDEX_NONE)
+			{
+				RayTracingSBT.ReleasePersistentSBT(Found->MaterialRayTracingSBTID);
+			}
+			if (Found->LumenRayTracingSBTID != INDEX_NONE)
+			{
+				RayTracingSBT.ReleasePersistentSBT(Found->LumenRayTracingSBTID);
+			}
+			PerViewRayTracingData.Remove(ViewKey);
+		}
+	}
+#endif // RHI_RAYTRACING
 }
 
 
diff --git a/Engine/Source/Runtime/Renderer/Private/ScenePrivate.h b/Engine/Source/Runtime/Renderer/Private/ScenePrivate.h
--- a/Engine/Source/Runtime/Renderer/Private/ScenePrivate.h
+++ b/Engine/Source/Runtime/Renderer/Private/ScenePrivate.h
@@ -1098,7 +1098,7 @@ public:
 
 	virtual void SetLandscapeRayTracingStates(TPimplPtr<FLandscapeRayTracingStateList>&& InLandscapeRayTracingStates) final { LandscapeRayTracingStates = MoveTemp(InLandscapeRayTracingStates); }
 	virtual FLandscapeRayTracingStateList* GetLandscapeRayTracingStates() const final { return LandscapeRayTracingStates.Get(); }
-#endif
+#endif // RHI_RAYTRACING
 
 	TUniquePtr<FForwardLightingViewResources> ForwardLightingResources;
 
@@ -3195,6 +3195,14 @@ public:
 	FRayTracingPersistentShaderBindingTableID MaterialRayTracingSBTID = INDEX_NONE;
 	FRayTracingPersistentShaderBindingTableID LumenRayTracingSBTID = INDEX_NONE;
 	FRayTracingPersistentShaderBindingTableID InlineRayTracingSBTID = INDEX_NONE;
+
+	struct FViewRayTracingData
+	{
+		FRayTracingPersistentShaderBindingTableID MaterialRayTracingSBTID = INDEX_NONE;
+		FRayTracingPersistentShaderBindingTableID LumenRayTracingSBTID = INDEX_NONE;
+	};
+
+	TMap<uint32, FViewRayTracingData> PerViewRayTracingData;
 #endif // RHI_RAYTRACING
 
 	/** Distance field object scene data. */
@@ -3526,6 +3534,14 @@ public:
 			return FindLumenSceneData(View.State ? ((const FSceneViewState*)View.State)->GetShareOriginViewKey() : 0, View.GPUMask.GetFirstIndex());
 		}
 	}
+
+#if RHI_RAYTRACING
+	FViewRayTracingData& GetViewRayTracingData(const FViewInfo& View)
+	{
+		return PerViewRayTracingData[View.State->GetViewKey()];
+	}
+#endif // RHI_RAYTRACING
+
 	virtual void AddPrimitive(FPrimitiveSceneDesc* Primitive) override;
 	virtual void RemovePrimitive(FPrimitiveSceneDesc* Primitive) override;
 	virtual void ReleasePrimitive(FPrimitiveSceneDesc* Primitive) override;
diff --git a/Engine/Source/Runtime/Renderer/Private/SceneRendering.h b/Engine/Source/Runtime/Renderer/Private/SceneRendering.h
--- a/Engine/Source/Runtime/Renderer/Private/SceneRendering.h
+++ b/Engine/Source/Runtime/Renderer/Private/SceneRendering.h
@@ -1555,6 +1555,7 @@ public:
 
 	FRayTracingData MaterialRayTracingData;
 	FRayTracingData LumenRayTracingData;
+	bool bOwnsShaderBindingTables = true;
 
 	// Buffer to the shader binding data used for inline raytracing - only valid when inline raytracing is enabled
 	FRDGBufferRef InlineRayTracingBindingDataBuffer = nullptr;
