Stop replacing LoadAction for ViewFamilyOutput in AddMobilePostProcessingPasses

Let FScreenPassRenderTarget::CreateViewFamilyOutput do its job, so the mobile path can take advantage of ERenderTargetLoadAction::ENoAction. Clears can be expensive on some hardware taking the mobile path on desktop.

diff --git a/Engine/Source/Runtime/Renderer/Private/PostProcess/PostProcessing.cpp b/Engine/Source/Runtime/Renderer/Private/PostProcess/PostProcessing.cpp
--- a/Engine/Source/Runtime/Renderer/Private/PostProcess/PostProcessing.cpp
+++ b/Engine/Source/Runtime/Renderer/Private/PostProcess/PostProcessing.cpp
@@ -3151,10 +3151,8 @@ void AddMobilePostProcessingPasses(FRDGBuilder& GraphBuilder, FScene* Scene, con
 				{
 					OutputViewRect = View.UnscaledViewRect;
 				}
-				ERenderTargetLoadAction  OutputLoadAction = View.IsFirstInFamily() ? ERenderTargetLoadAction::EClear : ERenderTargetLoadAction::ELoad;
 
 				TonemapperInputs.OverrideOutput.ViewRect = OutputViewRect;
-				TonemapperInputs.OverrideOutput.LoadAction = OutputLoadAction;
 				TonemapperInputs.OverrideOutput.UpdateVisualizeTextureExtent();
 			}
 
@@ -3219,7 +3217,6 @@ void AddMobilePostProcessingPasses(FRDGBuilder& GraphBuilder, FScene* Scene, con
 		HighResolutionScreenshotMaskInputs.MaskMaterial = View.FinalPostProcessSettings.HighResScreenshotMaskMaterial;
 		HighResolutionScreenshotMaskInputs.CaptureRegionMaterial = View.FinalPostProcessSettings.HighResScreenshotCaptureRegionMaterial;
 		PassSequence.AcceptOverrideIfLastPass(EPass::HighResolutionScreenshotMask, HighResolutionScreenshotMaskInputs.OverrideOutput);
-		HighResolutionScreenshotMaskInputs.OverrideOutput.LoadAction = View.IsFirstInFamily() ? ERenderTargetLoadAction::EClear : ERenderTargetLoadAction::ELoad;
 
 		SceneColor = AddHighResolutionScreenshotMaskPass(GraphBuilder, View, HighResolutionScreenshotMaskInputs);
 	}
@@ -3242,7 +3239,6 @@ void AddMobilePostProcessingPasses(FRDGBuilder& GraphBuilder, FScene* Scene, con
 		PassInputs.SceneColor = SceneColor;
 		PassInputs.SceneDepth = SceneDepth;
 		PassInputs.SceneTextures = GetSceneTextureShaderParameters(Inputs.SceneTextures);
-		PassInputs.OverrideOutput.LoadAction = View.IsFirstInFamily() ? ERenderTargetLoadAction::EClear : ERenderTargetLoadAction::ELoad;
 
 		// TODO: Nanite - pipe through results
 		FRDGTextureRef DummyStencilTexture = nullptr;
@@ -3257,7 +3253,6 @@ void AddMobilePostProcessingPasses(FRDGBuilder& GraphBuilder, FScene* Scene, con
 		PassInputs.SceneColor = SceneColor;
 		PassInputs.SceneDepth = SceneDepth;
 		PassInputs.BasePassType = FCompositePrimitiveInputs::EBasePassType::Mobile;
-		PassInputs.OverrideOutput.LoadAction = View.IsFirstInFamily() ? ERenderTargetLoadAction::EClear : ERenderTargetLoadAction::ELoad;
 
 		SceneColor = AddEditorPrimitivePass(GraphBuilder, View, PassInputs, InstanceCullingManager);
 	}
@@ -3284,7 +3279,6 @@ void AddMobilePostProcessingPasses(FRDGBuilder& GraphBuilder, FScene* Scene, con
 		PassSequence.AcceptOverrideIfLastPass(EPass::PrimaryUpscale, PassInputs.OverrideOutput);
 		PassInputs.Stage = PassSequence.IsEnabled(EPass::SecondaryUpscale) ? EUpscaleStage::PrimaryToSecondary : EUpscaleStage::PrimaryToOutput;
 		PassInputs.SceneColor = SceneColor;
-		PassInputs.OverrideOutput.LoadAction = View.IsFirstInFamily() ? ERenderTargetLoadAction::EClear : ERenderTargetLoadAction::ELoad;
 
 		if (const ISpatialUpscaler* CustomUpscaler = View.Family->GetPrimarySpatialUpscalerInterface())
 		{
@@ -3360,7 +3354,6 @@ void AddMobilePostProcessingPasses(FRDGBuilder& GraphBuilder, FScene* Scene, con
 			PassInputs.ColorSamplingMethod = FVisualizeComplexityInputs::EColorSamplingMethod::Stair;
 			PassInputs.ComplexityScale = ComplexityScale;
 			PassInputs.bDrawLegend = true;
-			PassInputs.OverrideOutput.LoadAction = View.IsFirstInFamily() ? ERenderTargetLoadAction::EClear : ERenderTargetLoadAction::ELoad;
 
 			SceneColor = AddVisualizeComplexityPass(GraphBuilder, View, PassInputs);
 			break;
@@ -3376,7 +3369,6 @@ void AddMobilePostProcessingPasses(FRDGBuilder& GraphBuilder, FScene* Scene, con
 			PassInputs.ColorSamplingMethod = FVisualizeComplexityInputs::EColorSamplingMethod::Ramp;
 			PassInputs.ComplexityScale = 1.0f;
 			PassInputs.bDrawLegend = true;
-			PassInputs.OverrideOutput.LoadAction = View.IsFirstInFamily() ? ERenderTargetLoadAction::EClear : ERenderTargetLoadAction::ELoad;
 
 			SceneColor = AddVisualizeComplexityPass(GraphBuilder, View, PassInputs);
 			break;
@@ -3404,7 +3396,6 @@ void AddMobilePostProcessingPasses(FRDGBuilder& GraphBuilder, FScene* Scene, con
 		FHMDDistortionInputs PassInputs;
 		PassSequence.AcceptOverrideIfLastPass(EPass::HMDDistortion, PassInputs.OverrideOutput);
 		PassInputs.SceneColor = SceneColor;
-		PassInputs.OverrideOutput.LoadAction = View.IsFirstInFamily() ? ERenderTargetLoadAction::EClear : ERenderTargetLoadAction::ELoad;
 
 		SceneColor = AddHMDDistortionPass(GraphBuilder, View, PassInputs);
 	}
