Fix incorrect SSS look when using Substrate and Format 0 due to min limit on how many SSS tiles are required to run MIPs generation and Burley processing.

diff --git a/Engine/Source/Runtime/Renderer/Private/PostProcess/PostProcessSubsurface.cpp b/Engine/Source/Runtime/Renderer/Private/PostProcess/PostProcessSubsurface.cpp
--- a/Engine/Source/Runtime/Renderer/Private/PostProcess/PostProcessSubsurface.cpp
+++ b/Engine/Source/Runtime/Renderer/Private/PostProcess/PostProcessSubsurface.cpp
@@ -263,6 +263,14 @@ BEGIN_SHADER_PARAMETER_STRUCT(FSubsurfaceParameters, )
 	SHADER_PARAMETER_STRUCT_REF(FViewUniformShaderParameters, ViewUniformBuffer)
 END_SHADER_PARAMETER_STRUCT()
 
+uint32 GetSSSMipmapsMinTileCount(const FViewInfo& View)
+{
+	const int32 MinGenerateMipsTileCount = FMath::Max(0, CVarSSSMipmapsMinTileCount.GetValueOnRenderThread());
+	// When using Substrate and Adaptive GBuffer format, we support per-pixel MFP, which requires to run Burley categorization. 
+	// This forces Burley to run even when a limited number of tile requires MIPs/Burley in order to output the correct look
+	return Substrate::IsSubstrateEnabled() && !Substrate::IsSubstrateBlendableGBufferEnabled(View.GetShaderPlatform()) ? 0u : MinGenerateMipsTileCount;
+}
+
 FSubsurfaceParameters GetSubsurfaceCommonParameters(FRDGBuilder& GraphBuilder, const FViewInfo& View, TRDGUniformBufferRef<FSceneTextureUniformParameters> SceneTextures)
 {
 	const float DistanceToProjectionWindow = View.ViewMatrices.GetProjectionMatrix().M[0][0];
@@ -270,7 +278,7 @@ FSubsurfaceParameters GetSubsurfaceCommonParameters(FRDGBuilder& GraphBuilder, c
 	const float SSSScaleX = SSSScaleZ / SUBSURFACE_KERNEL_SIZE * 0.5f;
 
 	const float SSSOverrideNumSamples = float(CVarSSSBurleyNumSamplesOverride.GetValueOnRenderThread());
-	const float MinGenerateMipsTileCount = FMath::Max(0.0f, CVarSSSMipmapsMinTileCount.GetValueOnRenderThread());
+	const float MinGenerateMipsTileCount = GetSSSMipmapsMinTileCount(View);
 
 	FSubsurfaceParameters Parameters;
 	Parameters.SubsurfaceParams = FVector4f(SSSScaleX, SSSScaleZ, SSSOverrideNumSamples, MinGenerateMipsTileCount);
@@ -1049,7 +1057,7 @@ void AddSubsurfaceViewPass(
 
 	const bool bUseProfileIdCache = IsProfileIdCacheEnabled();
 
-	const int32 MinGenerateMipsTileCount = FMath::Max(0, CVarSSSMipmapsMinTileCount.GetValueOnRenderThread());
+	const uint32 MinGenerateMipsTileCount = GetSSSMipmapsMinTileCount(View);
 
 	/**
 	 * All subsurface passes within the screen-space subsurface effect can operate at half or full resolution,
