Fix SceneTexture parameter handling in SceneHitProxyRenderer to correctly create either the Mobile or Deferred scene textures uniform buffer, depending on the rendering path.

#tests Editor preview mode

diff --git a/Engine/Source/Runtime/Renderer/Private/SceneHitProxyRendering.cpp b/Engine/Source/Runtime/Renderer/Private/SceneHitProxyRendering.cpp
--- a/Engine/Source/Runtime/Renderer/Private/SceneHitProxyRendering.cpp
+++ b/Engine/Source/Runtime/Renderer/Private/SceneHitProxyRendering.cpp
@@ -224,7 +224,7 @@ void InitHitProxyRender(FRDGBuilder& GraphBuilder, FSceneRenderer* SceneRenderer
 
 BEGIN_SHADER_PARAMETER_STRUCT(FHitProxyPassParameters, )
 	SHADER_PARAMETER_STRUCT_INCLUDE(FViewShaderParameters, View)
-	SHADER_PARAMETER_RDG_UNIFORM_BUFFER(FSceneTextureUniformParameters, SceneTextures)
+	SHADER_PARAMETER_STRUCT_INCLUDE(FSceneTextureShaderParameters, SceneTextures)
 	SHADER_PARAMETER_STRUCT_INCLUDE(FInstanceCullingDrawParams, InstanceCullingDrawParams)
 	RENDER_TARGET_BINDING_SLOTS()
 END_SHADER_PARAMETER_STRUCT()
@@ -359,7 +359,7 @@ static void DoRenderHitProxies(
 
 		PassParameters->RenderTargets[0] = FRenderTargetBinding(HitProxyTexture, ERenderTargetLoadAction::ELoad);
 		PassParameters->RenderTargets.DepthStencil = FDepthStencilBinding(HitProxyDepthTexture, ERenderTargetLoadAction::ELoad, ERenderTargetLoadAction::ELoad, FExclusiveDepthStencil::DepthWrite_StencilWrite);
-		PassParameters->SceneTextures = CreateSceneTextureUniformBuffer(GraphBuilder, &SceneRenderer->GetActiveSceneTextures(), SceneRenderer->FeatureLevel, ESceneTextureSetupMode::None);
+		PassParameters->SceneTextures = CreateSceneTextureShaderParameters(GraphBuilder, View, ESceneTextureSetupMode::None);
 
 		GraphBuilder.AddPass(
 			RDG_EVENT_NAME("HitProxies::Render"),
