[VSM] Replace showstats cvar to prevent issues with settings in ini files when upgrading, and to allow setting the layout from an ini. The new cvars are now r.shadow.virtual.stats, r.shadow.virtual.stats.visibile and r.shadow.virtual.stats.sections. Move dynamic mesh bounds visualization to a separate cvar.

diff --git a/Engine/Shaders/Private/VirtualShadowMaps/VirtualShadowMapBuildPerPageDrawCommands.usf b/Engine/Shaders/Private/VirtualShadowMaps/VirtualShadowMapBuildPerPageDrawCommands.usf
--- a/Engine/Shaders/Private/VirtualShadowMaps/VirtualShadowMapBuildPerPageDrawCommands.usf
+++ b/Engine/Shaders/Private/VirtualShadowMaps/VirtualShadowMapBuildPerPageDrawCommands.usf
@@ -208,7 +208,7 @@ void CullPerPageDrawCommandsCs(uint3 GroupId : SV_GroupID, int GroupThreadIndex
 		float3 LocalBoundsExtent = InstanceData.LocalBoundsExtent;
 		LoadDynamicMeshBounds(DrawCommandDesc.DynamicMeshBoundsIndex, LocalBoundsCenter, LocalBoundsExtent);
 
-#if VSM_GENERATE_STATS
+#if WITH_SHADERPRINT
 		float4 BoundsColor = all(LocalBoundsCenter == InstanceData.LocalBoundsCenter) ? ColorRed : ColorGreen;
 		AddOBBWS(LocalBoundsCenter - LocalBoundsExtent, LocalBoundsCenter + LocalBoundsExtent, BoundsColor, DFHackToFloat(InstanceData.LocalToWorld));
 #endif
diff --git a/Engine/Source/Editor/CommonMenuExtensions/Private/VirtualShadowMapVisualizationMenuCommands.cpp b/Engine/Source/Editor/CommonMenuExtensions/Private/VirtualShadowMapVisualizationMenuCommands.cpp
--- a/Engine/Source/Editor/CommonMenuExtensions/Private/VirtualShadowMapVisualizationMenuCommands.cpp
+++ b/Engine/Source/Editor/CommonMenuExtensions/Private/VirtualShadowMapVisualizationMenuCommands.cpp
@@ -144,7 +144,7 @@ void FVirtualShadowMapVisualizationMenuCommands::BindCommands(FUICommandList& Co
 		);
 	}
 
-	IConsoleVariable* ShowStatsVisibleCVar = IConsoleManager::Get().FindConsoleVariable(TEXT("r.Shadow.Virtual.ShowStats.Visible"));
+	IConsoleVariable* ShowStatsVisibleCVar = IConsoleManager::Get().FindConsoleVariable(TEXT("r.Shadow.Virtual.Stats.Visible"));
 	check(ShowStatsVisibleCVar);
 	CommandList.MapAction(
 		ShowStatsCommand,
diff --git a/Engine/Source/Runtime/Renderer/Private/VirtualShadowMaps/VirtualShadowMapArray.cpp b/Engine/Source/Runtime/Renderer/Private/VirtualShadowMaps/VirtualShadowMapArray.cpp
--- a/Engine/Source/Runtime/Renderer/Private/VirtualShadowMaps/VirtualShadowMapArray.cpp
+++ b/Engine/Source/Runtime/Renderer/Private/VirtualShadowMaps/VirtualShadowMapArray.cpp
@@ -110,7 +110,7 @@ TAutoConsoleVariable<int32> CVarBuildDynamicHZB(
 );
 
 TAutoConsoleVariable<bool> CVarShowStatsVisible(
-	TEXT("r.Shadow.Virtual.ShowStats.Visible"),
+	TEXT("r.Shadow.Virtual.Stats.Visible"),
 	false,
 	TEXT("Render VSM Stats to screen"),
 	ECVF_RenderThreadSafe
@@ -127,8 +127,86 @@ EVSMStatSection::Type FVirtualShadowMap::GetEnabledStatSections()
 	return GVSMStatSectionsEnabled;
 }
 
+static FString SetEnabledStatSections(const FString& Sections)
+{
+	FString OutputWarning;
+
+	if (Sections == TEXT("basic"))
+	{
+		GVSMStatSectionsEnabled = EVSMStatSection::Basic;
+	}
+	else if (Sections == TEXT("all"))
+	{
+		GVSMStatSectionsEnabled = EVSMStatSection::All;
+	}
+	else
+	{
+		// Iterate over flags in comma-separated string
+		auto Tokenize = [](const FStringView& String, TCHAR Separator, TFunctionRef<void(const FStringView&)> Func)
+		{
+			FStringView RemainingString { String };
+			for (int TokenEnd; RemainingString.FindChar(Separator, TokenEnd);)
+			{
+				FStringView Token = RemainingString.SubStr(0, TokenEnd);
+				Func(Token);
+				RemainingString.RightChopInline(TokenEnd + 1);
+			}
+			if (RemainingString.Len() > 0)
+			{
+				Func(RemainingString);
+			}
+		};
+
+		EVSMStatSection::Type SectionsEnabled = EVSMStatSection::Basic;
+
+		// Match string to section flag and turn it on
+		auto ParseFlag = [&](EVSMStatSection::Type Flag, FString FlagLabel, FStringView Token)
+		{
+			if (Token.Equals(FlagLabel, ESearchCase::IgnoreCase))
+			{
+				SectionsEnabled = EVSMStatSection::Type(SectionsEnabled | Flag);
+				return true;
+			}
+			return false;
+		};
+
+		Tokenize(Sections, ',', [&](const FStringView& Token)
+		{
+			#define PARSE_FLAG(x) ParseFlag(EVSMStatSection::x, TEXT(#x), Token) ||
+			bool bFlagValid = ENUMERATE_VSM_STAT_SECTIONS(PARSE_FLAG) 0;
+			#undef PARSE_FLAG
+
+			if(!bFlagValid)
+			{
+				FString TokenStr(Token);
+				OutputWarning += FString::Printf(TEXT("Invalid argument '%s'\n"), *TokenStr);
+			}
+		});
+
+		GVSMStatSectionsEnabled = SectionsEnabled;
+	}
+	return OutputWarning;
+}
+
+static TAutoConsoleVariable<FString> CVarShowStatsSections(
+	TEXT("r.Shadow.Virtual.Stats.Sections"),
+	"all",
+	TEXT("Select which parts of the stats visualization should be visible."),
+	FConsoleVariableDelegate::CreateLambda(
+		[](IConsoleVariable* CVar)
+		{
+			FString Value = CVar->GetString();
+			FString Output = SetEnabledStatSections(Value);
+			if (Output.Len() > 0)
+			{
+				UE_LOG(LogRenderer, Warning, TEXT("%s"), *Output);
+			}
+		}
+	)
+);
+
 static FAutoConsoleCommand CCmdVSMShowStats = FAutoConsoleCommand(
-	TEXT("r.Shadow.Virtual.ShowStats"),
+	TEXT("r.Shadow.Virtual.Stats"),
 	TEXT("Show/hide VSM statistics.\n")
 	TEXT("Usage: ShowStats <0|1|basic|all>\n")
 	TEXT("Usage: ShowStats <section1,section2,...>\n")
@@ -169,68 +247,20 @@ static FAutoConsoleCommand CCmdVSMShowStats = FAutoConsoleCommand(
 		{
 			CVarShowStatsVisible->Set(true);
 		}
-		else if (ValueString == TEXT("basic"))
-		{
-			CVarShowStatsVisible->Set(true);
-			GVSMStatSectionsEnabled = EVSMStatSection::Basic;
-			
-		}
-		else if (ValueString == TEXT("all"))
-		{
-			CVarShowStatsVisible->Set(true);
-			GVSMStatSectionsEnabled = EVSMStatSection::All;
-		}
 		else
 		{
 			CVarShowStatsVisible->Set(true);
-
-			// Iterate over flags in comma-separated string
-			auto Tokenize = [](const FStringView& String, TCHAR Separator, TFunctionRef<void(const FStringView&)> Func)
+			FString OutputString = SetEnabledStatSections(ValueString);
+			if (OutputString.Len() > 0)
 			{
-				FStringView RemainingString { String };
-				for (int TokenEnd; RemainingString.FindChar(Separator, TokenEnd);)
-				{
-					FStringView Token = RemainingString.SubStr(0, TokenEnd);
-					Func(Token);
-					RemainingString.RightChopInline(TokenEnd + 1);
-				}
-				if (RemainingString.Len() > 0)
-				{
-					Func(RemainingString);
-				}
-			};
-
-			EVSMStatSection::Type SectionsEnabled = EVSMStatSection::Basic;
-
-			// Match string to section flag and turn it on
-			auto ParseFlag = [&](EVSMStatSection::Type Flag, FString FlagLabel, FStringView Token)
-			{
-				if (Token.Equals(FlagLabel, ESearchCase::IgnoreCase))
-				{
-					SectionsEnabled = EVSMStatSection::Type(SectionsEnabled | Flag);
-					return true;
-				}
-				return false;
-			};
-
-			Tokenize(ValueString, ',', [&](const FStringView& Token)
-			{
-				#define PARSE_FLAG(x) ParseFlag(EVSMStatSection::x, TEXT(#x), Token) ||
-				bool bFlagValid = ENUMERATE_VSM_STAT_SECTIONS(PARSE_FLAG) 0;
-				#undef PARSE_FLAG
-
-				if(!bFlagValid)
-				{
-					FString TokenStr(Token);
-					Output.Logf(TEXT("Invalid argument '%s'"), *TokenStr);
-				}
-			});
-
-			GVSMStatSectionsEnabled = SectionsEnabled;
+				Output.Log(TEXT("LogRenderer"), ELogVerbosity::Warning, *OutputString);
+			}
 		}
 	})
 );
 
+static FAutoConsoleVariableDeprecated CVarVSMShowStats_Deprecated(TEXT("r.Shadow.Virtual.ShowStats"), TEXT("r.Shadow.Virtual.Stats.Visible"), TEXT("5.7"));
+
 static TAutoConsoleVariable<float> CVarPageDilationBorderSizeDirectional(
 	TEXT("r.Shadow.Virtual.PageDilationBorderSizeDirectional"),
 	0.05f,
@@ -446,6 +476,13 @@ static TAutoConsoleVariable<int32> CVarVisualizeLayout(
 	ECVF_RenderThreadSafe
 );
 
+static TAutoConsoleVariable<int32> CVarVisualizeDynamicMeshBounds(
+	TEXT("r.Shadow.Virtual.VisualizeDynamicMeshBounds"),
+	0,
+	TEXT("\n"),
+	ECVF_RenderThreadSafe
+);
+
 TAutoConsoleVariable<int32> CVarDebugSkipMergePhysical(
 	TEXT("r.Shadow.Virtual.DebugSkipMergePhysical"),
 	0,
@@ -3182,8 +3219,9 @@ class FCullPerPageDrawCommandsCs : public FGlobalShader
 
 	class FUseHzbDim : SHADER_PERMUTATION_BOOL("USE_HZB_OCCLUSION");
 	class FGenerateStatsDim : SHADER_PERMUTATION_BOOL("VSM_GENERATE_STATS");
+	class FWithShaderPrintDim : SHADER_PERMUTATION_BOOL("WITH_SHADERPRINT");
 	class FBatchedDim : SHADER_PERMUTATION_BOOL("ENABLE_BATCH_MODE");
-	using FPermutationDomain = TShaderPermutationDomain< FUseHzbDim, FBatchedDim, FGenerateStatsDim >;
+	using FPermutationDomain = TShaderPermutationDomain< FUseHzbDim, FBatchedDim, FGenerateStatsDim, FWithShaderPrintDim >;
 
 	static constexpr uint32 ThreadGroupSize = FInstanceProcessingGPULoadBalancer::ThreadGroupSize;
 
@@ -3449,16 +3487,20 @@ static FCullingResult AddCullingPasses(FRDGBuilder& GraphBuilder,
 			PassParameters, 
 			PermutationVector);
 
+		bool bVisualizeDynamicMeshBounds = false;
+#if !UE_BUILD_SHIPPING
 		if (bGenerateStats)
 		{
-#if !UE_BUILD_SHIPPING
 			PassParameters->NumPageAreaDiagnosticSlots = CVarNumPageAreaDiagSlots.GetValueOnRenderThread() < 0 ? FVirtualShadowMapArray::MaxPageAreaDiagnosticSlots : FMath::Min(FVirtualShadowMapArray::MaxPageAreaDiagnosticSlots, uint32(CVarNumPageAreaDiagSlots.GetValueOnRenderThread()));
 			PassParameters->LargeInstancePageAreaThreshold = CVarLargeInstancePageAreaThreshold.GetValueOnRenderThread() >= 0 ? CVarLargeInstancePageAreaThreshold.GetValueOnRenderThread() : (VirtualShadowMapArray.GetMaxPhysicalPages() / 8);
-#endif
 		}
 
+		bVisualizeDynamicMeshBounds = CVarVisualizeDynamicMeshBounds.GetValueOnRenderThread() > 0;
+#endif
+
 		PermutationVector.Set< FCullPerPageDrawCommandsCs::FBatchedDim >(bUseBatchMode);
 		PermutationVector.Set< FCullPerPageDrawCommandsCs::FUseHzbDim >(HZBShaderParameters.HZBTextureArray != nullptr);
+		PermutationVector.Set< FCullPerPageDrawCommandsCs::FWithShaderPrintDim >(bVisualizeDynamicMeshBounds);
 
 		auto ComputeShader = ShaderMap->GetShader<FCullPerPageDrawCommandsCs>(PermutationVector);
 
