Fixed: Issue with a LightFunction material PSO pre-caching on Android devices

diff --git a/Engine/Source/Runtime/Renderer/Private/MobileDeferredShadingPass.cpp b/Engine/Source/Runtime/Renderer/Private/MobileDeferredShadingPass.cpp
--- a/Engine/Source/Runtime/Renderer/Private/MobileDeferredShadingPass.cpp
+++ b/Engine/Source/Runtime/Renderer/Private/MobileDeferredShadingPass.cpp
@@ -1198,8 +1198,11 @@ void FDeferredMobileLightMaterialPSOCollector::CollectPSOInitializers(
 	const FPSOPrecacheParams& PreCacheParams,
 	TArray<FPSOPrecacheData>& PSOInitializers)
 {
-	CollectPSOInitializersDirectional(SceneTexturesConfig, Material, PSOInitializers);
-	CollectPSOInitializersLocal(SceneTexturesConfig, Material, PSOInitializers);
+	if (Material.GetMaterialDomain() == MD_LightFunction)
+	{
+		CollectPSOInitializersDirectional(SceneTexturesConfig, Material, PSOInitializers);
+		CollectPSOInitializersLocal(SceneTexturesConfig, Material, PSOInitializers);
+	}
 }
 
 void FDeferredMobileLightMaterialPSOCollector::CollectPSOInitializersDirectional(
@@ -1302,6 +1305,14 @@ void FDeferredMobileLightMaterialPSOCollector::CollectPSOInitializersDirectional
 			GraphicsPSOInit.StatePrecachePSOHash = RHIComputeStatePrecachePSOHash(GraphicsPSOInit);
 			ApplyTargetsInfo(GraphicsPSOInit, RenderTargetsInfo);
 
+			GraphicsPSOInit.SubpassIndex = 0;
+			GraphicsPSOInit.SubpassHint = ESubpassHint::None;
+			if (MobileAllowFramebufferFetch(GMaxRHIShaderPlatform))
+			{
+				GraphicsPSOInit.SubpassIndex = 2;
+				GraphicsPSOInit.SubpassHint = ESubpassHint::DeferredShadingSubpass;
+			}
+
 			FPSOPrecacheData PSOPrecacheData;
 			PSOPrecacheData.bRequired = true;
 			PSOPrecacheData.Type = FPSOPrecacheData::EType::Graphics;
@@ -1401,6 +1412,14 @@ void FDeferredMobileLightMaterialPSOCollector::CollectPSOInitializersLocal(
 			GraphicsPSOInit.StatePrecachePSOHash = RHIComputeStatePrecachePSOHash(GraphicsPSOInit);
 			ApplyTargetsInfo(GraphicsPSOInit, RenderTargetsInfo);
 
+			GraphicsPSOInit.SubpassIndex = 0;
+			GraphicsPSOInit.SubpassHint = ESubpassHint::None;
+			if (MobileAllowFramebufferFetch(GMaxRHIShaderPlatform))
+			{
+				GraphicsPSOInit.SubpassIndex = 2;
+				GraphicsPSOInit.SubpassHint = ESubpassHint::DeferredShadingSubpass;
+			}
+
 			FPSOPrecacheData PSOPrecacheData;
 			PSOPrecacheData.bRequired = true;
 			PSOPrecacheData.Type = FPSOPrecacheData::EType::Graphics;
