Mobile: Added ambient occlusion pass in multi-pass deferred shading mode

diff --git a/Engine/Source/Runtime/Renderer/Private/MobileShadingRenderer.cpp b/Engine/Source/Runtime/Renderer/Private/MobileShadingRenderer.cpp
--- a/Engine/Source/Runtime/Renderer/Private/MobileShadingRenderer.cpp
+++ b/Engine/Source/Runtime/Renderer/Private/MobileShadingRenderer.cpp
@@ -2424,6 +2424,11 @@ void FMobileSceneRenderer::RenderDeferredMultiPass(FRDGBuilder& GraphBuilder, FS
 		RenderMobileShadowProjections(GraphBuilder);
 	}
 
+	if (bRequiresAmbientOcclusionPass)
+	{
+		RenderAmbientOcclusion(GraphBuilder, SceneTextures.Depth.Resolve, SceneTextures.ScreenSpaceAO);
+	}
+
 	const bool bSupportsSM5Nodes = MobileSupportsSM5MaterialNodes(ShaderPlatform);
 	const bool bShouldRenderFullSLW = bSupportsSM5Nodes && ShouldRenderSingleLayerWater(Views);
 	const bool bNeedsDepthCopy = bShouldRenderFullSLW;
