Mobile: Apply AmbientOcclusionStaticFraction only for projects with a static lighting. AO now correctly applied when SkyLight uses real-time capture.

diff --git a/Engine/Shaders/Private/MobileDeferredShading.usf b/Engine/Shaders/Private/MobileDeferredShading.usf
--- a/Engine/Shaders/Private/MobileDeferredShading.usf
+++ b/Engine/Shaders/Private/MobileDeferredShading.usf
@@ -201,7 +201,9 @@ void MobileDirectionalLightPS(
 	half GatheredAmbientOcclusion = 1;
 #if ENABLE_AMBIENT_OCCLUSION
 	GatheredAmbientOcclusion = Texture2DSample(MobileBasePass.AmbientOcclusionTexture, MobileBasePass.AmbientOcclusionSampler, UVAndScreenPos.xy).r;
-	GatheredAmbientOcclusion = lerp(1.0, GatheredAmbientOcclusion, MobileBasePass.AmbientOcclusionStaticFraction);
+	#if ALLOW_STATIC_LIGHTING
+		GatheredAmbientOcclusion = lerp(1.0, GatheredAmbientOcclusion, MobileBasePass.AmbientOcclusionStaticFraction);
+	#endif
 	GBuffer.GBufferAO *= GatheredAmbientOcclusion;
 #endif
 
diff --git a/Engine/Source/Runtime/Renderer/Private/MobileShadingRenderer.cpp b/Engine/Source/Runtime/Renderer/Private/MobileShadingRenderer.cpp
--- a/Engine/Source/Runtime/Renderer/Private/MobileShadingRenderer.cpp
+++ b/Engine/Source/Runtime/Renderer/Private/MobileShadingRenderer.cpp
@@ -564,7 +564,7 @@ void FMobileSceneRenderer::InitViews(
 
 	bRequiresAmbientOcclusionPass = IsUsingMobileAmbientOcclusion(ShaderPlatform)
 		&& Views[0].FinalPostProcessSettings.AmbientOcclusionIntensity > 0
-		&& (Views[0].FinalPostProcessSettings.AmbientOcclusionStaticFraction >= 1 / 100.0f || (Scene && Scene->SkyLight && Scene->SkyLight->ProcessedTexture && Views[0].Family->EngineShowFlags.SkyLighting))
+		&& (Views[0].FinalPostProcessSettings.AmbientOcclusionStaticFraction >= 1 / 100.0f || (Scene && Scene->SkyLight && Views[0].Family->EngineShowFlags.SkyLighting))
 		&& ViewFamily.EngineShowFlags.Lighting
 		&& !Views[0].bIsReflectionCapture
 		&& !Views[0].bIsPlanarReflection
