On mobile, only exclude inset shadow casting objects from the CSM if inset shadows are supported

diff --git a/Engine/Source/Runtime/Renderer/Private/ShadowSetup.cpp b/Engine/Source/Runtime/Renderer/Private/ShadowSetup.cpp
--- a/Engine/Source/Runtime/Renderer/Private/ShadowSetup.cpp
+++ b/Engine/Source/Runtime/Renderer/Private/ShadowSetup.cpp
@@ -5027,6 +5027,15 @@ struct FGatherShadowPrimitivesPacket
 					}
 				}
 
+				bool bCastsInsetShadow = DoesPrimitiveCastInsetShadow(TaskData, PrimitiveSceneInfo, PrimitiveProxy);
+
+				// if we are on mobile also check if inset shadows are supported
+				if (bCastsInsetShadow && TaskData.FeatureLevel < ERHIFeatureLevel::SM5)
+				{
+					const bool bMobileInsetShadows = CVarMobileInsetShadows.GetValueOnRenderThread();
+					bCastsInsetShadow &= bMobileInsetShadows;
+				}
+
 				if (!bScreenSpaceSizeCulled
 					&& bShadowRelevance
 					&& ProjectedShadowInfo->GetLightSceneInfoCompact().AffectsPrimitive(PrimitiveBounds, PrimitiveProxy)
@@ -5034,7 +5043,7 @@ struct FGatherShadowPrimitivesPacket
 					// Since lights with static shadowing still create per-object shadows for primitives without static shadowing.
 					&& (!LightProxy.HasStaticLighting() || (!LightSceneInfo.IsPrecomputedLightingValid() || LightProxy.UseCSMForDynamicObjects()))
 					// Exclude primitives that will create their own per-object shadow, except when rendering RSMs
-					&& (!DoesPrimitiveCastInsetShadow(TaskData, PrimitiveSceneInfo, PrimitiveProxy))
+					&& (!bCastsInsetShadow)
 					// Exclude primitives that will create a per-object shadow from a stationary light
 					&& !ShouldCreateObjectShadowForStationaryLight(&LightSceneInfo, PrimitiveProxy, true)
 					// Only render shadows from objects that use static lighting during a reflection capture, since the reflection capture doesn't update at runtime
