Fix SLW MRT format declaration being incorrect when using Substrate Format=0 causing artifact on water on PS5.

diff --git a/Engine/Source/Runtime/Renderer/Private/BasePassRendering.cpp b/Engine/Source/Runtime/Renderer/Private/BasePassRendering.cpp
--- a/Engine/Source/Runtime/Renderer/Private/BasePassRendering.cpp
+++ b/Engine/Source/Runtime/Renderer/Private/BasePassRendering.cpp
@@ -1053,7 +1053,7 @@ void ModifyBasePassCSPSCompilationEnvironment(const FMeshMaterialShaderPermutati
 	if (bOutputVelocity)
 	{
 		// As defined in BasePassPixelShader.usf. Also account for Substrate setting velocity in slot 1 as described in FetchLegacyGBufferInfo.
-		const int32 VelocityIndex = Substrate::IsSubstrateEnabled() ? 1 : (IsForwardShadingEnabled(Parameters.Platform) ? 1 : 4);
+		const int32 VelocityIndex = (Substrate::IsSubstrateEnabled() && !Substrate::IsSubstrateBlendableGBufferEnabled(Parameters.Platform)) ? 1 : (IsForwardShadingEnabled(Parameters.Platform) ? 1 : 4);
 		OutEnvironment.SetRenderTargetOutputFormat(VelocityIndex, PF_G16R16);
 	}
 
