Removed unused function that was part of the old path for updating primitive instances.

#rnx

diff --git a/Engine/Source/Runtime/Renderer/Private/RendererScene.cpp b/Engine/Source/Runtime/Renderer/Private/RendererScene.cpp
--- a/Engine/Source/Runtime/Renderer/Private/RendererScene.cpp
+++ b/Engine/Source/Runtime/Renderer/Private/RendererScene.cpp
@@ -1805,31 +1805,6 @@ void FScene::UpdatePrimitiveInstancesFromCompute(FPrimitiveSceneDesc* Primitive,
 	UpdatePrimitiveInternal(SceneProxy, MoveTemp(UpdateCommand));
 }
 
-void FScene::UpdatePrimitiveInstancesFromCompute(FPrimitiveSceneInfo* PrimitiveSceneInfo, FGPUSceneWriteDelegate&& DataWriterGPU)
-{
-	SCOPED_NAMED_EVENT(FScene_UpdatePrimitiveInstanceFromCompute, FColor::Yellow);
-	check(PrimitiveSceneInfo);
-
-	// Primitive must have a scene proxy already created in order to update it's instance data.
-	if (!ensureAlways(PrimitiveSceneInfo->Proxy != nullptr))
-	{
-		return;
-	}
-
-	// Only updates to GPU-only primitives are currently allowed.
-	if (!ensureAlways(PrimitiveSceneInfo->Proxy->IsInstanceDataGPUOnly()))
-	{
-		return;
-	}
-
-	FUpdateInstanceFromComputeCommand UpdateCommand;
-	UpdateCommand.PrimitiveSceneProxy = PrimitiveSceneInfo->Proxy;
-	UpdateCommand.GPUSceneWriter = MoveTemp(DataWriterGPU);
-
-	// This is already on the renderthread so queue directly.
-	PrimitiveUpdates.Enqueue(PrimitiveSceneInfo, MoveTemp(UpdateCommand));
-}
-
 static bool AllowPersistSurfaceCacheOnVisibilityChange(const FPrimitiveSceneProxy* PrimitiveSceneProxy)
 {
 	return CVarLumenSceneSurfaceCacheAllowPersistOnVisibilityChange.GetValueOnRenderThread() != 0
diff --git a/Engine/Source/Runtime/Renderer/Private/ScenePrivate.h b/Engine/Source/Runtime/Renderer/Private/ScenePrivate.h
--- a/Engine/Source/Runtime/Renderer/Private/ScenePrivate.h
+++ b/Engine/Source/Runtime/Renderer/Private/ScenePrivate.h
@@ -3928,8 +3928,6 @@ public:
 protected:
 
 private:
-	void UpdatePrimitiveInstancesFromCompute(FPrimitiveSceneInfo* PrimitiveSceneInfo, FGPUSceneWriteDelegate&& DataWriterGPU);
-
 	template<class T> 	
 	void BatchAddPrimitivesInternal(TArrayView<T*> InPrimitives);
 
