[HWRT] Fix dynamic instances not being properly rendered in ray tracing debug visualizations.
- need to use RayTracing::GetVisibleShaderBindings() instead of View.VisibleRayTracingShaderBindings to get dynamic instances bindings.

diff --git a/Engine/Source/Runtime/Renderer/Private/DeferredShadingRenderer.cpp b/Engine/Source/Runtime/Renderer/Private/DeferredShadingRenderer.cpp
--- a/Engine/Source/Runtime/Renderer/Private/DeferredShadingRenderer.cpp
+++ b/Engine/Source/Runtime/Renderer/Private/DeferredShadingRenderer.cpp
@@ -1718,7 +1718,13 @@ void FDeferredShadingSceneRenderer::RenderNanite(FRDGBuilder& GraphBuilder, cons
 }
 
 #if RHI_RAYTRACING
-extern void RenderRayTracingDebug(FRDGBuilder& GraphBuilder, const FScene& Scene, const FViewInfo& View, FSceneTextures& SceneTextures, FRayTracingPickingFeedback& PickingFeedback);
+extern void RenderRayTracingDebug(
+	FRDGBuilder& GraphBuilder,
+	const FScene& Scene,
+	const FViewInfo& View,
+	FSceneTextures& SceneTextures,
+	TConstArrayView<FRayTracingShaderBindingData> VisibleRayTracingShaderBindings,
+	FRayTracingPickingFeedback& PickingFeedback);
 extern void RayTracingDebugDisplayOnScreenMessages(FScreenMessageWriter& Writer, const FViewInfo& View);
 #endif // RHI_RAYTRACING
 
@@ -3768,10 +3774,14 @@ void FDeferredShadingSceneRenderer::Render(FRDGBuilder& GraphBuilder, const FSce
 			}
 			else if (ViewFamily.EngineShowFlags.RayTracingDebug)
 			{
+				// TODO: This will include visible bindings for all views, but we could potentially also provide a way to get visible bindings for a single view
+				// Although that would require running deduplication logic separately for each view in VisibleRayTracingShaderBindingsFinalizeTask
+				TConstArrayView<FRayTracingShaderBindingData> VisibleRayTracingShaderBindings = RayTracing::GetVisibleShaderBindings(*InitViewTaskDatas.RayTracingGatherInstances);
+
 				for (const FViewInfo& View : Views)
 				{
 					FRayTracingPickingFeedback PickingFeedback = {};
-					RenderRayTracingDebug(GraphBuilder, *Scene, View, SceneTextures, PickingFeedback);
+					RenderRayTracingDebug(GraphBuilder, *Scene, View, SceneTextures, VisibleRayTracingShaderBindings, PickingFeedback);
 
 					OnGetOnScreenMessages.AddLambda([this, &View, PickingFeedback](FScreenMessageWriter& ScreenMessageWriter)->void
 						{
diff --git a/Engine/Source/Runtime/Renderer/Private/RayTracing/RayTracingDebug.cpp b/Engine/Source/Runtime/Renderer/Private/RayTracing/RayTracingDebug.cpp
--- a/Engine/Source/Runtime/Renderer/Private/RayTracing/RayTracingDebug.cpp
+++ b/Engine/Source/Runtime/Renderer/Private/RayTracing/RayTracingDebug.cpp
@@ -718,6 +718,7 @@ static void BindRayTracingDebugHitStatsCHSMaterialBindings(
 	FRHICommandList& RHICmdList,
 	FRHIShaderBindingTable* SBT,
 	const FRayTracingMeshCommandStorage& RayTracingMeshCommands,
+	TConstArrayView<FRayTracingShaderBindingData> VisibleRayTracingShaderBindings,
 	const FViewInfo& View,
 	FRHIUniformBuffer* SceneUniformBuffer,
 	FRHIUniformBuffer* NaniteRayTracingUniformBuffer,
@@ -733,7 +734,7 @@ static void BindRayTracingDebugHitStatsCHSMaterialBindings(
 			: RHICmdList.Alloc(Size, Align);
 	};
 
-	const int32 NumTotalBindings = View.VisibleRayTracingShaderBindings.Num();
+	const int32 NumTotalBindings = VisibleRayTracingShaderBindings.Num();
 	const uint32 MergedBindingsSize = sizeof(FRayTracingLocalShaderBindings) * NumTotalBindings;
 	FRayTracingLocalShaderBindings* Bindings = (FRayTracingLocalShaderBindings*)Alloc(MergedBindingsSize, alignof(FRayTracingLocalShaderBindings));
 
@@ -793,7 +794,7 @@ static void BindRayTracingDebugHitStatsCHSMaterialBindings(
 	const uint32 NumShaderSlotsPerGeometrySegment = SBT->GetInitializer().NumShaderSlotsPerGeometrySegment;
 
 	uint32 BindingIndex = 0;
-	for (const FRayTracingShaderBindingData DirtyShaderBinding : View.VisibleRayTracingShaderBindings)
+	for (const FRayTracingShaderBindingData DirtyShaderBinding : VisibleRayTracingShaderBindings)
 	{
 		const FRayTracingMeshCommand& MeshCommand = DirtyShaderBinding.GetRayTracingMeshCommand(RayTracingMeshCommands);
 
@@ -823,6 +824,7 @@ static void BindRayTracingDebugCHSMaterialBindings(
 	FRHICommandList& RHICmdList,
 	FRHIShaderBindingTable* SBT,
 	const FRayTracingMeshCommandStorage& RayTracingMeshCommands,
+	TConstArrayView<FRayTracingShaderBindingData> VisibleRayTracingShaderBindings,
 	const FViewInfo& View,
 	FRHIUniformBuffer* SceneUniformBuffer,
 	FRHIUniformBuffer* NaniteRayTracingUniformBuffer,
@@ -837,7 +839,7 @@ static void BindRayTracingDebugCHSMaterialBindings(
 			: RHICmdList.Alloc(Size, Align);
 	};
 
-	const int32 NumTotalBindings = View.VisibleRayTracingShaderBindings.Num();
+	const int32 NumTotalBindings = VisibleRayTracingShaderBindings.Num();
 	const uint32 MergedBindingsSize = sizeof(FRayTracingLocalShaderBindings) * NumTotalBindings;
 	FRayTracingLocalShaderBindings* Bindings = (FRayTracingLocalShaderBindings*)Alloc(MergedBindingsSize, alignof(FRayTracingLocalShaderBindings));
 
@@ -894,7 +896,7 @@ static void BindRayTracingDebugCHSMaterialBindings(
 	const uint32 NumShaderSlotsPerGeometrySegment = SBT->GetInitializer().NumShaderSlotsPerGeometrySegment;
 
 	uint32 BindingIndex = 0;
-	for (const FRayTracingShaderBindingData DirtyShaderBinding : View.VisibleRayTracingShaderBindings)
+	for (const FRayTracingShaderBindingData DirtyShaderBinding : VisibleRayTracingShaderBindings)
 	{
 		const FRayTracingMeshCommand& MeshCommand = DirtyShaderBinding.GetRayTracingMeshCommand(RayTracingMeshCommands);
 
@@ -941,7 +943,12 @@ void FDeferredShadingSceneRenderer::PrepareRayTracingDebug(const FSceneViewFamil
 	}
 }
 
-static FRDGBufferRef RayTracingPerformPicking(FRDGBuilder& GraphBuilder, const FScene& Scene, const FViewInfo& View, FRayTracingPickingFeedback& PickingFeedback, bool bInlineRayTracing)
+static FRDGBufferRef RayTracingPerformPicking(
+	FRDGBuilder& GraphBuilder,
+	const FScene& Scene,
+	const FViewInfo& View,
+	TConstArrayView<FRayTracingShaderBindingData> VisibleRayTracingShaderBindings,
+	FRayTracingPickingFeedback& PickingFeedback, bool bInlineRayTracing)
 {
 	const FRayTracingScene& RayTracingScene = Scene.RayTracingScene;
 
@@ -1020,7 +1027,7 @@ static FRDGBufferRef RayTracingPerformPicking(FRDGBuilder& GraphBuilder, const F
 			RDG_EVENT_NAME("RayTracingPicking"),
 			RayGenParameters,
 			ERDGPassFlags::Compute,
-			[RayGenParameters, RayGenShader, &View, PickingSBT, PickingPipeline, &RayTracingMeshCommands = Scene.CachedRayTracingMeshCommands](FRDGAsyncTask, FRHICommandList& RHICmdList)
+			[RayGenParameters, RayGenShader, &View, PickingSBT, PickingPipeline, &RayTracingMeshCommands = Scene.CachedRayTracingMeshCommands, VisibleRayTracingShaderBindings](FRDGAsyncTask, FRHICommandList& RHICmdList)
 			{
 				FRHIBatchedShaderParameters& GlobalResources = RHICmdList.GetScratchShaderParameters();
 				SetShaderParameters(GlobalResources, RayGenShader, *RayGenParameters);
@@ -1029,7 +1036,7 @@ static FRDGBufferRef RayTracingPerformPicking(FRDGBuilder& GraphBuilder, const F
 				FRHIUniformBuffer* NaniteRayTracingUniformBuffer = RayGenParameters->SharedParameters.NaniteRayTracingUniformBuffer->GetRHI();
 				TOptional<FScopedUniformBufferStaticBindings> StaticUniformBufferScope = RayTracing::BindStaticUniformBufferBindings(View, SceneUniformBuffer, NaniteRayTracingUniformBuffer, RHICmdList);
 
-				BindRayTracingDebugCHSMaterialBindings(RHICmdList, PickingSBT, RayTracingMeshCommands, View, SceneUniformBuffer, NaniteRayTracingUniformBuffer, PickingPipeline);
+				BindRayTracingDebugCHSMaterialBindings(RHICmdList, PickingSBT, RayTracingMeshCommands, VisibleRayTracingShaderBindings, View, SceneUniformBuffer, NaniteRayTracingUniformBuffer, PickingPipeline);
 				RHICmdList.SetRayTracingMissShader(PickingSBT, 0, PickingPipeline, 0 /* ShaderIndexInPipeline */, 0, nullptr, 0);
 				RHICmdList.CommitShaderBindingTable(PickingSBT);
 
@@ -1222,7 +1229,7 @@ static void PrintTopKMostHitMessage(FRDGBuilder& GraphBuilder, const FScene& Sce
 	}
 }
 
-static FRDGBufferRef RayTracingPerformHitStatsPerPrimitive(FRDGBuilder& GraphBuilder, const FScene& Scene, const FViewInfo& View)
+static FRDGBufferRef RayTracingPerformHitStatsPerPrimitive(FRDGBuilder& GraphBuilder, const FScene& Scene, const FViewInfo& View, TConstArrayView<FRayTracingShaderBindingData> VisibleRayTracingShaderBindings)
 {
 	const FRayTracingScene& RayTracingScene = Scene.RayTracingScene;
 
@@ -1281,7 +1288,7 @@ static FRDGBufferRef RayTracingPerformHitStatsPerPrimitive(FRDGBuilder& GraphBui
 		RDG_EVENT_NAME("RayTracingHitStats"),
 		RayGenParameters,
 		ERDGPassFlags::Compute,
-		[RayGenParameters, RayGenShader, &View, HitStatsSBT, HitStatsPerPrimitivePipeline, ViewRect, &RayTracingMeshCommands = Scene.CachedRayTracingMeshCommands](FRDGAsyncTask, FRHICommandList& RHICmdList)
+		[RayGenParameters, RayGenShader, &View, HitStatsSBT, HitStatsPerPrimitivePipeline, ViewRect, &RayTracingMeshCommands = Scene.CachedRayTracingMeshCommands, VisibleRayTracingShaderBindings](FRDGAsyncTask, FRHICommandList& RHICmdList)
 		{
 			FRHIBatchedShaderParameters& GlobalResources = RHICmdList.GetScratchShaderParameters();
 			SetShaderParameters(GlobalResources, RayGenShader, *RayGenParameters);
@@ -1290,7 +1297,7 @@ static FRDGBufferRef RayTracingPerformHitStatsPerPrimitive(FRDGBuilder& GraphBui
 			FRHIUniformBuffer* NaniteRayTracingUniformBuffer = RayGenParameters->NaniteRayTracingUniformBuffer->GetRHI();
 			TOptional<FScopedUniformBufferStaticBindings> StaticUniformBufferScope = RayTracing::BindStaticUniformBufferBindings(View, SceneUniformBuffer, NaniteRayTracingUniformBuffer, RHICmdList);
 
-			BindRayTracingDebugHitStatsCHSMaterialBindings(RHICmdList, HitStatsSBT, RayTracingMeshCommands, View, SceneUniformBuffer, NaniteRayTracingUniformBuffer, DebugHitStatsUniformBuffer->GetRHI(), HitStatsPerPrimitivePipeline);
+			BindRayTracingDebugHitStatsCHSMaterialBindings(RHICmdList, HitStatsSBT, RayTracingMeshCommands, VisibleRayTracingShaderBindings, View, SceneUniformBuffer, NaniteRayTracingUniformBuffer, DebugHitStatsUniformBuffer->GetRHI(), HitStatsPerPrimitivePipeline);
 			RHICmdList.SetRayTracingMissShader(HitStatsSBT, 0, HitStatsPerPrimitivePipeline, 0 /* ShaderIndexInPipeline */, 0, nullptr, 0);
 			RHICmdList.CommitShaderBindingTable(HitStatsSBT);
 
@@ -1549,7 +1556,7 @@ extern void RenderRayTracingPrimaryRaysView(
 	FRDGTextureRef* InOutColorTexture, FRDGTextureRef* InOutRayHitDistanceTexture,
 	int32 SamplePerPixel, int32 HeightFog, float ResolutionFraction, ERayTracingPrimaryRaysFlag Flags);
 
-void RenderRayTracingDebug(FRDGBuilder& GraphBuilder, const FScene& Scene, const FViewInfo& View, FSceneTextures& SceneTextures, FRayTracingPickingFeedback& PickingFeedback)
+void RenderRayTracingDebug(FRDGBuilder& GraphBuilder, const FScene& Scene, const FViewInfo& View, FSceneTextures& SceneTextures, TConstArrayView<FRayTracingShaderBindingData> VisibleRayTracingShaderBindings, FRayTracingPickingFeedback& PickingFeedback)
 {
 	const EShaderPlatform ShaderPlatform = Scene.GetShaderPlatform();
 
@@ -1659,7 +1666,7 @@ void RenderRayTracingDebug(FRDGBuilder& GraphBuilder, const FScene& Scene, const
 		&& RayTracingScene.GetInstanceExtraDataBuffer(ERayTracingSceneLayer::Base) != nullptr
 		&& RayTracingScene.GetInstanceDebugBuffer(ERayTracingSceneLayer::Base) != nullptr)
 	{
-		PickingBuffer = RayTracingPerformPicking(GraphBuilder, Scene, View, PickingFeedback, bInlineRayTracing);
+		PickingBuffer = RayTracingPerformPicking(GraphBuilder, Scene, View, VisibleRayTracingShaderBindings, PickingFeedback, bInlineRayTracing);
 	}
 	else
 	{
@@ -1668,7 +1675,7 @@ void RenderRayTracingDebug(FRDGBuilder& GraphBuilder, const FScene& Scene, const
 
 	if (DebugVisualizationMode == RAY_TRACING_DEBUG_VIZ_HITCOUNT_PER_INSTANCE)
 	{
-		StatsBuffer = RayTracingPerformHitStatsPerPrimitive(GraphBuilder, Scene, View);
+		StatsBuffer = RayTracingPerformHitStatsPerPrimitive(GraphBuilder, Scene, View, VisibleRayTracingShaderBindings);
 	}
 	else 
 	{
@@ -1850,7 +1857,7 @@ void RenderRayTracingDebug(FRDGBuilder& GraphBuilder, const FScene& Scene, const
 			RDG_EVENT_NAME("RayTracingDebug"),
 			RayGenParameters,
 			ERDGPassFlags::Compute,
-			[RayGenParameters, RayGenShader, &View, Pipeline, SBT, ViewRect, &RayTracingMeshCommands = Scene.CachedRayTracingMeshCommands, bRequiresBindings](FRDGAsyncTask, FRHICommandList& RHICmdList)
+			[RayGenParameters, RayGenShader, &View, Pipeline, SBT, ViewRect, &RayTracingMeshCommands = Scene.CachedRayTracingMeshCommands, VisibleRayTracingShaderBindings, bRequiresBindings](FRDGAsyncTask, FRHICommandList& RHICmdList)
 			{
 				FRHIBatchedShaderParameters& GlobalResources = RHICmdList.GetScratchShaderParameters();
 				SetShaderParameters(GlobalResources, RayGenShader, *RayGenParameters);
@@ -1861,7 +1868,7 @@ void RenderRayTracingDebug(FRDGBuilder& GraphBuilder, const FScene& Scene, const
 
 				if (bRequiresBindings)
 				{
-					BindRayTracingDebugCHSMaterialBindings(RHICmdList, SBT, RayTracingMeshCommands, View, SceneUniformBuffer, NaniteRayTracingUniformBuffer, Pipeline);
+					BindRayTracingDebugCHSMaterialBindings(RHICmdList, SBT, RayTracingMeshCommands, VisibleRayTracingShaderBindings, View, SceneUniformBuffer, NaniteRayTracingUniformBuffer, Pipeline);
 					RHICmdList.SetRayTracingMissShader(SBT, 0, Pipeline, 0 /* ShaderIndexInPipeline */, 0, nullptr, 0);
 					RHICmdList.CommitShaderBindingTable(SBT);
 				}
