Remove incorrect comment.

diff --git a/Engine/Source/Runtime/Renderer/Private/Substrate/Substrate.cpp b/Engine/Source/Runtime/Renderer/Private/Substrate/Substrate.cpp
--- a/Engine/Source/Runtime/Renderer/Private/Substrate/Substrate.cpp
+++ b/Engine/Source/Runtime/Renderer/Private/Substrate/Substrate.cpp
@@ -223,8 +223,6 @@ bool GetSubstrateUsesAnisotropy(const FViewInfo& View)
 {
 	if (Substrate::IsSubstrateEnabled())
 	{
-		// Use the per-view value rather than the per-scene data to have more accurate dispatching of special complex tiles 
-		// and avoid unecessary empty-dispatch
 		return View.SubstrateViewData.bUsesAnisotropy;
 	}
 	return false;
