Remove  missing stat scope on some builds

diff --git a/Engine/Source/Runtime/Renderer/Private/Shadows/ShadowSceneRenderer.cpp b/Engine/Source/Runtime/Renderer/Private/Shadows/ShadowSceneRenderer.cpp
--- a/Engine/Source/Runtime/Renderer/Private/Shadows/ShadowSceneRenderer.cpp
+++ b/Engine/Source/Runtime/Renderer/Private/Shadows/ShadowSceneRenderer.cpp
@@ -858,7 +858,6 @@ void FShadowSceneRenderer::BeginMarkVirtualShadowMapPages(
 	}
 
 	RDG_EVENT_SCOPE(GraphBuilder, "VirtualShadowMapMarkPages");
-	RDG_GPU_STAT_SCOPE(GraphBuilder, VirtualShadowMapMarkPages);
 
 	VirtualShadowMapArray.BeginMarkPages(
 		GraphBuilder,
