Restore PSO precaching for PostProcessing materials if the shader platform doesn't support manual vertex fetch and set up FRegisterPSOCollectorCreateFunction for EShadingPath::Mobile

diff --git a/Engine/Source/Runtime/Engine/Private/Materials/Material.cpp b/Engine/Source/Runtime/Engine/Private/Materials/Material.cpp
--- a/Engine/Source/Runtime/Engine/Private/Materials/Material.cpp
+++ b/Engine/Source/Runtime/Engine/Private/Materials/Material.cpp
@@ -4754,9 +4754,9 @@ void UMaterial::PostLoad()
 			PreloadMaterialShaderMap(GetMaterialResource(GMaxRHIShaderPlatform), Unused);
 		}
 	}
-	else if (IsUIMaterial() || (MaterialDomain == MD_LightFunction) ||
+	else if (IsUIMaterial() || (MaterialDomain == MD_LightFunction) || IsPostProcessMaterial() ||
 		// TODO: need to pass a correct vertex declaration for non-MVF platforms
-		((IsDeferredDecal() || IsPostProcessMaterial()) && RHISupportsManualVertexFetch(GMaxRHIShaderPlatform)))
+		(IsDeferredDecal() && RHISupportsManualVertexFetch(GMaxRHIShaderPlatform)))
 	{
 		FPSOPrecacheParams PSOPrecacheParams;
 		UMaterialInterface::PrecachePSOs(&FLocalVertexFactory::StaticType, PSOPrecacheParams);
diff --git a/Engine/Source/Runtime/Engine/Private/Materials/MaterialInstance.cpp b/Engine/Source/Runtime/Engine/Private/Materials/MaterialInstance.cpp
--- a/Engine/Source/Runtime/Engine/Private/Materials/MaterialInstance.cpp
+++ b/Engine/Source/Runtime/Engine/Private/Materials/MaterialInstance.cpp
@@ -3659,9 +3659,9 @@ void UMaterialInstance::PostLoad()
 			PreloadMaterialShaderMap(GetMaterialResource(GMaxRHIShaderPlatform), Unused);
 		}
 	}
-	else if (IsUIMaterial() ||
+	else if (IsUIMaterial() || IsPostProcessMaterial() ||
 		// TODO: need to pass a correct vertex declaration for non-MVF platforms
-		((IsDeferredDecal() || IsPostProcessMaterial()) && RHISupportsManualVertexFetch(GMaxRHIShaderPlatform)))
+		(IsDeferredDecal() && RHISupportsManualVertexFetch(GMaxRHIShaderPlatform)))
 	{
 		FPSOPrecacheParams PSOPrecacheParams;
 		UMaterialInterface::PrecachePSOs(&FLocalVertexFactory::StaticType, PSOPrecacheParams);
diff --git a/Engine/Source/Runtime/Renderer/Private/PostProcess/PostProcessMaterial.cpp b/Engine/Source/Runtime/Renderer/Private/PostProcess/PostProcessMaterial.cpp
--- a/Engine/Source/Runtime/Renderer/Private/PostProcess/PostProcessMaterial.cpp
+++ b/Engine/Source/Runtime/Renderer/Private/PostProcess/PostProcessMaterial.cpp
@@ -656,9 +656,6 @@ void AddNeuralPostProcessPass(
 
 	const uint32 MaterialStencilRef = Material->GetStencilRefValue();
 
-	const bool bMobilePlatform = IsMobilePlatform(View.GetShaderPlatform());
-
-
 	EScreenPassDrawFlags ScreenPassFlags = EScreenPassDrawFlags::AllowHMDHiddenAreaMask;
 
 	// check if we can skip that draw call in case if all pixels will fail the stencil test of the material
@@ -775,7 +772,7 @@ FScreenPassTexture AddPostProcessMaterialPass(
 
 	FRHIBlendState* DefaultBlendState = FScreenPassPipelineState::FDefaultBlendState::GetRHI();
 	FRHIBlendState* BlendState = DefaultBlendState;
-	
+
 	if (IsMaterialBlendEnabled(Material))
 	{
 		BlendState = GetMaterialBlendState(Material);
@@ -884,8 +881,6 @@ FScreenPassTexture AddPostProcessMaterialPass(
 
 	const uint32 MaterialStencilRef = Material->GetStencilRefValue();
 
-	const bool bMobilePlatform = IsMobilePlatform(View.GetShaderPlatform());
-
 	EScreenPassDrawFlags ScreenPassFlags = EScreenPassDrawFlags::AllowHMDHiddenAreaMask;
 
 	// check if we can skip that draw call in case if all pixels will fail the stencil test of the material
@@ -901,7 +896,6 @@ FScreenPassTexture AddPostProcessMaterialPass(
 		const uint32 StencilClearValue = Inputs.CustomDepthTexture ? Inputs.CustomDepthTexture->Desc.ClearValue.Value.DSValue.Stencil : 0;
 		bFailStencil &= PostProcessStencilTest(StencilClearValue, StencilComp, MaterialStencilRef);
 
-
 		for (const uint32& Value : View.CustomDepthStencilValues)
 		{
 			bFailStencil &= PostProcessStencilTest(Value, StencilComp, MaterialStencilRef);
@@ -917,7 +911,6 @@ FScreenPassTexture AddPostProcessMaterialPass(
 
 	if (!bSkipPostProcess)
 	{
-
 		FNeuralPostProcessResource NeuralPostProcessResource;
 		const bool bShouldApplyNeuralPostProcessing = ShouldApplyNeuralPostProcessForMaterial(Material);
 
@@ -1540,8 +1533,7 @@ class FPostProcessMaterialPSOCollector : public IPSOCollector
 {
 public:
 	FPostProcessMaterialPSOCollector(ERHIFeatureLevel::Type InFeatureLevel) : 
-	IPSOCollector(FPSOCollectorCreateManager::GetIndex(GetFeatureLevelShadingPath(InFeatureLevel), PostProcessMaterialPSOCollectorName)),
-	FeatureLevel(InFeatureLevel)
+	IPSOCollector(FPSOCollectorCreateManager::GetIndex(GetFeatureLevelShadingPath(InFeatureLevel), PostProcessMaterialPSOCollectorName))
 	{
 	}
 
@@ -1552,10 +1544,6 @@ public:
 		const FPSOPrecacheParams& PreCacheParams,
 		TArray<FPSOPrecacheData>& PSOInitializers
 	) override final;
-
-private:
-
-	ERHIFeatureLevel::Type FeatureLevel;
 };
 
 void FPostProcessMaterialPSOCollector::CollectPSOInitializers(
@@ -1633,3 +1621,4 @@ IPSOCollector* CreatePostProcessMaterialPSOCollector(ERHIFeatureLevel::Type Feat
 	return new FPostProcessMaterialPSOCollector(FeatureLevel);
 }
 FRegisterPSOCollectorCreateFunction RegisterPostProcessMaterialPSOCollector(&CreatePostProcessMaterialPSOCollector, EShadingPath::Deferred, PostProcessMaterialPSOCollectorName);
+FRegisterPSOCollectorCreateFunction RegisterMobilePostProcessMaterialPSOCollector(&CreatePostProcessMaterialPSOCollector, EShadingPath::Mobile, PostProcessMaterialPSOCollectorName);
