Restore PSO precaching for decal materials if the shader platform doesn't support manual vertex fetch

diff --git a/Engine/Source/Runtime/Engine/Private/Materials/Material.cpp b/Engine/Source/Runtime/Engine/Private/Materials/Material.cpp
--- a/Engine/Source/Runtime/Engine/Private/Materials/Material.cpp
+++ b/Engine/Source/Runtime/Engine/Private/Materials/Material.cpp
@@ -4754,9 +4754,7 @@ void UMaterial::PostLoad()
 			PreloadMaterialShaderMap(GetMaterialResource(GMaxRHIShaderPlatform), Unused);
 		}
 	}
-	else if (IsUIMaterial() || (MaterialDomain == MD_LightFunction) || IsPostProcessMaterial() ||
-		// TODO: need to pass a correct vertex declaration for non-MVF platforms
-		(IsDeferredDecal() && RHISupportsManualVertexFetch(GMaxRHIShaderPlatform)))
+	else if (IsDeferredDecal() || IsUIMaterial() || (MaterialDomain == MD_LightFunction) || IsPostProcessMaterial())
 	{
 		FPSOPrecacheParams PSOPrecacheParams;
 		UMaterialInterface::PrecachePSOs(&FLocalVertexFactory::StaticType, PSOPrecacheParams);
diff --git a/Engine/Source/Runtime/Engine/Private/Materials/MaterialInstance.cpp b/Engine/Source/Runtime/Engine/Private/Materials/MaterialInstance.cpp
--- a/Engine/Source/Runtime/Engine/Private/Materials/MaterialInstance.cpp
+++ b/Engine/Source/Runtime/Engine/Private/Materials/MaterialInstance.cpp
@@ -3659,9 +3659,7 @@ void UMaterialInstance::PostLoad()
 			PreloadMaterialShaderMap(GetMaterialResource(GMaxRHIShaderPlatform), Unused);
 		}
 	}
-	else if (IsUIMaterial() || IsPostProcessMaterial() ||
-		// TODO: need to pass a correct vertex declaration for non-MVF platforms
-		(IsDeferredDecal() && RHISupportsManualVertexFetch(GMaxRHIShaderPlatform)))
+	else if (IsDeferredDecal() || IsUIMaterial() || IsPostProcessMaterial())
 	{
 		FPSOPrecacheParams PSOPrecacheParams;
 		UMaterialInterface::PrecachePSOs(&FLocalVertexFactory::StaticType, PSOPrecacheParams);
diff --git a/Engine/Source/Runtime/Renderer/Private/CompositionLighting/PostProcessMeshDecals.cpp b/Engine/Source/Runtime/Renderer/Private/CompositionLighting/PostProcessMeshDecals.cpp
--- a/Engine/Source/Runtime/Renderer/Private/CompositionLighting/PostProcessMeshDecals.cpp
+++ b/Engine/Source/Runtime/Renderer/Private/CompositionLighting/PostProcessMeshDecals.cpp
@@ -30,6 +30,13 @@ static TAutoConsoleVariable<int32> CVarParallelMeshDecal(
 	ECVF_RenderThreadSafe
 );
 
+static int32 GDeferredDecalMeshMaterialPSOPrecache = 1;
+static FAutoConsoleVariableRef CVarGDeferredDecalMeshMaterialPSOPrecache(
+	TEXT("r.PSOPrecache.DeferredDecalMeshMaterial"), 
+	GDeferredDecalMeshMaterialPSOPrecache, 
+	TEXT("Precache all possible required PSOs for loaded DeferredDecalMesh Materials."),
+	ECVF_ReadOnly);
+
 class FMeshDecalsVS : public FMeshMaterialShader
 {
 public:
@@ -441,8 +448,7 @@ void FMeshDecalMeshProcessor::CollectPSOInitializers(
 		// Collect decal pass PSOs
 		CollectDeferredDecalPassPSOInitializers(PSOCollectorIndex, FeatureLevel, SceneTexturesConfig, Material, LocalDecalRenderStage, PSOInitializers);
 
-		// TODO: need to pass a correct vertex declaration for non-MVF platforms
-		if (RHISupportsManualVertexFetch(ShaderPlatform))
+		if (GDeferredDecalMeshMaterialPSOPrecache)
 		{
 			// Collect decal mesh PSOs
 			CollectDeferredDecalMeshPSOInitializers(SceneTexturesConfig, VertexFactoryData, PreCacheParams, Material, DecalBlendDesc, LocalDecalRenderStage, PSOInitializers);
