Mobile: Fixed case where AO pass is run twice a frame
#DisabledValidations nonUnityAndStaticAnalysis

diff --git a/Engine/Source/Runtime/Renderer/Private/MobileShadingRenderer.cpp b/Engine/Source/Runtime/Renderer/Private/MobileShadingRenderer.cpp
--- a/Engine/Source/Runtime/Renderer/Private/MobileShadingRenderer.cpp
+++ b/Engine/Source/Runtime/Renderer/Private/MobileShadingRenderer.cpp
@@ -2424,7 +2424,7 @@ void FMobileSceneRenderer::RenderDeferredMultiPass(FRDGBuilder& GraphBuilder, FS
 		RenderMobileShadowProjections(GraphBuilder);
 	}
 
-	if (bRequiresAmbientOcclusionPass)
+	if (bRequiresAmbientOcclusionPass && !bIsFullDepthPrepassEnabled)
 	{
 		RenderAmbientOcclusion(GraphBuilder, SceneTextures.Depth.Resolve, SceneTextures.ScreenSpaceAO);
 	}
