[HWRT] Fix crash due to use-after-free error accessing DirtyPersistentRayTracingShaderBindings.
- allocate DirtyPersistentRayTracingShaderBindings from GraphBuilder since it's accessed inside RDG setup tasks.

diff --git a/Engine/Source/Runtime/Renderer/Private/DeferredShadingRenderer.cpp b/Engine/Source/Runtime/Renderer/Private/DeferredShadingRenderer.cpp
--- a/Engine/Source/Runtime/Renderer/Private/DeferredShadingRenderer.cpp
+++ b/Engine/Source/Runtime/Renderer/Private/DeferredShadingRenderer.cpp
@@ -1112,7 +1112,7 @@ void FDeferredShadingSceneRenderer::SetupRayTracingRenderingData(FRDGBuilder& Gr
 
 	TConstArrayView<FRayTracingShaderBindingData> VisibleRayTracingShaderBindings = RayTracing::GetVisibleShaderBindings(RayTracingGatherInstancesTaskData);
 
-	FRayTracingShaderBindingDataOneFrameArray DirtyPersistentRayTracingShaderBindings;
+	FRayTracingShaderBindingDataOneFrameArray& DirtyPersistentRayTracingShaderBindings = GraphBuilder.AllocArray<FRayTracingShaderBindingData>();
 
 	bool bRequireBindingsUpdate = false;
 
