LensDistortion: Allow LUT generation access by plugins.

* Mostly useful for local distortion debugging without calibrated lens files.
* This change lets us generate distortion LUTs before they are accessible in the main render, for video pre-processing.

diff --git a/Engine/Source/Runtime/Renderer/Internal/PostProcess/LensDistortion.h b/Engine/Source/Runtime/Renderer/Internal/PostProcess/LensDistortion.h
--- a/Engine/Source/Runtime/Renderer/Internal/PostProcess/LensDistortion.h
+++ b/Engine/Source/Runtime/Renderer/Internal/PostProcess/LensDistortion.h
@@ -36,10 +36,10 @@ struct FPaniniProjectionConfig
 	FPaniniProjectionConfig() = default;
 
 	/** Returns whether the panini is enabled by cvar. */
-	static bool IsEnabledByCVars();
+	RENDERER_API static bool IsEnabledByCVars();
 
 	/** Returns teh cvars' configuration. */
-	static FPaniniProjectionConfig ReadCVars();
+	RENDERER_API static FPaniniProjectionConfig ReadCVars();
 
 	bool IsEnabled() const
 	{
@@ -60,6 +60,12 @@ struct FPaniniProjectionConfig
 
 	/** Add a RDG pass to generate the lens distortion LUT from the settings. */
 	FLensDistortionLUT GenerateLUTPasses(FRDGBuilder& GraphBuilder, const FViewInfo& View) const;
+
+	/**
+	 * Add a RDG pass to generate the lens distortion LUT from the settings.
+	 * Unsafe to due internal argument cast from "FSceneView" to "FViewInfo".
+	 */
+	RENDERER_API FLensDistortionLUT GenerateLUTPassesUnsafe(FRDGBuilder& GraphBuilder, const FSceneView& InView) const;
 };
 
 namespace LensDistortion
diff --git a/Engine/Source/Runtime/Renderer/Private/PostProcess/LensDistortion.cpp b/Engine/Source/Runtime/Renderer/Private/PostProcess/LensDistortion.cpp
--- a/Engine/Source/Runtime/Renderer/Private/PostProcess/LensDistortion.cpp
+++ b/Engine/Source/Runtime/Renderer/Private/PostProcess/LensDistortion.cpp
@@ -220,6 +220,11 @@ FLensDistortionLUT FPaniniProjectionConfig::GenerateLUTPasses(FRDGBuilder& Graph
 	return LensDistortionLUT;
 }
 
+FLensDistortionLUT FPaniniProjectionConfig::GenerateLUTPassesUnsafe(FRDGBuilder& GraphBuilder, const FSceneView& InView) const
+{
+	return GenerateLUTPasses(GraphBuilder, static_cast<const FViewInfo&>(InView));
+}
+
 LensDistortion::EPassLocation LensDistortion::GetPassLocation(const FViewInfo& InViewInfo)
 {
 	if (IsPostProcessingEnabled(InViewInfo)
