Add Substrate tile type tracking on CPU for dispatching only what is necessary.

diff --git a/Engine/Shaders/Private/Substrate/SubstrateMaterialClassification.usf b/Engine/Shaders/Private/Substrate/SubstrateMaterialClassification.usf
--- a/Engine/Shaders/Private/Substrate/SubstrateMaterialClassification.usf
+++ b/Engine/Shaders/Private/Substrate/SubstrateMaterialClassification.usf
@@ -139,7 +139,7 @@ void TileMainCS(uint2 DispatchThreadId : SV_DispatchThreadID, uint LinearIndex :
 		{
 			bool bHasAnisotropy = abs(GBufferData.Anisotropy) > 0;
 
-			bContainsComplexMaterial = bHasAnisotropy;
+			bContainsComplexMaterial = bHasAnisotropy || GBufferData.ShadingModelID == SHADINGMODELID_EYE || GBufferData.ShadingModelID == SHADINGMODELID_HAIR;
 			bContainsSimpleMaterial = !bContainsComplexMaterial && GBufferData.ShadingModelID == SHADINGMODELID_DEFAULT_LIT;
 			bContainsSingleMaterial = !bContainsComplexMaterial && !bContainsSimpleMaterial;
 			bContainsScreenSpaceSubsurfaceScattering = GBufferData.ShadingModelID == SHADINGMODELID_SUBSURFACE_PROFILE || GBufferData.ShadingModelID == SHADINGMODELID_EYE;
diff --git a/Engine/Shaders/Shared/SubstrateDefinitions.h b/Engine/Shaders/Shared/SubstrateDefinitions.h
--- a/Engine/Shaders/Shared/SubstrateDefinitions.h
+++ b/Engine/Shaders/Shared/SubstrateDefinitions.h
@@ -83,6 +83,7 @@
 #define SUBSTRATE_MATERIAL_TYPE_SINGLE					1
 #define SUBSTRATE_MATERIAL_TYPE_COMPLEX					2
 #define SUBSTRATE_MATERIAL_TYPE_COMPLEX_SPECIAL			3
+#define SUBSTRATE_MATERIAL_TYPE_COUNT					4
 
 #define SUBSTRATE_TILE_ENCODING_16BITS 					0
 #define SUBSTRATE_TILE_ENCODING_8BITS  					1
diff --git a/Engine/Source/Editor/UnrealEd/Private/Commandlets/DumpMaterialInfo.cpp b/Engine/Source/Editor/UnrealEd/Private/Commandlets/DumpMaterialInfo.cpp
--- a/Engine/Source/Editor/UnrealEd/Private/Commandlets/DumpMaterialInfo.cpp
+++ b/Engine/Source/Editor/UnrealEd/Private/Commandlets/DumpMaterialInfo.cpp
@@ -348,7 +348,7 @@ TArray<MaterialInfo::FPropertySet> GetMaterialInfoProperties()
 				TEXT("ShadingModelMask"),
 				TEXT("SubstrateUintPerPixel"),
 				TEXT("SubstrateClosureCountMask"),
-				TEXT("bUsesComplexSpecialRenderPath"),
+				TEXT("SubstrateTileTypeMask"),
 				TEXT("bOpaque"),
 				TEXT("bMasked"),
 				TEXT("bDistortion"),
@@ -381,7 +381,7 @@ TArray<MaterialInfo::FPropertySet> GetMaterialInfoProperties()
 				Output.Write(Input.MaterialRelevance.ShadingModelMask);
 				Output.Write(Input.MaterialRelevance.SubstrateUintPerPixel);
 				Output.Write(Input.MaterialRelevance.SubstrateClosureCountMask);
-				Output.Write(Input.MaterialRelevance.bUsesComplexSpecialRenderPath);
+				Output.Write(Input.MaterialRelevance.SubstrateTileTypeMask);
 				Output.Write(Input.MaterialRelevance.bOpaque);
 				Output.Write(Input.MaterialRelevance.bMasked);
 				Output.Write(Input.MaterialRelevance.bDistortion);
diff --git a/Engine/Source/Runtime/Core/Private/UObject/FortniteMainBranchObjectVersion.cpp b/Engine/Source/Runtime/Core/Private/UObject/FortniteMainBranchObjectVersion.cpp
--- a/Engine/Source/Runtime/Core/Private/UObject/FortniteMainBranchObjectVersion.cpp
+++ b/Engine/Source/Runtime/Core/Private/UObject/FortniteMainBranchObjectVersion.cpp
@@ -13,7 +13,7 @@ TMap<FGuid, FGuid> FFortniteMainBranchObjectVersion::GetSystemGuids()
 	SystemGuids.Add(DevGuids.MATERIALSHADERMAP_DERIVEDDATA_VER, FGuid("F051877974B44D409435F321BDB39DBB"));
 	SystemGuids.Add(DevGuids.NIAGARASHADERMAP_DERIVEDDATA_VER, FGuid("6360A977062842A29ED30E3A7ACB0E64"));
 	SystemGuids.Add(DevGuids.Niagara_LatestScriptCompileVersion, FGuid("783FF3700BB641C799B2D43552FBF91E"));
-	SystemGuids.Add(DevGuids.MaterialTranslationDDCVersion, FGuid("9D0538021BDB47DAA63FFC09DDAAC00E"));
+	SystemGuids.Add(DevGuids.MaterialTranslationDDCVersion, FGuid("27F89B90C20549F9BC400291211324B7"));
 
 	// These GUIDs are in AppendRestrictedFortniteMainStreamObjectVersionGuids as they should rarely be changed.
 	//SystemGuids.Add(DevGuids.NANITE_DERIVEDDATA_VER, FGuid("ACF98184978742A592537DA223A1E6BE"));
diff --git a/Engine/Source/Runtime/Engine/Private/Materials/MaterialInterface.cpp b/Engine/Source/Runtime/Engine/Private/Materials/MaterialInterface.cpp
--- a/Engine/Source/Runtime/Engine/Private/Materials/MaterialInterface.cpp
+++ b/Engine/Source/Runtime/Engine/Private/Materials/MaterialInterface.cpp
@@ -560,7 +560,7 @@ FMaterialRelevance UMaterialInterface::GetRelevance_Internal(const UMaterial* Ma
 	MaterialRelevance.bUsesVertexInterpolator = MaterialResource->HasVertexInterpolator();
 	MaterialRelevance.SubstrateClosureCountMask = SubstrateBSDFCountMask;
 	MaterialRelevance.SubstrateUintPerPixel = SubstrateUintPerPixel;
-	MaterialRelevance.bUsesComplexSpecialRenderPath = MaterialResource->MaterialGetSubstrateUsesComplexSpecialRenderPath_GameThread();
+	MaterialRelevance.SubstrateTileTypeMask = GetSubstrateTileTypeAsUint8(MaterialResource->MaterialGetSubstrateTileType_GameThread());
 	MaterialRelevance.bIsLightFunctionAtlasCompatible = MaterialResource->MaterialIsLightFunctionAtlasCompatible_GameThread();
 	MaterialRelevance.bSamplesMaterialCache = MaterialResource->SamplesMaterialCache();
 	MaterialRelevance.bHasMaterialCacheOutput = MaterialResource->HasMaterialCacheOutput();
diff --git a/Engine/Source/Runtime/Engine/Private/Materials/MaterialShared.cpp b/Engine/Source/Runtime/Engine/Private/Materials/MaterialShared.cpp
--- a/Engine/Source/Runtime/Engine/Private/Materials/MaterialShared.cpp
+++ b/Engine/Source/Runtime/Engine/Private/Materials/MaterialShared.cpp
@@ -1191,15 +1191,15 @@ uint8 FMaterial::MaterialGetSubstrateUintPerPixel_RenderThread() const
 	return RenderingThreadShaderMap ? RenderingThreadShaderMap->GetSubstrateUintPerPixel() : false;
 }
 
-bool FMaterial::MaterialGetSubstrateUsesComplexSpecialRenderPath_GameThread() const
+ESubstrateTileType FMaterial::MaterialGetSubstrateTileType_GameThread() const
 {
-	return GameThreadShaderMap ? GameThreadShaderMap->GetSubstrateUsesComplexSpecialRenderPath() : false;
+	return GameThreadShaderMap ? GameThreadShaderMap->GetSubstrateTileType() : ESubstrateTileType::ECount;
 }
 
-bool FMaterial::MaterialGetSubstrateUsesComplexSpecialRenderPath_RenderThread() const
+ESubstrateTileType FMaterial::MaterialGetSubstrateTileType_RenderThread() const
 {
 	check(IsInParallelRenderingThread());
-	return RenderingThreadShaderMap ? RenderingThreadShaderMap->GetSubstrateUsesComplexSpecialRenderPath() : false;
+	return RenderingThreadShaderMap ? RenderingThreadShaderMap->GetSubstrateTileType() : ESubstrateTileType::ECount;
 }
 
 ESubstrateBsdfFeature FMaterial::MaterialGetSubstrateMaterialBsdfFeatures_GameThread() const
diff --git a/Engine/Source/Runtime/Engine/Public/MaterialShared.h b/Engine/Source/Runtime/Engine/Public/MaterialShared.h
--- a/Engine/Source/Runtime/Engine/Public/MaterialShared.h
+++ b/Engine/Source/Runtime/Engine/Public/MaterialShared.h
@@ -1711,8 +1711,8 @@ public:
 	uint8 GetSubstrateMaterialType() const { return GetSubstrateMaterialCompilationOutput().SubstrateMaterialType; }
 	uint8 GetSubstrateClosureCount() const { return GetSubstrateMaterialCompilationOutput().SubstrateClosureCount; }
 	uint8 GetSubstrateUintPerPixel() const { return GetSubstrateMaterialCompilationOutput().SubstrateUintPerPixel; }
-	bool GetSubstrateUsesComplexSpecialRenderPath() const { return GetSubstrateMaterialCompilationOutput().SubstrateMaterialType == SUBSTRATE_MATERIAL_TYPE_COMPLEX_SPECIAL; }
 	ESubstrateBsdfFeature GetSubstrateMaterialBsdfFeatures() const { return ESubstrateBsdfFeature(GetSubstrateMaterialCompilationOutput().SubstrateMaterialBsdfFeatures); }
+	ESubstrateTileType GetSubstrateTileType() const { return GetSubstrateTileTypeFromMaterialType(GetSubstrateMaterialCompilationOutput().SubstrateMaterialType); }
 
 #if WITH_EDITOR
 	uint32 GetNumUsedUVScalars() const { return GetContent()->MaterialCompilationOutput.NumUsedUVScalars; }
@@ -2610,9 +2610,9 @@ public:
 	ENGINE_API uint8 MaterialGetSubstrateUintPerPixel_GameThread() const;
 	ENGINE_API uint8 MaterialGetSubstrateUintPerPixel_RenderThread() const;
 
-	/** Get Substrate material special path requirement (for more expenssive features such as Glints or SpecularLUT). */
-	ENGINE_API bool MaterialGetSubstrateUsesComplexSpecialRenderPath_GameThread() const;
-	ENGINE_API bool MaterialGetSubstrateUsesComplexSpecialRenderPath_RenderThread() const;
+	/** Get Substrate tile type requirement. */
+	ENGINE_API ESubstrateTileType MaterialGetSubstrateTileType_GameThread() const;
+	ENGINE_API ESubstrateTileType MaterialGetSubstrateTileType_RenderThread() const;
 
 	/** Get Substrate material BSDFs features. */
 	ENGINE_API ESubstrateBsdfFeature MaterialGetSubstrateMaterialBsdfFeatures_GameThread() const;
diff --git a/Engine/Source/Runtime/Engine/Public/Materials/MaterialRelevance.h b/Engine/Source/Runtime/Engine/Public/Materials/MaterialRelevance.h
--- a/Engine/Source/Runtime/Engine/Public/Materials/MaterialRelevance.h
+++ b/Engine/Source/Runtime/Engine/Public/Materials/MaterialRelevance.h
@@ -16,10 +16,13 @@ struct FMaterialRelevance
 	{
 		struct
 		{
-			uint16 ShadingModelMask;
-			uint8 SubstrateUintPerPixel;
-			uint8 SubstrateClosureCountMask;
-			uint8 bUsesComplexSpecialRenderPath : 1;
+			// 32 bits
+			uint32 ShadingModelMask : 14;
+			uint32 SubstrateTileTypeMask : 4;
+			uint32 SubstrateUintPerPixel : 6;
+			uint32 SubstrateClosureCountMask : 8;
+
+			// 32 bits
 			uint8 bOpaque : 1;
 			uint8 bMasked : 1;
 			uint8 bDistortion : 1;
@@ -50,6 +53,7 @@ struct FMaterialRelevance
 			uint8 bIsLightFunctionAtlasCompatible : 1;
 			uint8 bSamplesMaterialCache : 1;
 			uint8 bHasMaterialCacheOutput : 1;
+			// 1bit free
 		};
 		uint64 Raw;
 	};
@@ -63,7 +67,7 @@ struct FMaterialRelevance
 	/** Bitwise OR operator.  Sets any relevance bits which are present in either. */
 	FMaterialRelevance& operator|=(const FMaterialRelevance& B)
 	{
-		const uint8 LocalSubstrateUintPerPixel = SubstrateUintPerPixel;
+		const uint32 LocalSubstrateUintPerPixel = SubstrateUintPerPixel;
 		Raw |= B.Raw;
 		SubstrateUintPerPixel = FMath::Max(LocalSubstrateUintPerPixel, B.SubstrateUintPerPixel);
 		return *this;
diff --git a/Engine/Source/Runtime/Engine/Public/Rendering/SubstrateMaterialShared.cpp b/Engine/Source/Runtime/Engine/Public/Rendering/SubstrateMaterialShared.cpp
--- a/Engine/Source/Runtime/Engine/Public/Rendering/SubstrateMaterialShared.cpp
+++ b/Engine/Source/Runtime/Engine/Public/Rendering/SubstrateMaterialShared.cpp
@@ -104,3 +104,21 @@ FSubstrateRegisteredSharedLocalBasis::FSubstrateRegisteredSharedLocalBasis()
 	GraphSharedLocalBasisIndex = 0;
 #endif
 }
+
+const TCHAR* ToString(ESubstrateTileType Type)
+{
+	switch (Type)
+	{
+		case ESubstrateTileType::ESimple:							return TEXT("Simple");
+		case ESubstrateTileType::ESingle:							return TEXT("Single");
+		case ESubstrateTileType::EComplex:							return TEXT("Complex");
+		case ESubstrateTileType::EComplexSpecial:					return TEXT("ComplexSpecial");
+		case ESubstrateTileType::EOpaqueRoughRefraction:			return TEXT("Opaque/RoughRefraction");
+		case ESubstrateTileType::EOpaqueRoughRefractionSSSWithout:	return TEXT("Opaque/RoughRefraction/SSSWithout");
+		case ESubstrateTileType::EDecalSimple:						return TEXT("Decal/Simple");
+		case ESubstrateTileType::EDecalSingle:						return TEXT("Decal/Single");
+		case ESubstrateTileType::EDecalComplex:						return TEXT("Decal/Complex");
+
+	}
+	return TEXT("Unknown");
+}
\ No newline at end of file
diff --git a/Engine/Source/Runtime/Engine/Public/Rendering/SubstrateMaterialShared.h b/Engine/Source/Runtime/Engine/Public/Rendering/SubstrateMaterialShared.h
--- a/Engine/Source/Runtime/Engine/Public/Rendering/SubstrateMaterialShared.h
+++ b/Engine/Source/Runtime/Engine/Public/Rendering/SubstrateMaterialShared.h
@@ -62,6 +62,45 @@ enum class ESubstrateBsdfFeature : uint16
 ENUM_CLASS_FLAGS(ESubstrateBsdfFeature)
 DECLARE_INTRINSIC_TYPE_LAYOUT(ESubstrateBsdfFeature);
 
+// This must map to the SUBSTRATE_TILE_TYPE defines.
+enum ESubstrateTileType : uint32
+{
+	ESimple								= SUBSTRATE_TILE_TYPE_SIMPLE,
+	ESingle								= SUBSTRATE_TILE_TYPE_SINGLE,
+	EComplex							= SUBSTRATE_TILE_TYPE_COMPLEX,
+	EComplexSpecial						= SUBSTRATE_TILE_TYPE_COMPLEX_SPECIAL,
+	//ECount, 
+	EOpaqueRoughRefraction				= SUBSTRATE_TILE_TYPE_ROUGH_REFRACT,
+	EOpaqueRoughRefractionSSSWithout	= SUBSTRATE_TILE_TYPE_ROUGH_REFRACT_SSS_WITHOUT,
+	EDecalSimple						= SUBSTRATE_TILE_TYPE_DECAL_SIMPLE,
+	EDecalSingle						= SUBSTRATE_TILE_TYPE_DECAL_SINGLE,
+	EDecalComplex						= SUBSTRATE_TILE_TYPE_DECAL_COMPLEX,
+	ECount  //ETotalCount
+};
+
+ENGINE_API const TCHAR* ToString(ESubstrateTileType Type);
+
+FORCEINLINE ESubstrateTileType GetSubstrateTileTypeFromMaterialType(uint32 InMaterialType)
+{
+	// InMaterialType should be one of these
+	//  * SUBSTRATE_MATERIAL_TYPE_SIMPLE
+	//  * SUBSTRATE_MATERIAL_TYPE_SINGLE
+	//  * SUBSTRATE_MATERIAL_TYPE_COMPLEX
+	//  * SUBSTRATE_MATERIAL_TYPE_COMPLEX_SPECIAL
+	checkSlow(InMaterialType < SUBSTRATE_MATERIAL_TYPE_COUNT);
+	return ESubstrateTileType(InMaterialType);
+}
+
+FORCEINLINE uint32 GetSubstrateTileTypeAsUint8(ESubstrateTileType In)
+{
+	return 1u << uint32(In);
+}
+
+FORCEINLINE bool UsesSubstrateTileType(uint8 InTypes, ESubstrateTileType In)
+{
+	return uint32(InTypes) & (1u << uint32(In));
+}
+
 struct FSubstrateOperator
 {
 	DECLARE_TYPE_LAYOUT(FSubstrateOperator, NonVirtual);
diff --git a/Engine/Source/Runtime/Renderer/Private/ClusteredDeferredShadingPass.cpp b/Engine/Source/Runtime/Renderer/Private/ClusteredDeferredShadingPass.cpp
--- a/Engine/Source/Runtime/Renderer/Private/ClusteredDeferredShadingPass.cpp
+++ b/Engine/Source/Runtime/Renderer/Private/ClusteredDeferredShadingPass.cpp
@@ -329,61 +329,73 @@ void FDeferredShadingSceneRenderer::AddClusteredDeferredShadingPass(
 			{
 				SubstrateSceneData = &Scene->SubstrateSceneData;
 
-				InternalAddClusteredDeferredShadingPass(
-					GraphBuilder,
-					ViewIndex,
-					View,
-					SceneTextures,
-					SortedLightsSet,
-					EClusterPassInputType::Substrate,
-					ESubstrateTileType::EComplexSpecial,
-					LightingChannelsTexture,
-					ShadowMaskBits,
-					VirtualShadowMapArray,
-					nullptr,
-					SubstrateSceneData);
-
-				InternalAddClusteredDeferredShadingPass(
-					GraphBuilder,
-					ViewIndex,
-					View,
-					SceneTextures,
-					SortedLightsSet,
-					EClusterPassInputType::Substrate,
-					ESubstrateTileType::EComplex,
-					LightingChannelsTexture,
-					ShadowMaskBits,
-					VirtualShadowMapArray,
-					nullptr,
-					SubstrateSceneData);
-
-				InternalAddClusteredDeferredShadingPass(
-					GraphBuilder,
-					ViewIndex,
-					View,
-					SceneTextures,
-					SortedLightsSet,
-					EClusterPassInputType::Substrate,
-					ESubstrateTileType::ESingle,
-					LightingChannelsTexture,
-					ShadowMaskBits,
-					VirtualShadowMapArray,
-					nullptr,
-					SubstrateSceneData);
-
-				InternalAddClusteredDeferredShadingPass(
-					GraphBuilder,
-					ViewIndex,
-					View,
-					SceneTextures,
-					SortedLightsSet,
-					EClusterPassInputType::Substrate,
-					ESubstrateTileType::ESimple,
-					LightingChannelsTexture,
-					ShadowMaskBits,
-					VirtualShadowMapArray,
-					nullptr, 
-					SubstrateSceneData);
+				if (Substrate::GetSubstrateUsesTileType(View, ESubstrateTileType::EComplexSpecial))
+				{
+					InternalAddClusteredDeferredShadingPass(
+						GraphBuilder,
+						ViewIndex,
+						View,
+						SceneTextures,
+						SortedLightsSet,
+						EClusterPassInputType::Substrate,
+						ESubstrateTileType::EComplexSpecial,
+						LightingChannelsTexture,
+						ShadowMaskBits,
+						VirtualShadowMapArray,
+						nullptr,
+						SubstrateSceneData);
+				}
+
+				if (Substrate::GetSubstrateUsesTileType(View, ESubstrateTileType::EComplex))
+				{
+					InternalAddClusteredDeferredShadingPass(
+						GraphBuilder,
+						ViewIndex,
+						View,
+						SceneTextures,
+						SortedLightsSet,
+						EClusterPassInputType::Substrate,
+						ESubstrateTileType::EComplex,
+						LightingChannelsTexture,
+						ShadowMaskBits,
+						VirtualShadowMapArray,
+						nullptr,
+						SubstrateSceneData);
+				}
+
+				if (Substrate::GetSubstrateUsesTileType(View, ESubstrateTileType::ESingle))
+				{
+					InternalAddClusteredDeferredShadingPass(
+						GraphBuilder,
+						ViewIndex,
+						View,
+						SceneTextures,
+						SortedLightsSet,
+						EClusterPassInputType::Substrate,
+						ESubstrateTileType::ESingle,
+						LightingChannelsTexture,
+						ShadowMaskBits,
+						VirtualShadowMapArray,
+						nullptr,
+						SubstrateSceneData);
+				}
+
+				if (Substrate::GetSubstrateUsesTileType(View, ESubstrateTileType::ESimple))
+				{
+					InternalAddClusteredDeferredShadingPass(
+						GraphBuilder,
+						ViewIndex,
+						View,
+						SceneTextures,
+						SortedLightsSet,
+						EClusterPassInputType::Substrate,
+						ESubstrateTileType::ESimple,
+						LightingChannelsTexture,
+						ShadowMaskBits,
+						VirtualShadowMapArray,
+						nullptr, 
+						SubstrateSceneData);
+				}
 			}
 			else
 			{
diff --git a/Engine/Source/Runtime/Renderer/Private/IndirectLightRendering.cpp b/Engine/Source/Runtime/Renderer/Private/IndirectLightRendering.cpp
--- a/Engine/Source/Runtime/Renderer/Private/IndirectLightRendering.cpp
+++ b/Engine/Source/Runtime/Renderer/Private/IndirectLightRendering.cpp
@@ -1592,13 +1592,10 @@ void FDeferredShadingSceneRenderer::RenderDiffuseIndirectAndAmbientOcclusion(
 		{
 			if (Substrate::IsSubstrateEnabled())
 			{
-				if (Substrate::GetSubstrateUsesComplexSpecialPath(View))
-				{
-					ApplyDiffuseIndirect(ESubstrateTileType::EComplexSpecial);
-				}
-				ApplyDiffuseIndirect(ESubstrateTileType::EComplex);
-				ApplyDiffuseIndirect(ESubstrateTileType::ESingle);
-				ApplyDiffuseIndirect(ESubstrateTileType::ESimple);
+				if (Substrate::GetSubstrateUsesTileType(View, ESubstrateTileType::EComplexSpecial)) 	{ ApplyDiffuseIndirect(ESubstrateTileType::EComplexSpecial);}
+				if (Substrate::GetSubstrateUsesTileType(View, ESubstrateTileType::EComplex)) 			{ ApplyDiffuseIndirect(ESubstrateTileType::EComplex);       }
+				if (Substrate::GetSubstrateUsesTileType(View, ESubstrateTileType::ESingle)) 			{ ApplyDiffuseIndirect(ESubstrateTileType::ESingle);        }
+				if (Substrate::GetSubstrateUsesTileType(View, ESubstrateTileType::ESimple)) 			{ ApplyDiffuseIndirect(ESubstrateTileType::ESimple);        }
 			}
 			else
 			{
@@ -1706,13 +1703,10 @@ void FDeferredShadingSceneRenderer::RenderDiffuseIndirectAndAmbientOcclusion(
 
 			if (Substrate::IsSubstrateEnabled())
 			{
-				if (Substrate::GetSubstrateUsesComplexSpecialPath(View))
-				{
-					ApplyAmbientCubemapComposite(ESubstrateTileType::EComplexSpecial);
-				}
-				ApplyAmbientCubemapComposite(ESubstrateTileType::EComplex);
-				ApplyAmbientCubemapComposite(ESubstrateTileType::ESingle);
-				ApplyAmbientCubemapComposite(ESubstrateTileType::ESimple);
+				if (Substrate::GetSubstrateUsesTileType(View, ESubstrateTileType::EComplexSpecial)) 	{ ApplyAmbientCubemapComposite(ESubstrateTileType::EComplexSpecial); }
+				if (Substrate::GetSubstrateUsesTileType(View, ESubstrateTileType::EComplex))			{ ApplyAmbientCubemapComposite(ESubstrateTileType::EComplex); }
+				if (Substrate::GetSubstrateUsesTileType(View, ESubstrateTileType::ESingle)) 			{ ApplyAmbientCubemapComposite(ESubstrateTileType::ESingle); }
+				if (Substrate::GetSubstrateUsesTileType(View, ESubstrateTileType::ESimple)) 			{ ApplyAmbientCubemapComposite(ESubstrateTileType::ESimple); }
 			}
 			else
 			{
@@ -2142,69 +2136,81 @@ void FDeferredShadingSceneRenderer::RenderDeferredReflectionsAndSkyLighting(
 
 			if (Substrate::IsSubstrateEnabled())
 			{
-				AddSkyReflectionPass(
-					GraphBuilder,
-					View,
-					Scene,
-					SceneTextures,
-					DynamicBentNormalAOTexture,
-					ReflectionsColor,
-					SceneTextureParameters,
-					ScreenSpaceReflectionTileClassification,
-					ViewPipelineState.ReflectionsMethod,
-					bSkyLight,
-					bDynamicSkyLight,
-					bApplySkyShadowing,
-					bLumenStandaloneReflections,
-					ESubstrateTileType::EComplexSpecial);
+				if (Substrate::GetSubstrateUsesTileType(View, ESubstrateTileType::EComplexSpecial))
+				{
+					AddSkyReflectionPass(
+						GraphBuilder,
+						View,
+						Scene,
+						SceneTextures,
+						DynamicBentNormalAOTexture,
+						ReflectionsColor,
+						SceneTextureParameters,
+						ScreenSpaceReflectionTileClassification,
+						ViewPipelineState.ReflectionsMethod,
+						bSkyLight,
+						bDynamicSkyLight,
+						bApplySkyShadowing,
+						bLumenStandaloneReflections,
+						ESubstrateTileType::EComplexSpecial);
+				}
 
-				AddSkyReflectionPass(
-					GraphBuilder,
-					View,
-					Scene,
-					SceneTextures,
-					DynamicBentNormalAOTexture,
-					ReflectionsColor,
-					SceneTextureParameters,
-					ScreenSpaceReflectionTileClassification,
-					ViewPipelineState.ReflectionsMethod,
-					bSkyLight,
-					bDynamicSkyLight,
-					bApplySkyShadowing,
-					bLumenStandaloneReflections,
-					ESubstrateTileType::EComplex);
+				if (Substrate::GetSubstrateUsesTileType(View, ESubstrateTileType::EComplex))
+				{
+					AddSkyReflectionPass(
+						GraphBuilder,
+						View,
+						Scene,
+						SceneTextures,
+						DynamicBentNormalAOTexture,
+						ReflectionsColor,
+						SceneTextureParameters,
+						ScreenSpaceReflectionTileClassification,
+						ViewPipelineState.ReflectionsMethod,
+						bSkyLight,
+						bDynamicSkyLight,
+						bApplySkyShadowing,
+						bLumenStandaloneReflections,
+						ESubstrateTileType::EComplex);
+				}
 
-				AddSkyReflectionPass(
-					GraphBuilder,
-					View,
-					Scene,
-					SceneTextures,
-					DynamicBentNormalAOTexture,
-					ReflectionsColor,
-					SceneTextureParameters,
-					ScreenSpaceReflectionTileClassification,
-					ViewPipelineState.ReflectionsMethod,
-					bSkyLight,
-					bDynamicSkyLight,
-					bApplySkyShadowing,
-					bLumenStandaloneReflections,
-					ESubstrateTileType::ESingle);
+				if (Substrate::GetSubstrateUsesTileType(View, ESubstrateTileType::ESingle))
+				{
+					AddSkyReflectionPass(
+						GraphBuilder,
+						View,
+						Scene,
+						SceneTextures,
+						DynamicBentNormalAOTexture,
+						ReflectionsColor,
+						SceneTextureParameters,
+						ScreenSpaceReflectionTileClassification,
+						ViewPipelineState.ReflectionsMethod,
+						bSkyLight,
+						bDynamicSkyLight,
+						bApplySkyShadowing,
+						bLumenStandaloneReflections,
+						ESubstrateTileType::ESingle);
+				}
 
-				AddSkyReflectionPass(
-					GraphBuilder,
-					View,
-					Scene,
-					SceneTextures,
-					DynamicBentNormalAOTexture,
-					ReflectionsColor,
-					SceneTextureParameters,
-					ScreenSpaceReflectionTileClassification,
-					ViewPipelineState.ReflectionsMethod,
-					bSkyLight,
-					bDynamicSkyLight,
-					bApplySkyShadowing,
-					bLumenStandaloneReflections,
-					ESubstrateTileType::ESimple);
+				if (Substrate::GetSubstrateUsesTileType(View, ESubstrateTileType::ESimple))
+				{
+					AddSkyReflectionPass(
+						GraphBuilder,
+						View,
+						Scene,
+						SceneTextures,
+						DynamicBentNormalAOTexture,
+						ReflectionsColor,
+						SceneTextureParameters,
+						ScreenSpaceReflectionTileClassification,
+						ViewPipelineState.ReflectionsMethod,
+						bSkyLight,
+						bDynamicSkyLight,
+						bApplySkyShadowing,
+						bLumenStandaloneReflections,
+						ESubstrateTileType::ESimple);
+				}
 			}
 			else
 			{
diff --git a/Engine/Source/Runtime/Renderer/Private/LightRendering.cpp b/Engine/Source/Runtime/Renderer/Private/LightRendering.cpp
--- a/Engine/Source/Runtime/Renderer/Private/LightRendering.cpp
+++ b/Engine/Source/Runtime/Renderer/Private/LightRendering.cpp
@@ -2884,8 +2884,7 @@ static void RenderLight(
 		//   so that the geometry get correctly stencil culled on complex/simple part of the screen
 		if (Substrate::IsSubstrateEnabled())
 		{
-			// Complex Special tiles
-			if (Substrate::GetSubstrateUsesComplexSpecialPath(View) && !Substrate::IsSubstrateBlendableGBufferEnabled(View.GetShaderPlatform()))
+			if (Substrate::GetSubstrateUsesTileType(View, ESubstrateTileType::EComplexSpecial))
 			{
 				const ESubstrateTileType TileType = ESubstrateTileType::EComplexSpecial;
 				PermutationVector.Set<FDeferredLightPS::FSubstrateTileType>(TileType);
@@ -2893,7 +2892,7 @@ static void RenderLight(
 				TShaderMapRef< FDeferredLightPS > PixelShader(View.ShaderMap, PermutationVectorRemap);
 				InternalRenderLight(GraphBuilder, Scene, View, LightSceneInfo, PixelShader, PassParameters, TileType, TEXT("Light::StandardDeferred(ComplexSpecial)"));
 			}
-			// Complex tiles
+			if (Substrate::GetSubstrateUsesTileType(View, ESubstrateTileType::EComplex))
 			{
 				const ESubstrateTileType TileType = ESubstrateTileType::EComplex;
 				PermutationVector.Set<FDeferredLightPS::FSubstrateTileType>(TileType);
@@ -2901,7 +2900,7 @@ static void RenderLight(
 				TShaderMapRef< FDeferredLightPS > PixelShader(View.ShaderMap, PermutationVectorRemap);
 				InternalRenderLight(GraphBuilder, Scene, View, LightSceneInfo, PixelShader, PassParameters, TileType, TEXT("Light::StandardDeferred(Complex)"));
 			}
-			// Single tiles
+			if (Substrate::GetSubstrateUsesTileType(View, ESubstrateTileType::ESingle))
 			{
 				const ESubstrateTileType TileType = ESubstrateTileType::ESingle;
 				PermutationVector.Set<FDeferredLightPS::FSubstrateTileType>(TileType);
@@ -2909,7 +2908,7 @@ static void RenderLight(
 				TShaderMapRef< FDeferredLightPS > PixelShader(View.ShaderMap, PermutationVectorRemap);
 				InternalRenderLight(GraphBuilder, Scene, View, LightSceneInfo, PixelShader, PassParameters, TileType, TEXT("Light::StandardDeferred(Single)"));
 			}
-			// Simple tiles
+			if (Substrate::GetSubstrateUsesTileType(View, ESubstrateTileType::ESimple))
 			{
 				const ESubstrateTileType TileType = ESubstrateTileType::ESimple;
 				PermutationVector.Set<FDeferredLightPS::FSubstrateTileType>(TileType);
@@ -3452,13 +3451,22 @@ void FDeferredShadingSceneRenderer::RenderSimpleLightsStandardDeferred(
 
 		if (Substrate::IsSubstrateEnabled())
 		{
-			if (Substrate::GetSubstrateUsesComplexSpecialPath(View))
+			if (Substrate::GetSubstrateUsesTileType(View, ESubstrateTileType::EComplexSpecial))
 			{
 				InternalRenderSimpleLightsStandardDeferred(GraphBuilder, Scene, View, ViewIndex, NumViews, SceneTextures, SortedLightSet, ESubstrateTileType::EComplexSpecial);
 			}
-			InternalRenderSimpleLightsStandardDeferred(GraphBuilder, Scene, View, ViewIndex, NumViews, SceneTextures, SortedLightSet, ESubstrateTileType::EComplex);
-			InternalRenderSimpleLightsStandardDeferred(GraphBuilder, Scene, View, ViewIndex, NumViews, SceneTextures, SortedLightSet, ESubstrateTileType::ESingle);
-			InternalRenderSimpleLightsStandardDeferred(GraphBuilder, Scene, View, ViewIndex, NumViews, SceneTextures, SortedLightSet, ESubstrateTileType::ESimple);
+			if (Substrate::GetSubstrateUsesTileType(View, ESubstrateTileType::EComplex))
+			{
+				InternalRenderSimpleLightsStandardDeferred(GraphBuilder, Scene, View, ViewIndex, NumViews, SceneTextures, SortedLightSet, ESubstrateTileType::EComplex);
+			}
+			if (Substrate::GetSubstrateUsesTileType(View, ESubstrateTileType::ESingle))
+			{
+				InternalRenderSimpleLightsStandardDeferred(GraphBuilder, Scene, View, ViewIndex, NumViews, SceneTextures, SortedLightSet, ESubstrateTileType::ESingle);
+			}
+			if (Substrate::GetSubstrateUsesTileType(View, ESubstrateTileType::ESimple))
+			{
+				InternalRenderSimpleLightsStandardDeferred(GraphBuilder, Scene, View, ViewIndex, NumViews, SceneTextures, SortedLightSet, ESubstrateTileType::ESimple);
+			}
 		}
 		else
 		{
diff --git a/Engine/Source/Runtime/Renderer/Private/PathTracing.cpp b/Engine/Source/Runtime/Renderer/Private/PathTracing.cpp
--- a/Engine/Source/Runtime/Renderer/Private/PathTracing.cpp
+++ b/Engine/Source/Runtime/Renderer/Private/PathTracing.cpp
@@ -3295,7 +3295,7 @@ static FPathTracingRG::FPermutationDomain GetPathTracingRGPermutation(const FVie
 	const bool bUseAdaptiveSampling = bUseExperimental && CVarPathTracingAdaptiveSampling.GetValueOnRenderThread() != 0;
 	// NOTE: the decision about when to enable clouds involves checking lots of things, so rely on the presence of a valid index to signify that clouds are ready
 	const bool bUseCloudShader = uint32(View.PathTracingVolumetricCloudCallableShaderIndex) < Scene.RayTracingSBT.NumCallableShaderSlots;
-	const bool bHasComplexSpecialRenderPath = Substrate::IsSubstrateEnabled() && Scene.SubstrateSceneData.bUsesComplexSpecialRenderPath;
+	const bool bHasComplexSpecialRenderPath = Substrate::IsSubstrateEnabled() && UsesSubstrateTileType(Scene.SubstrateSceneData.UsesTileTypeMask, ESubstrateTileType::EComplexSpecial);
 	const bool bUseSER = GRHIGlobals.SupportsShaderExecutionReordering && CVarPathTracingShaderExecutionReordering.GetValueOnRenderThread();
 	const bool bUseTraceOpaqueFirst = !bUseExperimental || CVarPathTracingTraceOpaqueFirst.GetValueOnRenderThread();
 
diff --git a/Engine/Source/Runtime/Renderer/Private/SceneVisibility.cpp b/Engine/Source/Runtime/Renderer/Private/SceneVisibility.cpp
--- a/Engine/Source/Runtime/Renderer/Private/SceneVisibility.cpp
+++ b/Engine/Source/Runtime/Renderer/Private/SceneVisibility.cpp
@@ -1334,7 +1334,7 @@ void FRelevancePacket::Finalize()
 	WriteView.bUsesCustomStencil |= bUsesCustomStencil;
 	WriteView.SubstrateViewData.MaxClosurePerPixel = FMath::Max(WriteView.SubstrateViewData.MaxClosurePerPixel, 8u - FMath::CountLeadingZeros8(SubstrateClosureCountMask));
 	WriteView.SubstrateViewData.MaxBytesPerPixel = FMath::Max(WriteView.SubstrateViewData.MaxBytesPerPixel, SubstrateUintPerPixel * 4u);
-	WriteView.SubstrateViewData.bUsesComplexSpecialRenderPath |= bUsesComplexSpecialRenderPath;
+	WriteView.SubstrateViewData.UsesTileTypeMask |= SubstrateTileTypeMask;
 	WriteView.SubstrateViewData.bUsesAnisotropy |= bUsesAnisotropy;
 	DirtyIndirectLightingCacheBufferPrimitives.AppendTo(WriteView.DirtyIndirectLightingCacheBufferPrimitives);
 
@@ -1406,7 +1406,7 @@ void FRelevancePacket::ComputeRelevance(FDynamicPrimitiveIndexList& DynamicPrimi
 
 	CombinedShadingModelMask = 0;
 	SubstrateUintPerPixel = 0;
-	bUsesComplexSpecialRenderPath = false;
+	SubstrateTileTypeMask = 0;
 	bUsesAnisotropy = false;
 	SubstrateClosureCountMask = 0;
 	bSceneHasSkyMaterial = 0;
@@ -2050,8 +2050,8 @@ void FRelevancePacket::ComputeRelevance(FDynamicPrimitiveIndexList& DynamicPrimi
 		}
 
 		CombinedShadingModelMask |= ViewRelevance.ShadingModelMask;
-		SubstrateUintPerPixel = FMath::Max(SubstrateUintPerPixel, ViewRelevance.SubstrateUintPerPixel);
-		bUsesComplexSpecialRenderPath |= ViewRelevance.bUsesComplexSpecialRenderPath;
+		SubstrateUintPerPixel = FMath::Max(SubstrateUintPerPixel, uint8(ViewRelevance.SubstrateUintPerPixel));
+		SubstrateTileTypeMask |= ViewRelevance.SubstrateTileTypeMask;
 		bUsesAnisotropy |= ViewRelevance.bUsesAnisotropy;
 		SubstrateClosureCountMask |= ViewRelevance.SubstrateClosureCountMask;
 		bUsesGlobalDistanceField |= ViewRelevance.bUsesGlobalDistanceField;
@@ -2269,7 +2269,7 @@ void FComputeAndMarkRelevance::Finalize()
 				WriteView.bUsesCustomStencil |= (CombinedRelevance.CustomDepthStencilUsageMask & (1 << 1)) > 0;
 				WriteView.SubstrateViewData.MaxClosurePerPixel = FMath::Max(WriteView.SubstrateViewData.MaxClosurePerPixel, 8u - FMath::CountLeadingZeros8(CombinedRelevance.SubstrateClosureCountMask));
 				WriteView.SubstrateViewData.MaxBytesPerPixel = FMath::Max(WriteView.SubstrateViewData.MaxBytesPerPixel, CombinedRelevance.SubstrateUintPerPixel * 4u);
-				WriteView.SubstrateViewData.bUsesComplexSpecialRenderPath |= CombinedRelevance.bUsesComplexSpecialRenderPath;
+				WriteView.SubstrateViewData.UsesTileTypeMask |= CombinedRelevance.SubstrateTileTypeMask;
 				WriteView.SubstrateViewData.bUsesAnisotropy |= CombinedRelevance.bUsesAnisotropy;
 			}
 		}
@@ -2355,7 +2355,7 @@ void FComputeAndMarkRelevance::Finalize()
 			SecondaryView->bUsesCustomStencil = View.bUsesCustomStencil;
 			SecondaryView->SubstrateViewData.MaxClosurePerPixel = View.SubstrateViewData.MaxClosurePerPixel;
 			SecondaryView->SubstrateViewData.MaxBytesPerPixel = View.SubstrateViewData.MaxClosurePerPixel;
-			SecondaryView->SubstrateViewData.bUsesComplexSpecialRenderPath = View.SubstrateViewData.bUsesComplexSpecialRenderPath;
+			SecondaryView->SubstrateViewData.UsesTileTypeMask = View.SubstrateViewData.UsesTileTypeMask;
 			SecondaryView->SubstrateViewData.bUsesAnisotropy = View.SubstrateViewData.bUsesAnisotropy;
 		}
 	}
diff --git a/Engine/Source/Runtime/Renderer/Private/SceneVisibilityPrivate.h b/Engine/Source/Runtime/Renderer/Private/SceneVisibilityPrivate.h
--- a/Engine/Source/Runtime/Renderer/Private/SceneVisibilityPrivate.h
+++ b/Engine/Source/Runtime/Renderer/Private/SceneVisibilityPrivate.h
@@ -776,7 +776,7 @@ private:
 	uint16 CombinedShadingModelMask = 0;
 	uint8 SubstrateUintPerPixel = 0;
 	uint8 SubstrateClosureCountMask = 0;
-	bool bUsesComplexSpecialRenderPath = false;
+	uint8 SubstrateTileTypeMask = 0;
 	bool bUsesAnisotropy = false;
 	bool bHasDistortionPrimitives = false;
 	bool bHasCustomDepthPrimitives = false;
diff --git a/Engine/Source/Runtime/Renderer/Private/Substrate/Substrate.cpp b/Engine/Source/Runtime/Renderer/Private/Substrate/Substrate.cpp
--- a/Engine/Source/Runtime/Renderer/Private/Substrate/Substrate.cpp
+++ b/Engine/Source/Runtime/Renderer/Private/Substrate/Substrate.cpp
@@ -94,33 +94,15 @@ void FSubstrateViewData::Reset()
 	// * SubstrateClosureCountMask
 	static_assert(SUBSTRATE_MAX_CLOSURE_COUNT <= 8u);
 
-	// Propagate bUsesComplexSpecialRenderPath after reset as we use the per-view (vs. the per-scene) 
+	// Propagate values after reset as we use the per-view (vs. the per-scene) 
 	// value to know if a view needs special complex path or not
-	const bool OldUsesComplexSpecialRenderPath = bUsesComplexSpecialRenderPath;
+	const uint8 OldUsesTileTypeMask = UsesTileTypeMask;
 	const bool OldUsesAnistropy = bUsesAnisotropy;
 	*this = FSubstrateViewData();
-	bUsesComplexSpecialRenderPath = OldUsesComplexSpecialRenderPath;
+	UsesTileTypeMask = OldUsesTileTypeMask;
 	bUsesAnisotropy = OldUsesAnistropy;
 }
 
-const TCHAR* ToString(ESubstrateTileType Type)
-{
-	switch (Type)
-	{
-	case ESubstrateTileType::ESimple:							return TEXT("Simple");
-	case ESubstrateTileType::ESingle:							return TEXT("Single");
-	case ESubstrateTileType::EComplex:							return TEXT("Complex");
-	case ESubstrateTileType::EComplexSpecial:					return TEXT("ComplexSpecial");
-	case ESubstrateTileType::EOpaqueRoughRefraction:			return TEXT("Opaque/RoughRefraction");
-	case ESubstrateTileType::EOpaqueRoughRefractionSSSWithout:	return TEXT("Opaque/RoughRefraction/SSSWithout");
-	case ESubstrateTileType::EDecalSimple:						return TEXT("Decal/Simple");
-	case ESubstrateTileType::EDecalSingle:						return TEXT("Decal/Single");
-	case ESubstrateTileType::EDecalComplex:						return TEXT("Decal/Complex");
-
-	}
-	return TEXT("Unknown");
-}
-
 namespace MegaLights
 {
 	uint32 GetStateFrameIndex(FSceneViewState*);
@@ -208,13 +190,27 @@ bool Is8bitTileCoordEnabled()
 	return CVarSubstrateTileCoord8Bits.GetValueOnAnyThread() > 0 ? 1 : 0;
 }
 
+bool GetSubstrateUsesTileType(const FViewInfo& View, ESubstrateTileType TileType)
+{
+	if (Substrate::IsSubstrateEnabled())
+	{
+		static_assert(uint32(ESubstrateTileType::EComplexSpecial) < sizeof(uint8)*8u);
+		check(TileType <= ESubstrateTileType::EComplexSpecial);		
+
+		const bool bCompatible = TileType == ESubstrateTileType::EComplexSpecial ? !Substrate::IsSubstrateBlendableGBufferEnabled(View.GetShaderPlatform()) : true;
+		const uint8 Mask = GetSubstrateTileTypeAsUint8(TileType);
+		return bCompatible && UsesSubstrateTileType(View.SubstrateViewData.UsesTileTypeMask, TileType);
+	}
+	return false;
+}	
+
 bool GetSubstrateUsesComplexSpecialPath(const FViewInfo& View)
 {
 	if (Substrate::IsSubstrateEnabled())
 	{
 		// Use the per-view value rather than the per-scene data to have more accurate dispatching of special complex tiles 
 		// and avoid unecessary empty-dispatch
-		return View.SubstrateViewData.bUsesComplexSpecialRenderPath;
+		return UsesSubstrateTileType(View.SubstrateViewData.UsesTileTypeMask, ESubstrateTileType::EComplexSpecial);
 	}
 	return false;
 }
@@ -285,13 +281,11 @@ static void InitialiseSubstrateViewData(FRDGBuilder& GraphBuilder, FViewInfo& Vi
 			const bool bRequest8bit = Substrate::Is8bitTileCoordEnabled() && (TileResolution.X <= 256 && TileResolution.Y <= 256);
 			Out.TileEncoding = bRequest8bit ? SUBSTRATE_TILE_ENCODING_8BITS : SUBSTRATE_TILE_ENCODING_16BITS;
 
-			bool bUsesComplexSpecialRenderPath = SceneData.bUsesComplexSpecialRenderPath; // Use the Scene temporally stable bUsesComplexSpecialRenderPath to reduce buffer reallocation.
-
 			Out.ClassificationTileListBufferOffset[ESubstrateTileType::ESimple]							= 0;
 			Out.ClassificationTileListBufferOffset[ESubstrateTileType::ESingle]							= Out.ClassificationTileListBufferOffset[ESubstrateTileType::ESimple]							+ RegularTileCount;
 			Out.ClassificationTileListBufferOffset[ESubstrateTileType::EComplex]						= Out.ClassificationTileListBufferOffset[ESubstrateTileType::ESingle]							+ RegularTileCount;
 			Out.ClassificationTileListBufferOffset[ESubstrateTileType::EComplexSpecial]					= Out.ClassificationTileListBufferOffset[ESubstrateTileType::EComplex]							+ RegularTileCount;
-			Out.ClassificationTileListBufferOffset[ESubstrateTileType::EOpaqueRoughRefraction]			= Out.ClassificationTileListBufferOffset[ESubstrateTileType::EComplexSpecial]					+ (bUsesComplexSpecialRenderPath ? RegularTileCount : 4);
+			Out.ClassificationTileListBufferOffset[ESubstrateTileType::EOpaqueRoughRefraction]			= Out.ClassificationTileListBufferOffset[ESubstrateTileType::EComplexSpecial]					+ (UsesSubstrateTileType(SceneData.UsesTileTypeMask, ESubstrateTileType::EComplexSpecial) ? RegularTileCount : 4);
 			Out.ClassificationTileListBufferOffset[ESubstrateTileType::EOpaqueRoughRefractionSSSWithout]= Out.ClassificationTileListBufferOffset[ESubstrateTileType::EOpaqueRoughRefraction]			+ RoughTileCount;
 			Out.ClassificationTileListBufferOffset[ESubstrateTileType::EDecalSimple]					= Out.ClassificationTileListBufferOffset[ESubstrateTileType::EOpaqueRoughRefractionSSSWithout]	+ RoughTileCount;
 			Out.ClassificationTileListBufferOffset[ESubstrateTileType::EDecalSingle]					= Out.ClassificationTileListBufferOffset[ESubstrateTileType::EDecalSimple]						+ DecalTileCount;
@@ -341,6 +335,17 @@ static void InitialiseSubstrateViewData(FRDGBuilder& GraphBuilder, FViewInfo& Vi
 			BindSubstrateGlobalUniformParameters(GraphBuilder, &Out, bNeedsMaterialBuffer, bBlendableGBufferEnabled, *SubstrateUniformParameters);
 			Out.SubstrateGlobalUniformParameters = GraphBuilder.CreateUniformBuffer(SubstrateUniformParameters);
 		}
+
+		// Detect the complexity and mark less-compled tile types as available.
+		// The CPU complexity is conservative, and indicate the max complexity that an material can have. However in order to dispatch all relevant 
+		// tile types, we need to dispatch tile type with lower complexity as some part of the object might have lower complexity. 
+		{
+			const uint32 FirstSetBit = 8u - FMath::CountLeadingZeros8(View.SubstrateViewData.UsesTileTypeMask);
+			for (uint32 Bit = 0; Bit < FirstSetBit; ++Bit)
+			{
+				View.SubstrateViewData.UsesTileTypeMask |= 1u<<Bit;
+			}
+		}
 	}
 }
 
@@ -416,15 +421,15 @@ void InitialiseSubstrateFrameSceneData(FRDGBuilder& GraphBuilder, FSceneRenderer
 
 	// Reset Substrate scene data
 	{
-		const uint32 MinBytesPerPixel = Out.PersistentMaxBytesPerPixel;
-		const uint32 MaxClosureCount  = Out.PersistentMaxClosurePerPixel;
-		const bool bUsesComplexSpecialRenderPath = Out.bUsesComplexSpecialRenderPath;
-		const bool bUsesAnisotropy = Out.bUsesAnisotropy;
+		const uint32 MinBytesPerPixel	= Out.PersistentMaxBytesPerPixel;
+		const uint32 MaxClosureCount	= Out.PersistentMaxClosurePerPixel;
+		const uint8 UsesTileTypeMask	= Out.UsesTileTypeMask;
+		const bool bUsesAnisotropy		= Out.bUsesAnisotropy;
 		Out = FSubstrateSceneData();
-		Out.PersistentMaxBytesPerPixel    = MinBytesPerPixel;
-		Out.PersistentMaxClosurePerPixel  = MaxClosureCount;
-		Out.bUsesComplexSpecialRenderPath = bUsesComplexSpecialRenderPath;
-		Out.bUsesAnisotropy 			  = bUsesAnisotropy;
+		Out.PersistentMaxBytesPerPixel	= MinBytesPerPixel;
+		Out.PersistentMaxClosurePerPixel= MaxClosureCount;
+		Out.UsesTileTypeMask 			= UsesTileTypeMask;
+		Out.bUsesAnisotropy				= bUsesAnisotropy;
 	}
 
 	auto UpdateMaterialBufferToTiledResolution = [](FIntPoint InBufferSizeXY, FIntPoint& OutMaterialBufferSizeXY)
@@ -475,7 +480,7 @@ void InitialiseSubstrateFrameSceneData(FRDGBuilder& GraphBuilder, FSceneRenderer
 			{
 				Out.PersistentMaxBytesPerPixel = FMath::Max(Out.PersistentMaxBytesPerPixel, View.SubstrateViewData.MaxBytesPerPixel);
 				Out.PersistentMaxClosurePerPixel = FMath::Max(Out.PersistentMaxClosurePerPixel, View.SubstrateViewData.MaxClosurePerPixel);
-				Out.bUsesComplexSpecialRenderPath |= View.SubstrateViewData.bUsesComplexSpecialRenderPath;
+				Out.UsesTileTypeMask |= View.SubstrateViewData.UsesTileTypeMask;
 				Out.bUsesAnisotropy |= View.SubstrateViewData.bUsesAnisotropy;
 			}
 		}
@@ -1437,13 +1442,22 @@ void AddSubstrateStencilPass(
 		RDG_EVENT_SCOPE_CONDITIONAL(GraphBuilder, Views.Num() > 1, "View%d", i);
 
 		const FViewInfo& View = Views[i];
-		if (GetSubstrateUsesComplexSpecialPath(View))
+		if (Substrate::GetSubstrateUsesTileType(View, ESubstrateTileType::EComplexSpecial))
 		{
 			AddSubstrateInternalClassificationTilePass(GraphBuilder, View, &SceneTextures.Depth.Target, nullptr, ESubstrateTileType::EComplexSpecial);
 		}
-		AddSubstrateInternalClassificationTilePass(GraphBuilder, View, &SceneTextures.Depth.Target, nullptr, ESubstrateTileType::EComplex);
-		AddSubstrateInternalClassificationTilePass(GraphBuilder, View, &SceneTextures.Depth.Target, nullptr, ESubstrateTileType::ESingle);
-		AddSubstrateInternalClassificationTilePass(GraphBuilder, View, &SceneTextures.Depth.Target, nullptr, ESubstrateTileType::ESimple);
+		if (Substrate::GetSubstrateUsesTileType(View, ESubstrateTileType::EComplex))
+		{
+			AddSubstrateInternalClassificationTilePass(GraphBuilder, View, &SceneTextures.Depth.Target, nullptr, ESubstrateTileType::EComplex);
+		}
+		if (Substrate::GetSubstrateUsesTileType(View, ESubstrateTileType::ESingle))
+		{
+			AddSubstrateInternalClassificationTilePass(GraphBuilder, View, &SceneTextures.Depth.Target, nullptr, ESubstrateTileType::ESingle);
+		}
+		if (Substrate::GetSubstrateUsesTileType(View, ESubstrateTileType::ESimple))
+		{
+			AddSubstrateInternalClassificationTilePass(GraphBuilder, View, &SceneTextures.Depth.Target, nullptr, ESubstrateTileType::ESimple);
+		}
 	}
 }
 
@@ -1720,11 +1734,15 @@ void AddSubstrateMaterialClassificationPass(FRDGBuilder& GraphBuilder, const FMi
 					PassParameters->TileIndirectBuffer,
 					TileTypeDispatchIndirectArgOffset(TileType));
 			};
-			if (GetSubstrateUsesComplexSpecialPath(View))
+
+			if (Substrate::GetSubstrateUsesTileType(View, ESubstrateTileType::EComplexSpecial))
 			{
 				MarkClosureTilePass(ESubstrateTileType::EComplexSpecial);
 			}
-			MarkClosureTilePass(ESubstrateTileType::EComplex);
+			if (Substrate::GetSubstrateUsesTileType(View, ESubstrateTileType::EComplex))
+			{
+				MarkClosureTilePass(ESubstrateTileType::EComplex);
+			}
 		}
 
 		// Tile indirect dispatch args conversion
diff --git a/Engine/Source/Runtime/Renderer/Private/Substrate/Substrate.h b/Engine/Source/Runtime/Renderer/Private/Substrate/Substrate.h
--- a/Engine/Source/Runtime/Renderer/Private/Substrate/Substrate.h
+++ b/Engine/Source/Runtime/Renderer/Private/Substrate/Substrate.h
@@ -88,23 +88,6 @@ BEGIN_GLOBAL_SHADER_PARAMETER_STRUCT(FSubstratePublicGlobalUniformParameters, RE
 	SHADER_PARAMETER_STRUCT_INCLUDE(FSubstratePublicParameters, Public)
 END_GLOBAL_SHADER_PARAMETER_STRUCT()
 
-// This must map to the SUBSTRATE_TILE_TYPE defines.
-enum ESubstrateTileType : uint32
-{
-	ESimple								= SUBSTRATE_TILE_TYPE_SIMPLE,
-	ESingle								= SUBSTRATE_TILE_TYPE_SINGLE,
-	EComplex							= SUBSTRATE_TILE_TYPE_COMPLEX,
-	EComplexSpecial						= SUBSTRATE_TILE_TYPE_COMPLEX_SPECIAL,
-	EOpaqueRoughRefraction				= SUBSTRATE_TILE_TYPE_ROUGH_REFRACT,
-	EOpaqueRoughRefractionSSSWithout	= SUBSTRATE_TILE_TYPE_ROUGH_REFRACT_SSS_WITHOUT,
-	EDecalSimple						= SUBSTRATE_TILE_TYPE_DECAL_SIMPLE,
-	EDecalSingle						= SUBSTRATE_TILE_TYPE_DECAL_SINGLE,
-	EDecalComplex						= SUBSTRATE_TILE_TYPE_DECAL_COMPLEX,
-	ECount
-};
-
-const TCHAR* ToString(ESubstrateTileType Type);
-
 struct FSubstrateSceneData
 {
 	// Track max BytesPerPixel / ClosurePerPixel amoung all views
@@ -114,7 +97,7 @@ struct FSubstrateSceneData
 	// Track max BytesPerPixel / ClosurePerPixel ever encountered since the scene was created
 	uint32 PersistentMaxBytesPerPixel = 0;
 	uint32 PersistentMaxClosurePerPixel = 0;
-	bool bUsesComplexSpecialRenderPath = false;
+	uint8 UsesTileTypeMask = 0;
 	bool bUsesAnisotropy = false;
 
 	// Current max BytesPerPixel / ClosurePerPixel
@@ -165,7 +148,7 @@ struct FSubstrateViewData
 	uint32 MaxBytesPerPixel = 0;
 
 	// True if any material requires the complex special path (e.g. glints or SpecularLUT)
-	bool bUsesComplexSpecialRenderPath = false;
+	uint8 UsesTileTypeMask = 0;
 	bool bUsesAnisotropy = false;
 
 	FIntPoint TileCount  = FIntPoint(0, 0);
@@ -233,9 +216,11 @@ constexpr uint32 StencilBit_ComplexSpecial	= 0x80; // In sync with SceneRenderTa
 
 FIntPoint GetSubstrateTextureResolution(const FViewInfo& View, const FIntPoint& InResolution);
 uint32 GetSubstrateMaxClosureCount(const FViewInfo& View);
+UE_DEPRECATED(5.7, "Please, use the GetSubstrateUsesTileType() function instead")	
 bool GetSubstrateUsesComplexSpecialPath(const FViewInfo& View);
 bool GetSubstrateUsesAnisotropy(const FViewInfo& View);
 bool UsesSubstrateMaterialBuffer(EShaderPlatform In);
+bool GetSubstrateUsesTileType(const FViewInfo& View, ESubstrateTileType TileType);
 
 void InitialiseSubstrateFrameSceneData(FRDGBuilder& GraphBuilder, FSceneRenderer& SceneRenderer);
 
diff --git a/Engine/Source/Runtime/Renderer/Private/Substrate/SubstrateVisualize.cpp b/Engine/Source/Runtime/Renderer/Private/Substrate/SubstrateVisualize.cpp
--- a/Engine/Source/Runtime/Renderer/Private/Substrate/SubstrateVisualize.cpp
+++ b/Engine/Source/Runtime/Renderer/Private/Substrate/SubstrateVisualize.cpp
@@ -1160,7 +1160,7 @@ static void AddVisualizeSystemInfoPasses(FRDGBuilder& GraphBuilder, const FViewI
 				DrawPosY += 17;
 			};
 
-			Canvas.DrawTile(DrawPosX - 10.0f, DrawPosY - 10.0f, 350.0f, 700.0f, 0, 0, 1, 1, FLinearColor(0.0, 0.0, 0.0, 0.3));
+			Canvas.DrawTile(DrawPosX - 10.0f, DrawPosY - 10.0f, 350.0f, 800.0f, 0, 0, 1, 1, FLinearColor(0.0, 0.0, 0.0, 0.3));
 
 			const FLinearColor White = FLinearColor::White;
 			const FLinearColor Yellow = FLinearColor(1, 1, 0, 1);
@@ -1274,7 +1274,7 @@ static void AddVisualizeSystemInfoPasses(FRDGBuilder& GraphBuilder, const FViewI
 			}
 			NewLine();
 
-			PrintS(DrawPosX, DrawPosY, TEXT("Material Tiles"), Yellow);
+			PrintS(DrawPosX, DrawPosY, TEXT("Material Tiles (GPU)"), Yellow);
 			NewLine();
 			{
 				PrintS(DrawPosX, DrawPosY, TEXT("  # Simple tiles"), White);
@@ -1325,12 +1325,18 @@ static void AddVisualizeSystemInfoPasses(FRDGBuilder& GraphBuilder, const FViewI
 			}
 			NewLine();
 
-			PrintS(DrawPosX, DrawPosY, TEXT("In View"), Yellow);
+			PrintS(DrawPosX, DrawPosY, TEXT("In View (CPU)"), Yellow);
 			NewLine();
 			{
-				PrintS(DrawPosX, DrawPosY, TEXT("  Special Complex"), White);
-				PrintBool(DrawPosX + 150, DrawPosY, Substrate::GetSubstrateUsesComplexSpecialPath(View));
-				NewLine();
+				for (uint32 TileTypeIndex = 0; TileTypeIndex <= ESubstrateTileType::EComplexSpecial; ++TileTypeIndex)
+				{
+					const ESubstrateTileType TileType = ESubstrateTileType(TileTypeIndex);
+					const FString TileTypeName = ToString(ESubstrateTileType(TileType));
+					const FString FormattedTypeType = FString::Format(TEXT("  {0}"), {*TileTypeName});
+					PrintS(DrawPosX, DrawPosY, FormattedTypeType, White);
+					PrintBool(DrawPosX + 150, DrawPosY, Substrate::GetSubstrateUsesTileType(View, TileType));
+					NewLine();
+				}
 				PrintS(DrawPosX, DrawPosY, TEXT("  Anisotropy"), White);
 				PrintBool(DrawPosX + 150, DrawPosY, Substrate::GetSubstrateUsesAnisotropy(View));
 				NewLine();
@@ -1486,13 +1492,22 @@ FScreenPassTexture AddSubstrateDebugPasses(FRDGBuilder& GraphBuilder, const FVie
 		}
 		else if (DebugMode == FSubstrateViewMode::MaterialClassification)
 		{
-			if (GetSubstrateUsesComplexSpecialPath(View))
+			if (Substrate::GetSubstrateUsesTileType(View, ESubstrateTileType::EComplexSpecial))
 			{
 				AddSubstrateInternalClassificationTilePass(GraphBuilder, View, nullptr, &ScreenPassSceneColor.Texture, ESubstrateTileType::EComplexSpecial, bDebugPass);
 			}
-			AddSubstrateInternalClassificationTilePass(GraphBuilder, View, nullptr, &ScreenPassSceneColor.Texture, ESubstrateTileType::EComplex, bDebugPass);
-			AddSubstrateInternalClassificationTilePass(GraphBuilder, View, nullptr, &ScreenPassSceneColor.Texture, ESubstrateTileType::ESingle, bDebugPass);
-			AddSubstrateInternalClassificationTilePass(GraphBuilder, View, nullptr, &ScreenPassSceneColor.Texture, ESubstrateTileType::ESimple, bDebugPass);
+			if (Substrate::GetSubstrateUsesTileType(View, ESubstrateTileType::EComplex))
+			{
+				AddSubstrateInternalClassificationTilePass(GraphBuilder, View, nullptr, &ScreenPassSceneColor.Texture, ESubstrateTileType::EComplex, bDebugPass);
+			}
+			if (Substrate::GetSubstrateUsesTileType(View, ESubstrateTileType::ESingle))
+			{
+				AddSubstrateInternalClassificationTilePass(GraphBuilder, View, nullptr, &ScreenPassSceneColor.Texture, ESubstrateTileType::ESingle, bDebugPass);
+			}
+			if (Substrate::GetSubstrateUsesTileType(View, ESubstrateTileType::ESimple))
+			{
+				AddSubstrateInternalClassificationTilePass(GraphBuilder, View, nullptr, &ScreenPassSceneColor.Texture, ESubstrateTileType::ESimple, bDebugPass);
+			}
 		}
 	}
 
