Enabled r.SceneCulling.Precomputed by default as this is an important optimization for streaming scenes with large numbers of instances.

diff --git a/Engine/Source/Runtime/Renderer/Private/SceneCulling/SceneCulling.cpp b/Engine/Source/Runtime/Renderer/Private/SceneCulling/SceneCulling.cpp
--- a/Engine/Source/Runtime/Renderer/Private/SceneCulling/SceneCulling.cpp
+++ b/Engine/Source/Runtime/Renderer/Private/SceneCulling/SceneCulling.cpp
@@ -108,7 +108,7 @@ LLM_DEFINE_TAG(SceneCulling, NAME_None, NAME_None, GET_STATFNAME(STAT_SceneCulli
 
 static TAutoConsoleVariable<int32> CVarSceneCullingPrecomputed(
 	TEXT("r.SceneCulling.Precomputed"), 
-	0, 
+	1, 
 	TEXT("Enable/Disable precomputed spatial hashes for scene culling."),
 	ECVF_RenderThreadSafe | ECVF_ReadOnly);
 
