[HWRT] Automatically clear cached instances lock when ray tracing scene is reset.
- cached instances can't be added while there are dynamic instances in the Instances arrays.

diff --git a/Engine/Source/Runtime/Renderer/Private/RayTracing/RayTracing.cpp b/Engine/Source/Runtime/Renderer/Private/RayTracing/RayTracing.cpp
--- a/Engine/Source/Runtime/Renderer/Private/RayTracing/RayTracing.cpp
+++ b/Engine/Source/Runtime/Renderer/Private/RayTracing/RayTracing.cpp
@@ -1756,6 +1756,7 @@ namespace RayTracing
 		check(!FMaterialRenderProxy::HasDeferredUniformExpressionCacheRequests());
 
 		RayTracingSBT.ResetDynamicAllocationData();
+		RayTracingScene.LockCachedInstances();
 
 		INC_DWORD_STAT_BY(STAT_VisibleRayTracingPrimitives, TaskData.StaticPrimitives.Num() + TaskData.CachedStaticPrimitives.Num() + TaskData.DynamicPrimitivesIndices.Num());
 
@@ -1776,8 +1777,6 @@ namespace RayTracing
 			{
 				FTaskTagScope TaskTagScope(ETaskTag::EParallelRenderingThread);
 
-				Scene.RayTracingScene.LockCachedInstances();
-
 				AddDynamicInstancesToRayTracingScene(TaskData, RayTracingScene, RayTracingSBT, SharedBufferGenerationID);
 
 				RayTracingScene.bUsesLightingChannels |= TaskData.bUsesLightingChannels && TaskData.SceneOptions.bLightingChannelsUsingAHS;
@@ -1792,13 +1791,10 @@ namespace RayTracing
 					View.VisibleRayTracingShaderBindings);
 
 				AddVisibleCachedInstances(
-					Scene.RayTracingScene,
-					Scene.RayTracingSBT,
+					RayTracingScene,
+					RayTracingSBT,
 					TaskData.CachedStaticPrimitives,
 					View.VisibleRayTracingShaderBindings);
-
-				Scene.RayTracingScene.UnlockCachedInstances();
-
 			}, TaskData.GatherRelevantStaticPrimitivesTask);
 
 		// Scene init task can run only when all pre-init tasks are complete (including culling tasks that are spawned while adding instances)
diff --git a/Engine/Source/Runtime/Renderer/Private/RayTracing/RayTracingScene.cpp b/Engine/Source/Runtime/Renderer/Private/RayTracing/RayTracingScene.cpp
--- a/Engine/Source/Runtime/Renderer/Private/RayTracing/RayTracingScene.cpp
+++ b/Engine/Source/Runtime/Renderer/Private/RayTracing/RayTracingScene.cpp
@@ -966,6 +966,8 @@ void FRayTracingScene::Reset()
 	Allocator.Flush();
 
 	bUsesLightingChannels = false;
+
+	bCachedInstancesLocked = false;
 }
 
 void FRayTracingScene::EndFrame()
diff --git a/Engine/Source/Runtime/Renderer/Private/RayTracing/RayTracingScene.h b/Engine/Source/Runtime/Renderer/Private/RayTracing/RayTracingScene.h
--- a/Engine/Source/Runtime/Renderer/Private/RayTracing/RayTracingScene.h
+++ b/Engine/Source/Runtime/Renderer/Private/RayTracing/RayTracingScene.h
@@ -123,19 +123,13 @@ public:
 
 	void EndFrame();
 
-	// Prevent cached instances being added/freed (via validation checks).
+	// Prevent cached instances being added/freed (via validation checks) until Reset() or EndFrame()
 	// Must be done before adding transient instances.
 	void LockCachedInstances()
 	{
 		bCachedInstancesLocked = true;
 	}
 
-	// Must be done after all transient instances are added.
-	void UnlockCachedInstances()
-	{
-		bCachedInstancesLocked = false;
-	}
-
 	bool SetInstanceExtraDataBufferEnabled(bool bEnabled);
 	bool SetTracingFeedbackEnabled(bool bEnabled);
 	bool SetInstanceDebugDataEnabled(bool bEnabled);
