[HWRT] Extended FRayTracingScene with support for multiple views.
- Instances array in each layer are shared across all views.
- Views must be added to the ray tracing scene every frame using ViewKey which enables persistent state.
- Persistent state associated with each view is removed when the view state is deleted.

diff --git a/Engine/Plugins/FX/Niagara/Source/NiagaraShader/Private/NiagaraAsyncGpuTraceProviderHwrt.cpp b/Engine/Plugins/FX/Niagara/Source/NiagaraShader/Private/NiagaraAsyncGpuTraceProviderHwrt.cpp
--- a/Engine/Plugins/FX/Niagara/Source/NiagaraShader/Private/NiagaraAsyncGpuTraceProviderHwrt.cpp
+++ b/Engine/Plugins/FX/Niagara/Source/NiagaraShader/Private/NiagaraAsyncGpuTraceProviderHwrt.cpp
@@ -352,7 +352,7 @@ bool FNiagaraAsyncGpuTraceProviderHwrt::IsAvailable() const
 		return false;
 	}
 
-	return UE::FXRenderingUtils::RayTracing::HasRayTracingScene(Dispatcher->GetSceneInterface());
+	return UE::FXRenderingUtils::RayTracing::HasRayTracingScene(*Dispatcher->GetSceneInterface());
 }
 
 void FNiagaraAsyncGpuTraceProviderHwrt::PostRenderOpaque(FRHICommandList& RHICmdList, TConstStridedView<FSceneView> Views, TUniformBufferRef<FSceneUniformParameters> SceneUniformBufferRHI, FCollisionGroupHashMap* CollisionGroupHash)
@@ -360,11 +360,12 @@ void FNiagaraAsyncGpuTraceProviderHwrt::PostRenderOpaque(FRHICommandList& RHICmd
 	check(IsAvailable());
 	check(Views.Num() > 0);
 
-	const FSceneView& ReferenceView = Views[0];
+	const FSceneView& ReferenceView = Views[0]; // TODO: Should use appropriate view instead
 	FSceneInterface* Scene = Dispatcher->GetSceneInterface();
 
-	if (UE::FXRenderingUtils::RayTracing::HasRayTracingScene(Scene))
+	if (UE::FXRenderingUtils::RayTracing::HasRayTracingScene(*Scene))
 	{
+		TLASSRV = UE::FXRenderingUtils::RayTracing::GetRayTracingSceneView(RHICmdList, *Scene, ReferenceView);
 		ViewUniformBuffer = ReferenceView.ViewUniformBuffer;
 
 		if (!SupportsNiagaraAsyncGpuTraceHwrtInline() && GRHISupportsRayTracingShaders)
@@ -432,7 +433,7 @@ void FNiagaraAsyncGpuTraceProviderHwrt::IssueTraces(FRHICommandList& RHICmdList,
 	}
 
 	FSceneInterface* Scene = Dispatcher->GetSceneInterface();
-	Params.TLAS = UE::FXRenderingUtils::RayTracing::GetRayTracingSceneView(RHICmdList, Scene);
+	Params.TLAS = TLASSRV;
 	Params.Rays = Request.TracesBuffer->SRV;
 	Params.RaysOffset = Request.TracesOffset;
 	Params.CollisionOutput = Request.ResultsBuffer->UAV;
@@ -500,6 +501,7 @@ void FNiagaraAsyncGpuTraceProviderHwrt::Reset()
 {
 	RayTracingPipelineState = nullptr;
 	RayTracingSBT = nullptr;
+	TLASSRV = nullptr;
 	ViewUniformBuffer = TUniformBufferRef<FViewUniformShaderParameters>();
 }
 
diff --git a/Engine/Plugins/FX/Niagara/Source/NiagaraShader/Private/NiagaraAsyncGpuTraceProviderHwrt.h b/Engine/Plugins/FX/Niagara/Source/NiagaraShader/Private/NiagaraAsyncGpuTraceProviderHwrt.h
--- a/Engine/Plugins/FX/Niagara/Source/NiagaraShader/Private/NiagaraAsyncGpuTraceProviderHwrt.h
+++ b/Engine/Plugins/FX/Niagara/Source/NiagaraShader/Private/NiagaraAsyncGpuTraceProviderHwrt.h
@@ -26,6 +26,7 @@ public:
 private:
 	FRayTracingPipelineState* RayTracingPipelineState = nullptr;
 	FRHIShaderBindingTable* RayTracingSBT = nullptr; // TODO: Hold reference?
+	FRHIShaderResourceView* TLASSRV = nullptr;
 	TUniformBufferRef<FViewUniformShaderParameters> ViewUniformBuffer;
 };
 
diff --git a/Engine/Source/Runtime/Renderer/Private/DeferredShadingRenderer.cpp b/Engine/Source/Runtime/Renderer/Private/DeferredShadingRenderer.cpp
--- a/Engine/Source/Runtime/Renderer/Private/DeferredShadingRenderer.cpp
+++ b/Engine/Source/Runtime/Renderer/Private/DeferredShadingRenderer.cpp
@@ -1040,11 +1040,7 @@ bool FDeferredShadingSceneRenderer::DispatchRayTracingWorldUpdates(FRDGBuilder&
 	{
 		RDG_EVENT_SCOPE_STAT(GraphBuilder, RayTracingScene, "RayTracingScene");
 		RDG_GPU_STAT_SCOPE(GraphBuilder, RayTracingScene);
-
-		const int32 ReferenceViewIndex = 0;
-		const FViewInfo& ReferenceView = Views[ReferenceViewIndex];
-
-		RayTracingScene.Create(GraphBuilder, GetSceneUniforms(), ReferenceView, &Scene->GPUScene, ComputePassFlags);
+		RayTracingScene.Update(GraphBuilder, GetSceneUniforms(), &Scene->GPUScene, ComputePassFlags);
 		RayTracingScene.Build(GraphBuilder, ComputePassFlags | ERDGPassFlags::NeverCull, OutDynamicGeometryScratchBuffer);
 	}
 
@@ -1771,16 +1767,22 @@ void FDeferredShadingSceneRenderer::Render(FRDGBuilder& GraphBuilder, const FSce
 		GRayTracingGeometryManager->ProcessBuildRequests(GraphBuilder.RHICmdList);
 	}
 
-	const int32 ReferenceViewIndex = 0;
-	FViewInfo& ReferenceView = Views[ReferenceViewIndex];
-
 	FRayTracingShaderBindingTable& RayTracingSBT = Scene->RayTracingSBT;
 	FRayTracingScene& RayTracingScene = Scene->RayTracingScene;
 	RayTracingSBT.ResetMissAndCallableShaders();
 
 	RayTracingScene.Reset();
-	RayTracingScene.InitPreViewTranslation(ReferenceView.ViewMatrices);
 
+	for (FViewInfo& View : Views)
+	{
+		if (IStereoRendering::IsStereoEyeView(View) && IStereoRendering::IsASecondaryView(View))
+		{
+			continue;
+		}
+
+		View.SetRayTracingSceneViewHandle(RayTracingScene.AddView(View.GetViewKey()));
+		RayTracingScene.SetViewParams(View.GetRayTracingSceneViewHandle(), View.ViewMatrices, View.RayTracingCullingParameters);
+	}
 #endif
 
 	FInitViewTaskDatas InitViewTaskDatas = OnRenderBegin(GraphBuilder, SceneUpdateInputs);
@@ -1964,6 +1966,9 @@ void FDeferredShadingSceneRenderer::Render(FRDGBuilder& GraphBuilder, const FSce
 				}
 			}
 
+			const int32 ReferenceViewIndex = 0;
+			FViewInfo& ReferenceView = Views[ReferenceViewIndex];
+
 			InitViewTaskDatas.RayTracingGatherInstances = RayTracing::CreateGatherInstancesTaskData(
 				Allocator,
 				*Scene,
diff --git a/Engine/Source/Runtime/Renderer/Private/FXRenderingUtils.cpp b/Engine/Source/Runtime/Renderer/Private/FXRenderingUtils.cpp
--- a/Engine/Source/Runtime/Renderer/Private/FXRenderingUtils.cpp
+++ b/Engine/Source/Runtime/Renderer/Private/FXRenderingUtils.cpp
@@ -241,9 +241,9 @@ TRDGUniformBufferRef<FSceneUniformParameters> UE::FXRenderingUtils::GetSceneUnif
 
 #if RHI_RAYTRACING
 
-bool UE::FXRenderingUtils::RayTracing::HasRayTracingScene(const FSceneInterface* InScene)
+bool UE::FXRenderingUtils::RayTracing::HasRayTracingScene(const FSceneInterface& InScene)
 {
-	if (const FScene* Scene = InScene->GetRenderScene())
+	if (const FScene* Scene = InScene.GetRenderScene())
 	{
 		return Scene->RayTracingScene.IsCreated();
 	}
@@ -251,21 +251,25 @@ bool UE::FXRenderingUtils::RayTracing::HasRayTracingScene(const FSceneInterface*
 	return false;
 }
 
-FRHIRayTracingScene* UE::FXRenderingUtils::RayTracing::GetRayTracingScene(const FSceneInterface* InScene)
+FRHIRayTracingScene* UE::FXRenderingUtils::RayTracing::GetRayTracingScene(const FSceneInterface& InScene, const FSceneView& InView)
 {
-	if (const FScene* Scene = InScene->GetRenderScene())
+	if (const FScene* Scene = InScene.GetRenderScene())
 	{
-		return Scene->RayTracingScene.GetRHIRayTracingSceneChecked(ERayTracingSceneLayer::Base);
+		check(InView.bIsViewInfo);
+		const FViewInfo& View = static_cast<const FViewInfo&>(InView);
+		return View.GetRayTracingSceneChecked(ERayTracingSceneLayer::Base);
 	}
 
 	return nullptr;
 }
 
-FRHIShaderResourceView* UE::FXRenderingUtils::RayTracing::GetRayTracingSceneView(FRHICommandListBase& RHICmdList, const FSceneInterface* InScene)
+FRHIShaderResourceView* UE::FXRenderingUtils::RayTracing::GetRayTracingSceneView(FRHICommandListBase& RHICmdList, const FSceneInterface& InScene, const FSceneView& InView)
 {
-	if (const FScene* Scene = InScene->GetRenderScene())
+	if (const FScene* Scene = InScene.GetRenderScene())
 	{
-		return Scene->RayTracingScene.CreateLayerViewRHI(RHICmdList, ERayTracingSceneLayer::Base);
+		check(InView.bIsViewInfo);
+		const FViewInfo& View = static_cast<const FViewInfo&>(InView);
+		return Scene->RayTracingScene.CreateLayerViewRHI(RHICmdList, ERayTracingSceneLayer::Base, View.GetRayTracingSceneViewHandle());
 	}
 
 	return nullptr;
diff --git a/Engine/Source/Runtime/Renderer/Private/HeterogeneousVolumes/HeterogeneousVolumes.h b/Engine/Source/Runtime/Renderer/Private/HeterogeneousVolumes/HeterogeneousVolumes.h
--- a/Engine/Source/Runtime/Renderer/Private/HeterogeneousVolumes/HeterogeneousVolumes.h
+++ b/Engine/Source/Runtime/Renderer/Private/HeterogeneousVolumes/HeterogeneousVolumes.h
@@ -878,7 +878,7 @@ void GenerateRayTracingScene(
 	FRDGBuilder& GraphBuilder,
 	// Scene data
 	const FScene* Scene,
-	const FViewInfo& View,
+	FViewInfo& View,
 	// Ray tracing data
 	TConstArrayView<FRayTracingGeometryRHIRef> RayTracingGeometries,
 	TConstArrayView<FMatrix> RayTracingTransforms,
diff --git a/Engine/Source/Runtime/Renderer/Private/HeterogeneousVolumes/HeterogeneousVolumesHardwareRayTracing.cpp b/Engine/Source/Runtime/Renderer/Private/HeterogeneousVolumes/HeterogeneousVolumesHardwareRayTracing.cpp
--- a/Engine/Source/Runtime/Renderer/Private/HeterogeneousVolumes/HeterogeneousVolumesHardwareRayTracing.cpp
+++ b/Engine/Source/Runtime/Renderer/Private/HeterogeneousVolumes/HeterogeneousVolumesHardwareRayTracing.cpp
@@ -188,7 +188,7 @@ void GenerateRayTracingScene(
 	FRDGBuilder& GraphBuilder,
 	// Scene data
 	const FScene* Scene,
-	const FViewInfo& View,
+	FViewInfo& View,
 	// Ray tracing data
 	TConstArrayView<FRayTracingGeometryRHIRef> RayTracingGeometries,
 	TConstArrayView<FMatrix> RayTracingTransforms,
@@ -197,6 +197,7 @@ void GenerateRayTracingScene(
 )
 {
 	RayTracingScene.Reset();
+	View.SetRayTracingSceneViewHandle(RayTracingScene.AddView(0)); // TODO: Could use View.GetViewKey() if view is persistent
 
 	TArray<FRayTracingGeometryBuildParams> BuildParams;
 	uint32 BLASScratchSize = 0;
@@ -223,7 +224,7 @@ void GenerateRayTracingScene(
 		RayTracingGeometryInstance.NumTransforms = 1;
 		RayTracingGeometryInstance.Transforms = MakeArrayView(&RayTracingTransforms[GeometryIndex], 1);
 
-		RayTracingScene.AddTransientInstance(RayTracingGeometryInstance, ERayTracingSceneLayer::Base);
+		RayTracingScene.AddTransientInstance(RayTracingGeometryInstance, ERayTracingSceneLayer::Base, View.GetRayTracingSceneViewHandle());
 	}
 
 	FRDGBufferDesc ScratchBufferDesc;
@@ -256,7 +257,7 @@ void GenerateRayTracingScene(
 
 	// Create RayTracingScene
 	const FGPUScene* EmptyGPUScene = nullptr;
-	RayTracingScene.Create(GraphBuilder, View.GetSceneUniforms(), View, EmptyGPUScene, ERDGPassFlags::Compute);
+	RayTracingScene.Update(GraphBuilder, View.GetSceneUniforms(), EmptyGPUScene, ERDGPassFlags::Compute);
 	RayTracingScene.Build(GraphBuilder, ERDGPassFlags::Compute | ERDGPassFlags::NeverCull | ERDGPassFlags::NeverParallel, nullptr);
 
 	GraphBuilder.AddDispatchHint();
@@ -601,7 +602,7 @@ void RenderLightingCacheWithPreshadingHardwareRayTracing(
 	FRenderLightingCacheWithPreshadingRGS::FParameters* PassParameters = GraphBuilder.AllocParameters<FRenderLightingCacheWithPreshadingRGS::FParameters>();
 	{
 		// Scene
-		PassParameters->TLAS = RayTracingScene.GetLayerView(ERayTracingSceneLayer::Base);
+		PassParameters->TLAS = RayTracingScene.GetLayerView(ERayTracingSceneLayer::Base, View.GetRayTracingSceneViewHandle());
 		PassParameters->ViewUniformBuffer = View.ViewUniformBuffer;
 		PassParameters->SceneTextures = GetSceneTextureParameters(GraphBuilder, SceneTextures);
 
@@ -780,7 +781,7 @@ void RenderSingleScatteringWithPreshadingHardwareRayTracing(
 	FRenderSingleScatteringWithPreshadingRGS::FParameters* PassParameters = GraphBuilder.AllocParameters<FRenderSingleScatteringWithPreshadingRGS::FParameters>();
 	{
 		// Scene
-		PassParameters->TLAS = RayTracingScene.GetLayerView(ERayTracingSceneLayer::Base);
+		PassParameters->TLAS = RayTracingScene.GetLayerView(ERayTracingSceneLayer::Base, View.GetRayTracingSceneViewHandle());
 		PassParameters->ViewUniformBuffer = View.ViewUniformBuffer;
 		PassParameters->SceneTextures = GetSceneTextureParameters(GraphBuilder, SceneTextures);
 
diff --git a/Engine/Source/Runtime/Renderer/Private/PathTracing.cpp b/Engine/Source/Runtime/Renderer/Private/PathTracing.cpp
--- a/Engine/Source/Runtime/Renderer/Private/PathTracing.cpp
+++ b/Engine/Source/Runtime/Renderer/Private/PathTracing.cpp
@@ -1005,7 +1005,7 @@ static void PreparePathTracingData(const FScene* Scene, const FViewInfo& View, F
 
 	PathTracingData.DecalRoughnessCutoff = PathTracing::UsesDecals(*View.Family) && View.bHasRayTracingDecals ? CVarPathTracingDecalRoughnessCutoff.GetValueOnRenderThread() : -1.0f;
 
-	PathTracingData.MeshDecalRoughnessCutoff = PathTracing::UsesDecals(*View.Family) && Scene->RayTracingScene.GetNumNativeInstances(ERayTracingSceneLayer::Decals) > 0 ? CVarPathTracingMeshDecalRoughnessCutoff.GetValueOnRenderThread() : -1.0f;
+	PathTracingData.MeshDecalRoughnessCutoff = PathTracing::UsesDecals(*View.Family) && Scene->RayTracingScene.GetNumNativeInstances(ERayTracingSceneLayer::Decals, View.GetRayTracingSceneViewHandle()) > 0 ? CVarPathTracingMeshDecalRoughnessCutoff.GetValueOnRenderThread() : -1.0f;
 	PathTracingData.MeshDecalBias = CVarPathTracingMeshDecalBias.GetValueOnRenderThread();
 
 	PathTracingData.MaxRaymarchSteps = CVarPathTracingMaxRaymarchSteps.GetValueOnRenderThread();
@@ -3577,7 +3577,7 @@ void FDeferredShadingSceneRenderer::RenderPathTracing(
 			{
 				PassParameters->RWSceneColor = GraphBuilder.CreateUAV(SceneColorOutputTexture);
 				PassParameters->RWSceneDepth = GraphBuilder.CreateUAV(DepthTexture);
-				PassParameters->TLAS = Scene->RayTracingScene.GetLayerView(ERayTracingSceneLayer::Base);
+				PassParameters->TLAS = Scene->RayTracingScene.GetLayerView(ERayTracingSceneLayer::Base, View.GetRayTracingSceneViewHandle());
 				PassParameters->UsePathGuiding = CVarPathTracingVisualizeIndirectGuiding.GetValueOnRenderThread();
 				PassParameters->MarkovChainStateGrid = GraphBuilder.CreateUAV(MarkovChainStateGrid);
 				PassParameters->MarkovChainStateGridSize = MarkovChainStateGridHashSize;
@@ -3787,7 +3787,7 @@ void FDeferredShadingSceneRenderer::RenderPathTracing(
 		StartingExtinctionCoefficient = GraphBuilder.CreateBuffer(FRDGBufferDesc::CreateBufferDesc(sizeof(float), 3), TEXT("PathTracer.StartingExtinctionCoefficient"));
 
 		FPathTracingInitExtinctionCoefficientRG::FParameters* PassParameters = GraphBuilder.AllocParameters<FPathTracingInitExtinctionCoefficientRG::FParameters>();
-		PassParameters->TLAS = Scene->RayTracingScene.GetLayerView(ERayTracingSceneLayer::Base);
+		PassParameters->TLAS = Scene->RayTracingScene.GetLayerView(ERayTracingSceneLayer::Base, View.GetRayTracingSceneViewHandle());
 		PassParameters->Scene = View.GetSceneUniforms().GetBuffer(GraphBuilder);
 		PassParameters->NaniteRayTracing = Nanite::GRayTracingManager.GetUniformBuffer();
 		PassParameters->ViewUniformBuffer = View.ViewUniformBuffer;
@@ -4169,8 +4169,8 @@ void FDeferredShadingSceneRenderer::RenderPathTracing(
 						for (int Bounce = 0; Bounce <= CompactionDepth; Bounce++)
 						{
 							FPathTracingRG::FParameters* PassParameters = GraphBuilder.AllocParameters<FPathTracingRG::FParameters>();
-							PassParameters->TLAS = Scene->RayTracingScene.GetLayerView(ERayTracingSceneLayer::Base);
-							PassParameters->DecalTLAS = Scene->RayTracingScene.GetLayerView(ERayTracingSceneLayer::Decals);
+							PassParameters->TLAS = Scene->RayTracingScene.GetLayerView(ERayTracingSceneLayer::Base, View.GetRayTracingSceneViewHandle());
+							PassParameters->DecalTLAS = Scene->RayTracingScene.GetLayerView(ERayTracingSceneLayer::Decals, View.GetRayTracingSceneViewHandle());
 							PassParameters->Scene = View.GetSceneUniforms().GetBuffer(GraphBuilder);
 							PassParameters->NaniteRayTracing = Nanite::GRayTracingManager.GetUniformBuffer();
 							PassParameters->ViewUniformBuffer = View.ViewUniformBuffer;
diff --git a/Engine/Source/Runtime/Renderer/Private/RayTracing/RayTracing.cpp b/Engine/Source/Runtime/Renderer/Private/RayTracing/RayTracing.cpp
--- a/Engine/Source/Runtime/Renderer/Private/RayTracing/RayTracing.cpp
+++ b/Engine/Source/Runtime/Renderer/Private/RayTracing/RayTracing.cpp
@@ -953,15 +953,29 @@ namespace RayTracing
 					{
 						RayTracingInstance.InstanceContributionToHitGroupIndex = SBTAllocation->GetInstanceContributionToHitGroupIndex(ERayTracingShaderBindingLayer::Base);
 
-						ERayTracingPrimitiveFlags Flags = Scene.PrimitiveRayTracingFlags[PrimitiveIndex];
-						RayTracingScene.AddTransientInstance(RayTracingInstance, EnumHasAllFlags(Flags, ERayTracingPrimitiveFlags::FarField) ? ERayTracingSceneLayer::FarField : ERayTracingSceneLayer::Base, SceneProxy, true, Geometry->GetGeometryHandle());
+						const ERayTracingPrimitiveFlags Flags = Scene.PrimitiveRayTracingFlags[PrimitiveIndex];
+						const ERayTracingSceneLayer Layer = EnumHasAllFlags(Flags, ERayTracingPrimitiveFlags::FarField) ? ERayTracingSceneLayer::FarField : ERayTracingSceneLayer::Base;
+
+						RayTracingScene.AddTransientInstance(
+							RayTracingInstance,
+							Layer,
+							View.GetRayTracingSceneViewHandle(),
+							SceneProxy,
+							/*bDynamic*/ true,
+							Geometry->GetGeometryHandle());
 					}
 
 					if (bNeedDecalInstance)
 					{
 						FRayTracingGeometryInstance DecalRayTracingInstance = RayTracingInstance;
 						DecalRayTracingInstance.InstanceContributionToHitGroupIndex = SBTAllocation->GetInstanceContributionToHitGroupIndex(ERayTracingShaderBindingLayer::Decals);
-						RayTracingScene.AddTransientInstance(MoveTemp(DecalRayTracingInstance), ERayTracingSceneLayer::Decals, SceneProxy, true, Geometry->GetGeometryHandle());
+						RayTracingScene.AddTransientInstance(
+							MoveTemp(DecalRayTracingInstance),
+							ERayTracingSceneLayer::Decals,
+							View.GetRayTracingSceneViewHandle(),
+							SceneProxy,
+							/*bDynamic*/ true,
+							Geometry->GetGeometryHandle());
 					}
 
 					if (bNeedMainInstance || bNeedDecalInstance)
@@ -1456,6 +1470,7 @@ namespace RayTracing
 	
 	void AddStaticInstancesToRayTracingScene(
 		const FScene& Scene,
+		const FViewInfo& View,
 		const RayTracing::FSceneOptions& SceneOptions,
 		TConstArrayView<FRelevantPrimitive> RelevantStaticPrimitives,
 		FRayTracingScene& RayTracingScene,
@@ -1598,7 +1613,15 @@ namespace RayTracing
 				if (bNeedMainInstance)
 				{
 					RayTracingInstance.InstanceContributionToHitGroupIndex = RelevantPrimitive.InstanceContributionToHitGroupIndexBase;
-					InstanceBatch.InstanceHandle = RayTracingScene.AddTransientInstance(RayTracingInstance, EnumHasAllFlags(Flags, ERayTracingPrimitiveFlags::FarField) ? ERayTracingSceneLayer::FarField : ERayTracingSceneLayer::Base, SceneProxy, false, RelevantPrimitive.RayTracingGeometry->GetGeometryHandle());
+
+					const ERayTracingSceneLayer Layer = EnumHasAllFlags(Flags, ERayTracingPrimitiveFlags::FarField) ? ERayTracingSceneLayer::FarField : ERayTracingSceneLayer::Base;
+					InstanceBatch.InstanceHandle = RayTracingScene.AddTransientInstance(
+						RayTracingInstance,
+						Layer,
+						View.GetRayTracingSceneViewHandle(),
+						SceneProxy,
+						/*bDynamic*/ false,
+						RelevantPrimitive.RayTracingGeometry->GetGeometryHandle());
 
 					VisibleShaderBindingData.Append(RelevantPrimitive.CachedShaderBindingDataBase);
 				}
@@ -1608,7 +1631,14 @@ namespace RayTracing
 				{
 					FRayTracingGeometryInstance DecalRayTracingInstance = RayTracingInstance;
 					DecalRayTracingInstance.InstanceContributionToHitGroupIndex = RelevantPrimitive.InstanceContributionToHitGroupIndexDecal;
-					InstanceBatch.DecalInstanceHandle = RayTracingScene.AddTransientInstance(MoveTemp(DecalRayTracingInstance), ERayTracingSceneLayer::Decals, SceneProxy, false, RelevantPrimitive.RayTracingGeometry->GetGeometryHandle());
+
+					InstanceBatch.DecalInstanceHandle = RayTracingScene.AddTransientInstance(
+						MoveTemp(DecalRayTracingInstance),
+						ERayTracingSceneLayer::Decals,
+						View.GetRayTracingSceneViewHandle(),
+						SceneProxy,
+						/*bDynamic*/ false,
+						RelevantPrimitive.RayTracingGeometry->GetGeometryHandle());
 
 					VisibleShaderBindingData.Append(RelevantPrimitive.CachedShaderBindingDataDecal);
 				}
@@ -1617,6 +1647,7 @@ namespace RayTracing
 	}
 
 	void AddVisibleCachedInstances(
+		const FViewInfo& View,
 		FRayTracingScene& RayTracingScene,
 		FRayTracingShaderBindingTable& RayTracingSBT,
 		TConstArrayView<FRelevantCachedPrimitive> RelevantCachedPrimitives,
@@ -1634,12 +1665,12 @@ namespace RayTracing
 
 			if (RelevantPrimitive.MainRayTracingInstanceIndex != UINT32_MAX)
 			{
-				RayTracingScene.MarkInstanceVisible(RelevantPrimitive.MainRayTracingInstanceIndex);
+				RayTracingScene.MarkInstanceVisible(RelevantPrimitive.MainRayTracingInstanceIndex, View.GetRayTracingSceneViewHandle());
 			}
 
 			if (RelevantPrimitive.DecalRayTracingInstanceIndex != UINT32_MAX && !bExcludeDecals)
 			{
-				RayTracingScene.MarkInstanceVisible(RelevantPrimitive.DecalRayTracingInstanceIndex);
+				RayTracingScene.MarkInstanceVisible(RelevantPrimitive.DecalRayTracingInstanceIndex, View.GetRayTracingSceneViewHandle());
 			}
 
 			FBitReference BitReference = ProcessedSBTAllocations[RelevantPrimitive.SBTAllocationUniqueId];
@@ -1785,12 +1816,14 @@ namespace RayTracing
 
 				AddStaticInstancesToRayTracingScene(
 					Scene,
+					View,
 					TaskData.SceneOptions,
 					TaskData.StaticPrimitives,
 					RayTracingScene,
 					View.VisibleRayTracingShaderBindings);
 
 				AddVisibleCachedInstances(
+					View,
 					RayTracingScene,
 					RayTracingSBT,
 					TaskData.CachedStaticPrimitives,
diff --git a/Engine/Source/Runtime/Renderer/Private/RayTracing/RayTracingAmbientOcclusion.cpp b/Engine/Source/Runtime/Renderer/Private/RayTracing/RayTracingAmbientOcclusion.cpp
--- a/Engine/Source/Runtime/Renderer/Private/RayTracing/RayTracingAmbientOcclusion.cpp
+++ b/Engine/Source/Runtime/Renderer/Private/RayTracing/RayTracingAmbientOcclusion.cpp
@@ -175,7 +175,7 @@ void FDeferredShadingSceneRenderer::RenderRayTracingAmbientOcclusion(
 	PassParameters->MaxRayDistance = View.FinalPostProcessSettings.RayTracingAORadius;
 	PassParameters->Intensity = View.FinalPostProcessSettings.RayTracingAOIntensity;
 	PassParameters->MaxNormalBias = GetRaytracingMaxNormalBias();
-	PassParameters->TLAS = RayTracingScene.GetLayerView(ERayTracingSceneLayer::Base);
+	PassParameters->TLAS = RayTracingScene.GetLayerView(ERayTracingSceneLayer::Base, View.GetRayTracingSceneViewHandle());
 	PassParameters->Substrate = Substrate::BindSubstrateGlobalUniformParameters(View);
 	PassParameters->RWAmbientOcclusionMaskUAV = GraphBuilder.CreateUAV(DenoiserInputs.Mask);
 	PassParameters->RWAmbientOcclusionHitDistanceUAV = GraphBuilder.CreateUAV(DenoiserInputs.RayHitDistance);
diff --git a/Engine/Source/Runtime/Renderer/Private/RayTracing/RayTracingDebug.cpp b/Engine/Source/Runtime/Renderer/Private/RayTracing/RayTracingDebug.cpp
--- a/Engine/Source/Runtime/Renderer/Private/RayTracing/RayTracingDebug.cpp
+++ b/Engine/Source/Runtime/Renderer/Private/RayTracing/RayTracingDebug.cpp
@@ -957,11 +957,11 @@ static FRDGBufferRef RayTracingPerformPicking(
 	FRDGBufferRef PickingBuffer = GraphBuilder.CreateBuffer(PickingBufferDesc, TEXT("RayTracingDebug.PickingBuffer"));
 
 	FRayTracingPickingParameters SharedParameters = {};
-	SharedParameters.InstancesExtraData = GraphBuilder.CreateSRV(RayTracingScene.GetInstanceExtraDataBuffer(ERayTracingSceneLayer::Base));
+	SharedParameters.InstancesExtraData = GraphBuilder.CreateSRV(RayTracingScene.GetInstanceExtraDataBuffer(ERayTracingSceneLayer::Base, View.GetRayTracingSceneViewHandle()));
 	SharedParameters.InstancesDebugData = GraphBuilder.CreateSRV(RayTracingScene.GetInstanceDebugBuffer(ERayTracingSceneLayer::Base));
-	SharedParameters.TLAS = RayTracingScene.GetLayerView(ERayTracingSceneLayer::Base);
+	SharedParameters.TLAS = RayTracingScene.GetLayerView(ERayTracingSceneLayer::Base, View.GetRayTracingSceneViewHandle());
 	SharedParameters.OpaqueOnly = CVarRayTracingVisualizeOpaqueOnly.GetValueOnRenderThread();
-	SharedParameters.InstanceBuffer = GraphBuilder.CreateSRV(RayTracingScene.GetInstanceBuffer(ERayTracingSceneLayer::Base));
+	SharedParameters.InstanceBuffer = GraphBuilder.CreateSRV(RayTracingScene.GetInstanceBuffer(ERayTracingSceneLayer::Base, View.GetRayTracingSceneViewHandle()));
 	SharedParameters.ViewUniformBuffer = View.ViewUniformBuffer;
 	SharedParameters.SceneUniformBuffer = GetSceneUniformBufferRef(GraphBuilder, View); // TODO: use a separate params structure
 	SharedParameters.NaniteRayTracingUniformBuffer = Nanite::GRayTracingManager.GetUniformBuffer();
@@ -1264,7 +1264,7 @@ static FRDGBufferRef RayTracingPerformHitStatsPerPrimitive(FRDGBuilder& GraphBui
 	FShaderBindingTableRHIRef HitStatsSBT = Scene.RayTracingSBT.AllocateTransientRHI(GraphBuilder.RHICmdList, ERayTracingShaderBindingMode::RTPSO, ERayTracingHitGroupIndexingMode::Allow, Initializer.GetMaxLocalBindingDataSize());
 
 	// TODO: Should check RayTracingScene for actual number of instances instead of max number in FRHIRayTracingScene initializer
-	const uint32 NumInstancesInTLAS = FMath::Max(RayTracingScene.GetRHIRayTracingSceneChecked(ERayTracingSceneLayer::Base)->GetInitializer().MaxNumInstances, (uint32)CVarRayTracingVisualizeHitCountTopKHits.GetValueOnRenderThread());
+	const uint32 NumInstancesInTLAS = FMath::Max(RayTracingScene.GetRHIRayTracingSceneChecked(ERayTracingSceneLayer::Base, View.GetRayTracingSceneViewHandle())->GetInitializer().MaxNumInstances, (uint32)CVarRayTracingVisualizeHitCountTopKHits.GetValueOnRenderThread());
 	FRDGBufferDesc HitStatsPerPrimitiveBufferDesc = FRDGBufferDesc::CreateStructuredDesc(sizeof(FRayTracingHitStatsEntry), NumInstancesInTLAS);
 	HitStatsPerPrimitiveBufferDesc.Usage = EBufferUsageFlags(HitStatsPerPrimitiveBufferDesc.Usage | BUF_SourceCopy);
 	FRDGBufferRef HitStatsBuffer = GraphBuilder.CreateBuffer(HitStatsPerPrimitiveBufferDesc, TEXT("RayTracingDebug.HitStatsBuffer"));
@@ -1274,7 +1274,7 @@ static FRDGBufferRef RayTracingPerformHitStatsPerPrimitive(FRDGBuilder& GraphBui
 	DebugHitStatsUniformBuffer = GraphBuilder.CreateUniformBuffer(DebugHitStatsUniformBufferParameters);
 
 	FRayTracingDebugHitStatsRGS::FParameters* RayGenParameters = GraphBuilder.AllocParameters<FRayTracingDebugHitStatsRGS::FParameters>();
-	RayGenParameters->TLAS = RayTracingScene.GetLayerView(ERayTracingSceneLayer::Base);
+	RayGenParameters->TLAS = RayTracingScene.GetLayerView(ERayTracingSceneLayer::Base, View.GetRayTracingSceneViewHandle());
 	RayGenParameters->OpaqueOnly = CVarRayTracingVisualizeOpaqueOnly.GetValueOnRenderThread();
 	RayGenParameters->ViewUniformBuffer = View.ViewUniformBuffer;
 	RayGenParameters->SceneUniformBuffer = GetSceneUniformBufferRef(GraphBuilder, View); // TODO: use a separate params structure
@@ -1466,8 +1466,8 @@ static void DrawInstanceOverlap(FRDGBuilder& GraphBuilder, const FScene& Scene,
 		View,
 		InstanceOverlapTexture,
 		SceneDepthTexture,
-		Scene.RayTracingScene.GetInstanceExtraDataBuffer(ERayTracingSceneLayer::Base),
-		Scene.RayTracingScene.GetNumNativeInstances(ERayTracingSceneLayer::Base),
+		Scene.RayTracingScene.GetInstanceExtraDataBuffer(ERayTracingSceneLayer::Base, View.GetRayTracingSceneViewHandle()),
+		Scene.RayTracingScene.GetNumNativeInstances(ERayTracingSceneLayer::Base, View.GetRayTracingSceneViewHandle()),
 		false);
 
 	// Calculate heatmap of instance overlap and blend it on top of ray tracing debug output
@@ -1501,8 +1501,8 @@ static void DrawInstanceOverlap(FRDGBuilder& GraphBuilder, const FScene& Scene,
 			View,
 			SceneColorTexture,
 			SceneDepthTexture,
-			Scene.RayTracingScene.GetInstanceExtraDataBuffer(ERayTracingSceneLayer::Base),
-			Scene.RayTracingScene.GetNumNativeInstances(ERayTracingSceneLayer::Base),
+			Scene.RayTracingScene.GetInstanceExtraDataBuffer(ERayTracingSceneLayer::Base, View.GetRayTracingSceneViewHandle()),
+			Scene.RayTracingScene.GetNumNativeInstances(ERayTracingSceneLayer::Base, View.GetRayTracingSceneViewHandle()),
 			true);
 	}
 }
@@ -1587,7 +1587,7 @@ void RenderRayTracingDebug(FRDGBuilder& GraphBuilder, const FScene& Scene, const
 
 		FRayTracingDebugTraversalCS::FParameters* PassParameters = GraphBuilder.AllocParameters<FRayTracingDebugTraversalCS::FParameters>();
 		PassParameters->Output = GraphBuilder.CreateUAV(SceneColorTexture);
-		PassParameters->TLAS = RayTracingScene.GetLayerView(ERayTracingSceneLayer::Base);
+		PassParameters->TLAS = RayTracingScene.GetLayerView(ERayTracingSceneLayer::Base, View.GetRayTracingSceneViewHandle());
 		PassParameters->ViewUniformBuffer = View.ViewUniformBuffer;
 		PassParameters->Scene = View.GetSceneUniforms().GetBuffer(GraphBuilder);
 		PassParameters->NaniteRasterUniformBuffer = CreateDebugNaniteRasterUniformBuffer(GraphBuilder);
@@ -1663,7 +1663,7 @@ void RenderRayTracingDebug(FRDGBuilder& GraphBuilder, const FScene& Scene, const
 	FRDGBufferRef PickingBuffer = nullptr;
 	FRDGBufferRef StatsBuffer = nullptr;
 	if (IsRayTracingPickingEnabled(DebugVisualizationMode)
-		&& RayTracingScene.GetInstanceExtraDataBuffer(ERayTracingSceneLayer::Base) != nullptr
+		&& RayTracingScene.GetInstanceExtraDataBuffer(ERayTracingSceneLayer::Base, View.GetRayTracingSceneViewHandle()) != nullptr
 		&& RayTracingScene.GetInstanceDebugBuffer(ERayTracingSceneLayer::Base) != nullptr)
 	{
 		PickingBuffer = RayTracingPerformPicking(GraphBuilder, Scene, View, VisibleRayTracingShaderBindings, PickingFeedback, bInlineRayTracing);
@@ -1688,7 +1688,7 @@ void RenderRayTracingDebug(FRDGBuilder& GraphBuilder, const FScene& Scene, const
 		AddClearUAVPass(GraphBuilder, DebugHitStatsUniformBufferParameters->HitStatsOutput, 0);
 	}
 
-	FRDGBufferRef InstanceExtraDataBuffer = RayTracingScene.GetInstanceExtraDataBuffer(ERayTracingSceneLayer::Base);
+	FRDGBufferRef InstanceExtraDataBuffer = RayTracingScene.GetInstanceExtraDataBuffer(ERayTracingSceneLayer::Base, View.GetRayTracingSceneViewHandle());
 	if (InstanceExtraDataBuffer == nullptr)
 	{
 		InstanceExtraDataBuffer = GSystemTextures.GetDefaultStructuredBuffer(GraphBuilder, sizeof(uint32));
@@ -1700,7 +1700,7 @@ void RenderRayTracingDebug(FRDGBuilder& GraphBuilder, const FScene& Scene, const
 		InstanceDebugBuffer = GSystemTextures.GetDefaultStructuredBuffer(GraphBuilder, sizeof(uint32));
 	}
 
-	FRDGBufferRef InstanceBuffer = RayTracingScene.GetInstanceBuffer(ERayTracingSceneLayer::Base);
+	FRDGBufferRef InstanceBuffer = RayTracingScene.GetInstanceBuffer(ERayTracingSceneLayer::Base, View.GetRayTracingSceneViewHandle());
 	if (InstanceBuffer == nullptr)
 	{
 		InstanceBuffer = GSystemTextures.GetDefaultStructuredBuffer(GraphBuilder, sizeof(uint32));
@@ -1708,7 +1708,7 @@ void RenderRayTracingDebug(FRDGBuilder& GraphBuilder, const FScene& Scene, const
 
 	const bool bRequiresDebugCHS = RequiresRayTracingDebugCHS(DebugVisualizationMode);
 
-	const uint32 NumInstances = RayTracingScene.GetNumNativeInstances(ERayTracingSceneLayer::Base);
+	const uint32 NumInstances = RayTracingScene.GetNumNativeInstances(ERayTracingSceneLayer::Base, View.GetRayTracingSceneViewHandle());
 
 	FRayTracingDebugParameters SharedParameters;
 	SharedParameters.VisualizationMode = DebugVisualizationMode;
@@ -1763,8 +1763,8 @@ void RenderRayTracingDebug(FRDGBuilder& GraphBuilder, const FScene& Scene, const
 	SharedParameters.InstancesDebugData = GraphBuilder.CreateSRV(InstanceDebugBuffer);
 	SharedParameters.InstanceBuffer = GraphBuilder.CreateSRV(InstanceBuffer);
 	SharedParameters.PickingBuffer = GraphBuilder.CreateSRV(PickingBuffer);
-	SharedParameters.TLAS = RayTracingScene.GetLayerView(ERayTracingSceneLayer::Base);
-	SharedParameters.FarFieldTLAS = RayTracingScene.GetLayerView(ERayTracingSceneLayer::FarField);
+	SharedParameters.TLAS = RayTracingScene.GetLayerView(ERayTracingSceneLayer::Base, View.GetRayTracingSceneViewHandle());
+	SharedParameters.FarFieldTLAS = RayTracingScene.GetLayerView(ERayTracingSceneLayer::FarField, View.GetRayTracingSceneViewHandle());
 	SharedParameters.ViewUniformBuffer = View.ViewUniformBuffer;
 	SharedParameters.Output = GraphBuilder.CreateUAV(SceneColorTexture);
 	SharedParameters.TopKHitStats = GraphBuilder.CreateSRV(StatsBuffer);
diff --git a/Engine/Source/Runtime/Renderer/Private/RayTracing/RayTracingInstanceBufferUtil.cpp b/Engine/Source/Runtime/Renderer/Private/RayTracing/RayTracingInstanceBufferUtil.cpp
--- a/Engine/Source/Runtime/Renderer/Private/RayTracing/RayTracingInstanceBufferUtil.cpp
+++ b/Engine/Source/Runtime/Renderer/Private/RayTracing/RayTracingInstanceBufferUtil.cpp
@@ -340,7 +340,7 @@ void FillRayTracingInstanceUploadBuffer(
 
 	const bool bRLEAllowed = CVarRayTracingInstanceBufferRLE.GetValueOnRenderThread();
 
-	const int32 NumSceneInstances = Instances.Num();
+	const int32 NumSceneInstances = VisibleInstances.Num();
 	const int32 MinBatchSize = 128;
 	ParallelFor(
 		TEXT("FillRayTracingInstanceUploadBuffer_Parallel"),
diff --git a/Engine/Source/Runtime/Renderer/Private/RayTracing/RayTracingScene.cpp b/Engine/Source/Runtime/Renderer/Private/RayTracing/RayTracingScene.cpp
--- a/Engine/Source/Runtime/Renderer/Private/RayTracing/RayTracingScene.cpp
+++ b/Engine/Source/Runtime/Renderer/Private/RayTracing/RayTracingScene.cpp
@@ -89,6 +89,7 @@ BEGIN_SHADER_PARAMETER_STRUCT(FBuildInstanceBufferPassParams, )
 END_SHADER_PARAMETER_STRUCT()
 
 const FRayTracingScene::FInstanceHandle FRayTracingScene::INVALID_INSTANCE_HANDLE = FInstanceHandle();
+const FRayTracingScene::FViewHandle FRayTracingScene::INVALID_VIEW_HANDLE = FViewHandle();
 
 using FInstanceBufferStats = uint32;
 
@@ -118,50 +119,122 @@ void FRayTracingScene::BuildInitializationData(bool bUseLightingChannels, bool b
 	{
 		FLayer& Layer = Layers[LayerIndex];
 
-		FRayTracingInstanceBufferBuilderInitializer Initializer;
-		Initializer.Instances = Layer.Instances;
-		Initializer.VisibleInstances = Layer.VisibleInstances;
-		Initializer.PreViewTranslation = PreViewTranslation;
-		Initializer.bUseLightingChannels = bUseLightingChannels;
-		Initializer.bForceOpaque = bForceOpaque;
-		Initializer.bDisableTriangleCull = bDisableTriangleCull;
+		for (int32 ViewIndex : ActiveViews)
+		{
+			FLayerView& LayerView = Layer.Views[ViewIndex];
+
+			FRayTracingInstanceBufferBuilderInitializer Initializer;
+			Initializer.Instances = Layer.Instances;
+			Initializer.VisibleInstances = LayerView.VisibleInstances;
+			Initializer.PreViewTranslation = ViewParameters[ViewIndex].PreViewTranslation;
+			Initializer.bUseLightingChannels = bUseLightingChannels;
+			Initializer.bForceOpaque = bForceOpaque;
+			Initializer.bDisableTriangleCull = bDisableTriangleCull;
 
-		Layer.InstanceBufferBuilder.Init(MoveTemp(Initializer));
+			LayerView.InstanceBufferBuilder.Init(MoveTemp(Initializer));
+		}
 	}
 
 	bInitializationDataBuilt = true;
 }
 
-void FRayTracingScene::InitPreViewTranslation(const FViewMatrices& ViewMatrices)
+FRayTracingScene::FViewHandle FRayTracingScene::AddView(uint32 ViewKey)
+{
+	if (ViewIndexMap.Contains(ViewKey))
+	{
+		const int32 ViewIndex = ViewIndexMap[ViewKey];
+		check(ActiveViews.IsValidIndex(ViewIndex) && ActiveViews[ViewIndex] == ViewIndex);
+		return FViewHandle(ViewIndex);
+	}
+
+	const int32 ViewIndex = ActiveViews.Add(0);
+	ActiveViews[ViewIndex] = ViewIndex;
+
+	if (ViewKey == 0)
+	{
+		// Transient View (eg: no ViewState) are removed at the end of the frame
+		TransientViewIndices.Add(ViewIndex);
+	}
+	else
+	{
+		ViewIndexMap.Add(ViewKey, ViewIndex);
+	}
+
+	if (ViewParameters.Num() < ViewIndex + 1)
+	{
+		ViewParameters.SetNum(ViewIndex + 1);
+	}
+
+	const uint8 NumLayers = uint8(ERayTracingSceneLayer::NUM);
+	for (uint32 LayerIndex = 0; LayerIndex < NumLayers; ++LayerIndex)
+	{
+		FLayer& Layer = Layers[LayerIndex];
+
+		if (Layer.Views.Num() < ViewIndex + 1)
+		{
+			Layer.Views.SetNum(ViewIndex + 1);
+		}
+	}
+
+	return FViewHandle(ViewIndex);
+}
+
+void FRayTracingScene::RemoveView(uint32 ViewKey)
+{
+	if (!ViewIndexMap.Contains(ViewKey))
+	{
+		return;
+	}
+
+	const int32 ViewIndex = ViewIndexMap[ViewKey];
+	check(ActiveViews.IsValidIndex(ViewIndex) && ActiveViews[ViewIndex] == ViewIndex);
+
+	// clear ViewIndex in Layers Views
+	const uint8 NumLayers = uint8(ERayTracingSceneLayer::NUM);
+	for (uint32 LayerIndex = 0; LayerIndex < NumLayers; ++LayerIndex)
+	{
+		FLayer& Layer = Layers[LayerIndex];
+		Layer.Views[ViewIndex] = {};
+	}
+
+	ViewParameters[ViewIndex] = {};
+
+	ActiveViews.RemoveAt(ViewIndex);
+	ViewIndexMap.Remove(ViewKey);
+}
+
+void FRayTracingScene::SetViewParams(FViewHandle ViewHandle, const FViewMatrices& ViewMatrices, const FRayTracingCullingParameters& CullingParameters)
 {
-	PreViewTranslation = ViewMatrices.GetPreViewTranslation();
+	ViewParameters[ViewHandle].CullingParameters = &CullingParameters;
+	ViewParameters[ViewHandle].PreViewTranslation = ViewMatrices.GetPreViewTranslation();
 }
 
-void FRayTracingScene::Create(FRDGBuilder& GraphBuilder, FSceneUniformBuffer& SceneUniformBuffer, const FViewInfo& View, const FGPUScene* GPUScene, ERDGPassFlags ComputePassFlags)
+void FRayTracingScene::Update(FRDGBuilder& GraphBuilder, FSceneUniformBuffer& SceneUniformBuffer, const FGPUScene* GPUScene, ERDGPassFlags ComputePassFlags)
 {
 	// Round up buffer sizes to some multiple to avoid pathological growth reallocations.
 	static constexpr uint32 AllocationGranularity = 8 * 1024;
 	static constexpr uint64 BufferAllocationGranularity = 16 * 1024 * 1024;
 
-	TRACE_CPUPROFILER_EVENT_SCOPE(FRayTracingScene::Create);
-	QUICK_SCOPE_CYCLE_COUNTER(STAT_RayTracingScene_Create);
+	TRACE_CPUPROFILER_EVENT_SCOPE(FRayTracingScene::Update);
+	QUICK_SCOPE_CYCLE_COUNTER(STAT_RayTracingScene_Update);
 
 	const ERayTracingAccelerationStructureFlags BuildFlags = CVarRayTracingSceneBuildMode.GetValueOnRenderThread()
 		? ERayTracingAccelerationStructureFlags::FastTrace
 		: ERayTracingAccelerationStructureFlags::FastBuild;
 
-	checkf(bInitializationDataBuilt, TEXT("BuildInitializationData(...) must be called before Create(...)."));
+	checkf(bInitializationDataBuilt, TEXT("BuildInitializationData(...) must be called before Update(...)."));
 
 	bUsedThisFrame = true;
 
 	FRHICommandListBase& RHICmdList = GraphBuilder.RHICmdList;
 
 	const uint8 NumLayers = uint8(ERayTracingSceneLayer::NUM);
+	const uint32 MaxNumViews = ActiveViews.GetMaxIndex() + 1;
 
 	FRDGBufferUAVRef InstanceStatsBufferUAV = nullptr;
 	{
 		// one counter per layer in the stats buffer
-		FRDGBufferDesc InstanceStatsBufferDesc = FRDGBufferDesc::CreateStructuredDesc(sizeof(FInstanceBufferStats), NumLayers);
+		FRDGBufferDesc InstanceStatsBufferDesc = FRDGBufferDesc::CreateStructuredDesc(sizeof(FInstanceBufferStats), NumLayers * MaxNumViews);
 		InstanceStatsBufferDesc.Usage |= BUF_SourceCopy;
 
 		InstanceStatsBuffer = GraphBuilder.CreateBuffer(InstanceStatsBufferDesc, TEXT("FRayTracingScene::InstanceStatsBuffer"));
@@ -176,179 +249,190 @@ void FRayTracingScene::Create(FRDGBuilder& GraphBuilder, FSceneUniformBuffer& Sc
 	{
 		FLayer& Layer = Layers[LayerIndex];
 
-		uint32 NumNativeInstances = Layer.InstanceBufferBuilder.GetMaxNumInstances();
-
-		if (bCompactInstanceBuffer)
+		for (int32 ViewIndex : ActiveViews)
 		{
-			NumNativeInstances = FMath::Max<uint32>(CVarRayTracingSceneCompactInstancesMin.GetValueOnRenderThread(), Layer.NumActiveInstances + CVarRayTracingSceneCompactInstancesMargin.GetValueOnRenderThread());
-			NumNativeInstances = FMath::Min<uint32>(NumNativeInstances, Layer.InstanceBufferBuilder.GetMaxNumInstances());
-		}
-
-		Layer.MaxNumInstances = NumNativeInstances;
+			FLayerView& LayerView = Layer.Views[ViewIndex];
 
-		const uint32 NumNativeInstancesAligned = FMath::DivideAndRoundUp(FMath::Max(NumNativeInstances, 1U), AllocationGranularity) * AllocationGranularity;
-
-		{
-			FRayTracingSceneInitializer Initializer;
-			Initializer.DebugName = Layer.Name;
-			Initializer.MaxNumInstances = NumNativeInstances;
-			Initializer.BuildFlags = BuildFlags;
-			Layer.RayTracingSceneRHI = RHICreateRayTracingScene(MoveTemp(Initializer));
-		}
+			uint32 NumNativeInstances = LayerView.InstanceBufferBuilder.GetMaxNumInstances();
 
-		FRayTracingAccelerationStructureSize SizeInfo = Layer.RayTracingSceneRHI->GetSizeInfo();
-		SizeInfo.ResultSize = FMath::DivideAndRoundUp(FMath::Max(SizeInfo.ResultSize, 1ull), BufferAllocationGranularity) * BufferAllocationGranularity;
+			if (bCompactInstanceBuffer)
+			{
+				NumNativeInstances = FMath::Max<uint32>(CVarRayTracingSceneCompactInstancesMin.GetValueOnRenderThread(), LayerView.NumActiveInstances + CVarRayTracingSceneCompactInstancesMargin.GetValueOnRenderThread());
+				NumNativeInstances = FMath::Min<uint32>(NumNativeInstances, LayerView.InstanceBufferBuilder.GetMaxNumInstances());
+			}
 
-		// Allocate GPU buffer if current one is too small or significantly larger than what we need.
-		if (!Layer.RayTracingScenePooledBuffer.IsValid()
-			|| SizeInfo.ResultSize > Layer.RayTracingScenePooledBuffer->GetSize()
-			|| SizeInfo.ResultSize < Layer.RayTracingScenePooledBuffer->GetSize() / 2)
-		{
-			FRDGBufferDesc Desc = FRDGBufferDesc::CreateBufferDesc(1, uint32(SizeInfo.ResultSize));
-			Desc.Usage = EBufferUsageFlags::AccelerationStructure;
+			LayerView.MaxNumInstances = NumNativeInstances;
 
-			Layer.RayTracingScenePooledBuffer = AllocatePooledBuffer(Desc, TEXT("FRayTracingScene::SceneBuffer"));
-		}
+			const uint32 NumNativeInstancesAligned = FMath::DivideAndRoundUp(FMath::Max(NumNativeInstances, 1U), AllocationGranularity) * AllocationGranularity;
 
-		Layer.RayTracingSceneBufferRDG = GraphBuilder.RegisterExternalBuffer(Layer.RayTracingScenePooledBuffer);
-		Layer.RayTracingSceneBufferSRV = GraphBuilder.CreateSRV(FRDGBufferSRVDesc(Layer.RayTracingSceneBufferRDG, Layer.RayTracingSceneRHI, 0));
+			{
+				FRayTracingSceneInitializer Initializer;
+				Initializer.DebugName = Layer.Name; // TODO: also include ViewIndex in the name
+				Initializer.MaxNumInstances = NumNativeInstances;
+				Initializer.BuildFlags = BuildFlags;
+				LayerView.RayTracingSceneRHI = RHICreateRayTracingScene(MoveTemp(Initializer));
+			}
 
-		{
-			const uint64 ScratchAlignment = GRHIRayTracingScratchBufferAlignment;
-			FRDGBufferDesc ScratchBufferDesc;
-			ScratchBufferDesc.Usage = EBufferUsageFlags::RayTracingScratch | EBufferUsageFlags::StructuredBuffer;
-			ScratchBufferDesc.BytesPerElement = uint32(ScratchAlignment);
-			ScratchBufferDesc.NumElements = uint32(FMath::DivideAndRoundUp(SizeInfo.BuildScratchSize, ScratchAlignment));
+			FRayTracingAccelerationStructureSize SizeInfo = LayerView.RayTracingSceneRHI->GetSizeInfo();
+			SizeInfo.ResultSize = FMath::DivideAndRoundUp(FMath::Max(SizeInfo.ResultSize, 1ull), BufferAllocationGranularity) * BufferAllocationGranularity;
 
-			Layer.BuildScratchBuffer = GraphBuilder.CreateBuffer(ScratchBufferDesc, TEXT("FRayTracingScene::ScratchBuffer"));
-		}
+			// Allocate GPU buffer if current one is too small or significantly larger than what we need.
+			if (!LayerView.RayTracingScenePooledBuffer.IsValid()
+				|| SizeInfo.ResultSize > LayerView.RayTracingScenePooledBuffer->GetSize()
+				|| SizeInfo.ResultSize < LayerView.RayTracingScenePooledBuffer->GetSize() / 2)
+			{
+				FRDGBufferDesc Desc = FRDGBufferDesc::CreateBufferDesc(1, uint32(SizeInfo.ResultSize));
+				Desc.Usage = EBufferUsageFlags::AccelerationStructure;
 
-		{
-			FRDGBufferDesc InstanceBufferDesc;
-			InstanceBufferDesc.Usage = EBufferUsageFlags::UnorderedAccess | EBufferUsageFlags::ShaderResource | EBufferUsageFlags::StructuredBuffer;
-			InstanceBufferDesc.BytesPerElement = GRHIRayTracingInstanceDescriptorSize;
-			InstanceBufferDesc.NumElements = NumNativeInstancesAligned;
+				LayerView.RayTracingScenePooledBuffer = AllocatePooledBuffer(Desc, TEXT("FRayTracingScene::SceneBuffer"));
+			}
 
-			Layer.InstanceBuffer = GraphBuilder.CreateBuffer(InstanceBufferDesc, TEXT("FRayTracingScene::InstanceBuffer"));
+			LayerView.RayTracingSceneBufferRDG = GraphBuilder.RegisterExternalBuffer(LayerView.RayTracingScenePooledBuffer);
+			LayerView.RayTracingSceneBufferSRV = GraphBuilder.CreateSRV(FRDGBufferSRVDesc(LayerView.RayTracingSceneBufferRDG, LayerView.RayTracingSceneRHI, 0));
 
-			if (bCompactInstanceBuffer)
 			{
-				// need to clear since FRayTracingBuildInstanceBufferCS will only write active instances
-				AddClearUAVPass(GraphBuilder, GraphBuilder.CreateUAV(Layer.InstanceBuffer), 0, ComputePassFlags);
-			}
-		}
+				const uint64 ScratchAlignment = GRHIRayTracingScratchBufferAlignment;
+				FRDGBufferDesc ScratchBufferDesc;
+				ScratchBufferDesc.Usage = EBufferUsageFlags::RayTracingScratch | EBufferUsageFlags::StructuredBuffer;
+				ScratchBufferDesc.BytesPerElement = uint32(ScratchAlignment);
+				ScratchBufferDesc.NumElements = uint32(FMath::DivideAndRoundUp(SizeInfo.BuildScratchSize, ScratchAlignment));
 
-		// Feedback
-		if (bTracingFeedbackEnabled)
-		{
-			{
-				FRDGBufferDesc InstanceHitCountBufferDesc;
-				InstanceHitCountBufferDesc.Usage = EBufferUsageFlags::UnorderedAccess | EBufferUsageFlags::ShaderResource | EBufferUsageFlags::StructuredBuffer;
-				InstanceHitCountBufferDesc.BytesPerElement = sizeof(uint32);
-				InstanceHitCountBufferDesc.NumElements = NumNativeInstancesAligned;
-
-				Layer.InstanceHitCountBuffer = GraphBuilder.CreateBuffer(InstanceHitCountBufferDesc, TEXT("FRayTracingScene::InstanceHitCount"));
-				Layer.InstanceHitCountBufferUAV = GraphBuilder.CreateUAV(Layer.InstanceHitCountBuffer, ERDGUnorderedAccessViewFlags::SkipBarrier);
-				AddClearUAVPass(GraphBuilder, Layer.InstanceHitCountBufferUAV, 0, ComputePassFlags);
+				LayerView.BuildScratchBuffer = GraphBuilder.CreateBuffer(ScratchBufferDesc, TEXT("FRayTracingScene::ScratchBuffer"));
 			}
 
 			{
-				FRDGBufferDesc AccelerationStructureIndexBufferDesc;
-				AccelerationStructureIndexBufferDesc.Usage = EBufferUsageFlags::ShaderResource | EBufferUsageFlags::StructuredBuffer;
-				AccelerationStructureIndexBufferDesc.BytesPerElement = sizeof(uint32);
-				AccelerationStructureIndexBufferDesc.NumElements = NumNativeInstancesAligned;
+				FRDGBufferDesc InstanceBufferDesc;
+				InstanceBufferDesc.Usage = EBufferUsageFlags::UnorderedAccess | EBufferUsageFlags::ShaderResource | EBufferUsageFlags::StructuredBuffer;
+				InstanceBufferDesc.BytesPerElement = GRHIRayTracingInstanceDescriptorSize;
+				InstanceBufferDesc.NumElements = NumNativeInstancesAligned;
 
-				Layer.AccelerationStructureIndexBuffer = GraphBuilder.CreateBuffer(AccelerationStructureIndexBufferDesc, TEXT("FRayTracingScene::AccelerationStructureIndexBuffer"));
+				LayerView.InstanceBuffer = GraphBuilder.CreateBuffer(InstanceBufferDesc, TEXT("FRayTracingScene::InstanceBuffer"));
 
-				TConstArrayView<uint32> InstanceGeometryIndices = Layer.InstanceBufferBuilder.GetInstanceGeometryIndices();
-				GraphBuilder.QueueBufferUpload(Layer.AccelerationStructureIndexBuffer, InstanceGeometryIndices.GetData(), InstanceGeometryIndices.GetTypeSize() * InstanceGeometryIndices.Num());
+				if (bCompactInstanceBuffer)
+				{
+					// need to clear since FRayTracingBuildInstanceBufferCS will only write active instances
+					AddClearUAVPass(GraphBuilder, GraphBuilder.CreateUAV(LayerView.InstanceBuffer), 0, ComputePassFlags);
+				}
 			}
 
+			// Feedback
+			if (bTracingFeedbackEnabled)
 			{
-				FRDGBufferDesc GeometryHandleBufferDesc;
-				GeometryHandleBufferDesc.Usage = EBufferUsageFlags::ShaderResource | EBufferUsageFlags::StructuredBuffer;
-				GeometryHandleBufferDesc.BytesPerElement = sizeof(int32);
-				GeometryHandleBufferDesc.NumElements = FMath::Max(Layer.GeometryHandles.Num(), 1);
+				{
+					FRDGBufferDesc InstanceHitCountBufferDesc;
+					InstanceHitCountBufferDesc.Usage = EBufferUsageFlags::UnorderedAccess | EBufferUsageFlags::ShaderResource | EBufferUsageFlags::StructuredBuffer;
+					InstanceHitCountBufferDesc.BytesPerElement = sizeof(uint32);
+					InstanceHitCountBufferDesc.NumElements = NumNativeInstancesAligned;
+
+					LayerView.InstanceHitCountBuffer = GraphBuilder.CreateBuffer(InstanceHitCountBufferDesc, TEXT("FRayTracingScene::InstanceHitCount"));
+					LayerView.InstanceHitCountBufferUAV = GraphBuilder.CreateUAV(LayerView.InstanceHitCountBuffer, ERDGUnorderedAccessViewFlags::SkipBarrier);
+					AddClearUAVPass(GraphBuilder, LayerView.InstanceHitCountBufferUAV, 0, ComputePassFlags);
+				}
 
-				Layer.GeometryHandleBuffer = GraphBuilder.CreateBuffer(GeometryHandleBufferDesc, TEXT("FRayTracingScene::GeometryHandleBuffer"));
-				GraphBuilder.QueueBufferUpload(Layer.GeometryHandleBuffer, Layer.GeometryHandles.GetData(), Layer.GeometryHandles.GetTypeSize() * Layer.GeometryHandles.Num());
-			}
-		}
+				{
+					FRDGBufferDesc AccelerationStructureIndexBufferDesc;
+					AccelerationStructureIndexBufferDesc.Usage = EBufferUsageFlags::ShaderResource | EBufferUsageFlags::StructuredBuffer;
+					AccelerationStructureIndexBufferDesc.BytesPerElement = sizeof(uint32);
+					AccelerationStructureIndexBufferDesc.NumElements = NumNativeInstancesAligned;
 
-		FRDGBufferUAVRef InstanceExtraDataBufferUAV = nullptr;
-		if (bInstanceExtraDataBufferEnabled || bTracingFeedbackEnabled || bInstanceDebugDataEnabled)
-		{
-			FRDGBufferDesc InstanceExtraDataBufferDesc;
-			InstanceExtraDataBufferDesc.Usage = EBufferUsageFlags::UnorderedAccess | EBufferUsageFlags::ShaderResource | EBufferUsageFlags::StructuredBuffer;
-			InstanceExtraDataBufferDesc.BytesPerElement = sizeof(FRayTracingInstanceExtraData);
-			InstanceExtraDataBufferDesc.NumElements = FMath::Max(NumNativeInstances, 1u);
+					LayerView.AccelerationStructureIndexBuffer = GraphBuilder.CreateBuffer(AccelerationStructureIndexBufferDesc, TEXT("FRayTracingScene::AccelerationStructureIndexBuffer"));
 
-			Layer.InstanceExtraDataBuffer = GraphBuilder.CreateBuffer(InstanceExtraDataBufferDesc, TEXT("FRayTracingScene::InstanceExtraDataBuffer"));
-			InstanceExtraDataBufferUAV = GraphBuilder.CreateUAV(Layer.InstanceExtraDataBuffer);
+					TConstArrayView<uint32> InstanceGeometryIndices = LayerView.InstanceBufferBuilder.GetInstanceGeometryIndices();
+					GraphBuilder.QueueBufferUpload(LayerView.AccelerationStructureIndexBuffer, InstanceGeometryIndices.GetData(), InstanceGeometryIndices.GetTypeSize() * InstanceGeometryIndices.Num());
+				}
+			}
 
-			AddClearUAVPass(GraphBuilder, InstanceExtraDataBufferUAV, 0xFFFFFFFF, ComputePassFlags);
-		}
+			FRDGBufferUAVRef InstanceExtraDataBufferUAV = nullptr;
+			if (bInstanceExtraDataBufferEnabled || bTracingFeedbackEnabled || bInstanceDebugDataEnabled)
+			{
+				FRDGBufferDesc InstanceExtraDataBufferDesc;
+				InstanceExtraDataBufferDesc.Usage = EBufferUsageFlags::UnorderedAccess | EBufferUsageFlags::ShaderResource | EBufferUsageFlags::StructuredBuffer;
+				InstanceExtraDataBufferDesc.BytesPerElement = sizeof(FRayTracingInstanceExtraData);
+				InstanceExtraDataBufferDesc.NumElements = FMath::Max(NumNativeInstances, 1u);
 
-		if (Layer.InstancesDebugData.Num() > 0)
-		{
-			check(Layer.InstancesDebugData.Num() == Layer.Instances.Num());
+				LayerView.InstanceExtraDataBuffer = GraphBuilder.CreateBuffer(InstanceExtraDataBufferDesc, TEXT("FRayTracingScene::InstanceExtraDataBuffer"));
+				InstanceExtraDataBufferUAV = GraphBuilder.CreateUAV(LayerView.InstanceExtraDataBuffer);
 
-			Layer.InstanceDebugBuffer = CreateStructuredBuffer(GraphBuilder, TEXT("FRayTracingScene::InstanceDebugData"), Layer.InstancesDebugData);
-		}
+				AddClearUAVPass(GraphBuilder, InstanceExtraDataBufferUAV, 0xFFFFFFFF, ComputePassFlags);
+			}
 
-		if (NumNativeInstances > 0)
-		{
-			// Fill instance upload buffer on separate thread since results are only needed in RHI thread
-			GraphBuilder.AddCommandListSetupTask([&InstanceBufferBuilder = Layer.InstanceBufferBuilder](FRHICommandList& RHICmdList)
-				{
-					FTaskTagScope TaskTagScope(ETaskTag::EParallelRenderingThread);
+			if (NumNativeInstances > 0)
+			{
+				// Fill instance upload buffer on separate thread since results are only needed in RHI thread
+				GraphBuilder.AddCommandListSetupTask([&InstanceBufferBuilder = LayerView.InstanceBufferBuilder](FRHICommandList& RHICmdList)
+					{
+						FTaskTagScope TaskTagScope(ETaskTag::EParallelRenderingThread);
 
-					InstanceBufferBuilder.FillRayTracingInstanceUploadBuffer(RHICmdList);
-				});
+						InstanceBufferBuilder.FillRayTracingInstanceUploadBuffer(RHICmdList);
+					});
 
-			GraphBuilder.AddCommandListSetupTask([&InstanceBufferBuilder = Layer.InstanceBufferBuilder](FRHICommandList& RHICmdList)
-				{
-					FTaskTagScope TaskTagScope(ETaskTag::EParallelRenderingThread);
+				GraphBuilder.AddCommandListSetupTask([&InstanceBufferBuilder = LayerView.InstanceBufferBuilder](FRHICommandList& RHICmdList)
+					{
+						FTaskTagScope TaskTagScope(ETaskTag::EParallelRenderingThread);
 
-					InstanceBufferBuilder.FillAccelerationStructureAddressesBuffer(RHICmdList);
-				});
+						InstanceBufferBuilder.FillAccelerationStructureAddressesBuffer(RHICmdList);
+					});
 
-#if STATS
-			const bool bStatsEnabled = true;
-#else
-			const bool bStatsEnabled = false;
-#endif
+	#if STATS
+				const bool bStatsEnabled = true;
+	#else
+				const bool bStatsEnabled = false;
+	#endif
 
-			FBuildInstanceBufferPassParams* PassParams = GraphBuilder.AllocParameters<FBuildInstanceBufferPassParams>();
-			PassParams->InstanceBuffer = GraphBuilder.CreateUAV(Layer.InstanceBuffer);
-			PassParams->InstanceExtraDataBuffer = InstanceExtraDataBufferUAV;
-			PassParams->Scene = SceneUniformBuffer.GetBuffer(GraphBuilder);
-			PassParams->OutputStats = bCompactInstanceBuffer || bStatsEnabled ? InstanceStatsBufferUAV : nullptr;
-
-			GraphBuilder.AddPass(
-				RDG_EVENT_NAME("RayTracingBuildInstanceBuffer"),
-				PassParams,
-				ComputePassFlags,
-				[PassParams,
-				&InstanceBufferBuilder = Layer.InstanceBufferBuilder,
-				LayerIndex,
-				GPUScene,
-				CullingParameters = &View.RayTracingCullingParameters,
-				NumNativeInstances,
-				bCompactInstanceBuffer
-				](FRHICommandList& RHICmdList)
 				{
-					InstanceBufferBuilder.BuildRayTracingInstanceBuffer(
-						RHICmdList,
+					FBuildInstanceBufferPassParams* PassParams = GraphBuilder.AllocParameters<FBuildInstanceBufferPassParams>();
+					PassParams->InstanceBuffer = GraphBuilder.CreateUAV(LayerView.InstanceBuffer);
+					PassParams->InstanceExtraDataBuffer = InstanceExtraDataBufferUAV;
+					PassParams->Scene = SceneUniformBuffer.GetBuffer(GraphBuilder);
+					PassParams->OutputStats = bCompactInstanceBuffer || bStatsEnabled ? InstanceStatsBufferUAV : nullptr;
+
+					const uint32 OutputStatsOffset = LayerIndex * MaxNumViews + ViewIndex;
+
+					GraphBuilder.AddPass(
+						RDG_EVENT_NAME("RayTracingBuildInstanceBuffer"),
+						PassParams,
+						ComputePassFlags,
+						[PassParams,
+						&InstanceBufferBuilder = LayerView.InstanceBufferBuilder,
+						OutputStatsOffset,
 						GPUScene,
-						CullingParameters,
-						PassParams->InstanceBuffer->GetRHI(),
+						CullingParameters = ViewParameters[ViewIndex].CullingParameters,
 						NumNativeInstances,
-						bCompactInstanceBuffer,
-						PassParams->OutputStats ? PassParams->OutputStats->GetRHI() : nullptr,
-						LayerIndex,
-						PassParams->InstanceExtraDataBuffer ? PassParams->InstanceExtraDataBuffer->GetRHI() : nullptr);
-				});
+						bCompactInstanceBuffer
+						](FRHICommandList& RHICmdList)
+						{
+							InstanceBufferBuilder.BuildRayTracingInstanceBuffer(
+								RHICmdList,
+								GPUScene,
+								CullingParameters,
+								PassParams->InstanceBuffer->GetRHI(),
+								NumNativeInstances,
+								bCompactInstanceBuffer,
+								PassParams->OutputStats ? PassParams->OutputStats->GetRHI() : nullptr,
+								OutputStatsOffset,
+								PassParams->InstanceExtraDataBuffer ? PassParams->InstanceExtraDataBuffer->GetRHI() : nullptr);
+						});
+				}
+			}
+
+			// Feedback
+			if (bTracingFeedbackEnabled)
+			{
+				FRDGBufferDesc GeometryHandleBufferDesc;
+				GeometryHandleBufferDesc.Usage = EBufferUsageFlags::ShaderResource | EBufferUsageFlags::StructuredBuffer;
+				GeometryHandleBufferDesc.BytesPerElement = sizeof(int32);
+				GeometryHandleBufferDesc.NumElements = FMath::Max(Layer.GeometryHandles.Num(), 1);
+
+				Layer.GeometryHandleBuffer = GraphBuilder.CreateBuffer(GeometryHandleBufferDesc, TEXT("FRayTracingScene::GeometryHandleBuffer"));
+				GraphBuilder.QueueBufferUpload(Layer.GeometryHandleBuffer, Layer.GeometryHandles.GetData(), Layer.GeometryHandles.GetTypeSize() * Layer.GeometryHandles.Num());
+			}
+
+			if (Layer.InstancesDebugData.Num() > 0)
+			{
+				check(Layer.InstancesDebugData.Num() == Layer.Instances.Num());
+
+				Layer.InstanceDebugBuffer = CreateStructuredBuffer(GraphBuilder, TEXT("FRayTracingScene::InstanceDebugData"), Layer.InstancesDebugData);
+			}
 		}
 	}
 }
@@ -373,9 +457,14 @@ void FRayTracingScene::Build(FRDGBuilder& GraphBuilder, ERDGPassFlags ComputePas
 	{
 		FLayer& Layer = Layers[LayerIndex];
 
-		PassParams->TLASBuildBuffers.Emplace(Layer.BuildScratchBuffer, ERHIAccess::UAVCompute);
-		PassParams->TLASBuildBuffers.Emplace(Layer.InstanceBuffer, ERHIAccess::SRVCompute);
-		PassParams->TLASBuildBuffers.Emplace(Layer.RayTracingSceneBufferRDG, ERHIAccess::BVHWrite);
+		for (int32 ViewIndex : ActiveViews)
+		{
+			FLayerView& LayerView = Layer.Views[ViewIndex];
+
+			PassParams->TLASBuildBuffers.Emplace(LayerView.BuildScratchBuffer, ERHIAccess::UAVCompute);
+			PassParams->TLASBuildBuffers.Emplace(LayerView.InstanceBuffer, ERHIAccess::SRVCompute);
+			PassParams->TLASBuildBuffers.Emplace(LayerView.RayTracingSceneBufferRDG, ERHIAccess::BVHWrite);
+		}
 	}
 
 	GraphBuilder.AddPass(RDG_EVENT_NAME("RayTracingBuildScene"), PassParams, ComputePassFlags,
@@ -384,30 +473,35 @@ void FRayTracingScene::Build(FRDGBuilder& GraphBuilder, ERDGPassFlags ComputePas
 			const bool bUseBatchedBuild = CVarRayTracingSceneBatchedBuild.GetValueOnRenderThread();
 
 			TArray<FRayTracingSceneBuildParams> BatchedBuildParams;
-			BatchedBuildParams.Reserve(NumLayers);
+			BatchedBuildParams.Reserve(NumLayers); // TODO: should also take num views into account
 
 			for (uint32 LayerIndex = 0; LayerIndex < NumLayers; ++LayerIndex)
 			{
 				FLayer& Layer = Layers[LayerIndex];
 
-				FRayTracingSceneBuildParams BuildParams;
-				BuildParams.Scene = Layer.RayTracingSceneRHI;
-				BuildParams.ScratchBuffer = Layer.BuildScratchBuffer->GetRHI();
-				BuildParams.ScratchBufferOffset = 0;
-				BuildParams.InstanceBuffer = Layer.InstanceBuffer->GetRHI();
-				BuildParams.InstanceBufferOffset = 0;
-				BuildParams.NumInstances = Layer.MaxNumInstances;
-				BuildParams.ReferencedGeometries = Layer.InstanceBufferBuilder.GetReferencedGeometries();
+				for (int32 ViewIndex : ActiveViews)
+				{
+					FLayerView& LayerView = Layer.Views[ViewIndex];
 
-				RHICmdList.BindAccelerationStructureMemory(Layer.RayTracingSceneRHI, Layer.RayTracingSceneBufferRDG->GetRHI(), 0);
+					FRayTracingSceneBuildParams BuildParams;
+					BuildParams.Scene = LayerView.RayTracingSceneRHI;
+					BuildParams.ScratchBuffer = LayerView.BuildScratchBuffer->GetRHI();
+					BuildParams.ScratchBufferOffset = 0;
+					BuildParams.InstanceBuffer = LayerView.InstanceBuffer->GetRHI();
+					BuildParams.InstanceBufferOffset = 0;
+					BuildParams.NumInstances = LayerView.MaxNumInstances;
+					BuildParams.ReferencedGeometries = LayerView.InstanceBufferBuilder.GetReferencedGeometries();
 
-				if (bUseBatchedBuild)
-				{
-					BatchedBuildParams.Add(BuildParams);
-				}
-				else
-				{
-					RHICmdList.BuildAccelerationStructure(BuildParams);
+					RHICmdList.BindAccelerationStructureMemory(LayerView.RayTracingSceneRHI, LayerView.RayTracingSceneBufferRDG->GetRHI(), 0);
+
+					if (bUseBatchedBuild)
+					{
+						BatchedBuildParams.Add(BuildParams);
+					}
+					else
+					{
+						RHICmdList.BuildAccelerationStructure(BuildParams);
+					}
 				}
 			}
 
@@ -418,24 +512,29 @@ void FRayTracingScene::Build(FRDGBuilder& GraphBuilder, ERDGPassFlags ComputePas
 		});
 
 #if !UE_BUILD_SHIPPING
-	for (uint32 LayerIndex = 0; LayerIndex < NumLayers; ++LayerIndex)
+	if (GRayTracingSerializeSceneNextFrame && GRHIGlobals.RayTracing.SupportsSerializeAccelerationStructure)
 	{
-		FLayer& Layer = Layers[LayerIndex];
-
-		if (GRayTracingSerializeSceneNextFrame && GRHIGlobals.RayTracing.SupportsSerializeAccelerationStructure)
+		for (uint32 LayerIndex = 0; LayerIndex < NumLayers; ++LayerIndex)
 		{
-			FRayTracingSceneSerializePassParams* SerializePassParams = GraphBuilder.AllocParameters<FRayTracingSceneSerializePassParams>();
-			SerializePassParams->TLASBuffer = Layer.RayTracingSceneBufferRDG;
+			FLayer& Layer = Layers[LayerIndex];
 
-			GraphBuilder.AddPass(RDG_EVENT_NAME("RayTracingSerializeScene"), SerializePassParams, ERDGPassFlags::Readback,
-				[SerializePassParams, &Layer](FRHICommandListImmediate& RHICmdList)
-				{
-					FString Filename = FString::Printf(TEXT("%s_(%s)"), *Layer.Name.ToString(), *FDateTime::Now().ToString(TEXT("%Y%m%d_%H%M%S")));
-					FString RootPath = FPaths::ScreenShotDir() + TEXT("BVH/"); // Save BVH dumps to ScreenShot directory
-					FString OutputFilename = RootPath + Filename + TEXT(".bvh");
+			for (int32 ViewIndex : ActiveViews)
+			{
+				FLayerView& LayerView = Layer.Views[ViewIndex];
+
+				FRayTracingSceneSerializePassParams* SerializePassParams = GraphBuilder.AllocParameters<FRayTracingSceneSerializePassParams>();
+				SerializePassParams->TLASBuffer = LayerView.RayTracingSceneBufferRDG;
 
-					RHICmdList.SerializeAccelerationStructure(Layer.RayTracingSceneRHI, *OutputFilename);
-				});
+				GraphBuilder.AddPass(RDG_EVENT_NAME("RayTracingSerializeScene"), SerializePassParams, ERDGPassFlags::Readback,
+					[SerializePassParams, &Layer, &LayerView, ViewIndex](FRHICommandListImmediate& RHICmdList)
+					{
+						FString Filename = FString::Printf(TEXT("%s_%d_(%s)"), *Layer.Name.ToString(), ViewIndex, *FDateTime::Now().ToString(TEXT("%Y%m%d_%H%M%S")));
+						FString RootPath = FPaths::ScreenShotDir() + TEXT("BVH/"); // Save BVH dumps to ScreenShot directory
+						FString OutputFilename = RootPath + Filename + TEXT(".bvh");
+
+						RHICmdList.SerializeAccelerationStructure(LayerView.RayTracingSceneRHI, *OutputFilename);
+					});
+			}
 		}
 	}
 
@@ -519,12 +618,18 @@ END_SHADER_PARAMETER_STRUCT()
 void FRayTracingScene::FinishTracingFeedback(FRDGBuilder& GraphBuilder, ERDGPassFlags ComputePassFlags)
 {
 	RDG_EVENT_SCOPE(GraphBuilder, "RayTracingScene::FinishTracingFeedback");
+
+	if (!bTracingFeedbackEnabled)
+	{
+		return;
+	}
 	
 	const FLayer& Layer = Layers[0];
-	const uint32 NumGeometries = (uint32)Layer.InstanceBufferBuilder.GetReferencedGeometries().Num();
-	const uint32 NumInstances = Layer.InstanceBufferBuilder.GetMaxNumInstances();
+	const FLayerView& LayerView = Layer.Views[0];
+	const uint32 NumGeometries = (uint32)LayerView.InstanceBufferBuilder.GetReferencedGeometries().Num();
+	const uint32 NumInstances = LayerView.InstanceBufferBuilder.GetMaxNumInstances();
 
-	if (NumGeometries == 0 || !bTracingFeedbackEnabled)
+	if (NumGeometries == 0)
 	{
 		return;
 	}
@@ -552,11 +657,11 @@ void FRayTracingScene::FinishTracingFeedback(FRDGBuilder& GraphBuilder, ERDGPass
 		FRayTracingUpdateGeometryHitCountCS::FParameters* PassParameters = GraphBuilder.AllocParameters<FRayTracingUpdateGeometryHitCountCS::FParameters>();
 
 		PassParameters->GeometryHandleBuffer = GraphBuilder.CreateSRV(Layer.GeometryHandleBuffer);
-		PassParameters->AccelerationStructureIndexBuffer = GraphBuilder.CreateSRV(Layer.AccelerationStructureIndexBuffer);
-		PassParameters->InstanceHitCountBuffer = GraphBuilder.CreateSRV(Layer.InstanceHitCountBuffer);
+		PassParameters->AccelerationStructureIndexBuffer = GraphBuilder.CreateSRV(LayerView.AccelerationStructureIndexBuffer);
+		PassParameters->InstanceHitCountBuffer = GraphBuilder.CreateSRV(LayerView.InstanceHitCountBuffer);
 		PassParameters->RWGeometryHitCountBuffer = GraphBuilder.CreateUAV(GeometryHitCountBuffer);
 		PassParameters->RWGeometryHandleBuffer = GraphBuilder.CreateUAV(GeometryHandleBuffer);
-		PassParameters->InstanceExtraDataBuffer = GraphBuilder.CreateSRV(Layer.InstanceExtraDataBuffer);
+		PassParameters->InstanceExtraDataBuffer = GraphBuilder.CreateSRV(LayerView.InstanceExtraDataBuffer);
 		PassParameters->NumInstances = NumInstances;
 
 		const FIntVector GroupSize = FComputeShaderUtils::GetGroupCountWrapped(NumInstances, FRayTracingUpdateGeometryHitCountCS::ThreadGroupSize);
@@ -693,6 +798,7 @@ void FRayTracingScene::FinishStats(FRDGBuilder& GraphBuilder, ERDGPassFlags Comp
 	uint32 TotalNumActiveInstances = 0;
 
 	const uint8 NumLayers = uint8(ERayTracingSceneLayer::NUM);
+	const uint32 MaxNumViews = ActiveViews.GetMaxIndex() + 1;
 
 	// process ready results
 	while (StatsReadbackBuffersNumPending > 0)
@@ -703,18 +809,23 @@ void FRayTracingScene::FinishStats(FRDGBuilder& GraphBuilder, ERDGPassFlags Comp
 		{
 			StatsReadbackBuffersNumPending--;
 
-			auto ReadbackBufferPtr = (const FInstanceBufferStats*)ReadbackBuffer->Lock(sizeof(FInstanceBufferStats) * NumLayers);
+			auto ReadbackBufferPtr = (const FInstanceBufferStats*)ReadbackBuffer->Lock(sizeof(FInstanceBufferStats) * NumLayers * MaxNumViews);
 
 			for (uint32 LayerIndex = 0; LayerIndex < NumLayers; ++LayerIndex)
 			{
 				FLayer& Layer = Layers[LayerIndex];
 
-				const uint32 LayerNumNativeInstances = Layer.InstanceBufferBuilder.GetMaxNumInstances();
+				for (int32 ViewIndex : ActiveViews)
+				{
+					FLayerView& LayerView = Layer.Views[ViewIndex];
+
+					const uint32 LayerViewNumNativeInstances = LayerView.InstanceBufferBuilder.GetMaxNumInstances();
 
-				Layer.NumActiveInstances = FMath::Min(ReadbackBufferPtr[LayerIndex], LayerNumNativeInstances);
+					LayerView.NumActiveInstances = FMath::Min(ReadbackBufferPtr[LayerIndex * MaxNumViews + ViewIndex], LayerViewNumNativeInstances);
 
-				TotalNumNativeInstances += LayerNumNativeInstances;
-				TotalNumActiveInstances += Layer.NumActiveInstances;
+					TotalNumNativeInstances += LayerViewNumNativeInstances;
+					TotalNumActiveInstances += LayerView.NumActiveInstances;
+				}
 			}
 
 			ReadbackBuffer->Unlock();
@@ -740,41 +851,42 @@ bool FRayTracingScene::IsCreated() const
 	return bUsedThisFrame;
 }
 
-FRHIRayTracingScene* FRayTracingScene::GetRHIRayTracingScene(ERayTracingSceneLayer Layer) const
+FRHIRayTracingScene* FRayTracingScene::GetRHIRayTracingScene(ERayTracingSceneLayer InLayer, FViewHandle InViewHandle) const
 {
-	return Layers[uint8(Layer)].RayTracingSceneRHI.GetReference();
+	return Layers[uint8(InLayer)].Views[InViewHandle].RayTracingSceneRHI.GetReference();
 }
 
-FRHIRayTracingScene* FRayTracingScene::GetRHIRayTracingSceneChecked(ERayTracingSceneLayer Layer) const
+FRHIRayTracingScene* FRayTracingScene::GetRHIRayTracingSceneChecked(ERayTracingSceneLayer InLayer, FViewHandle InViewHandle) const
 {
-	FRHIRayTracingScene* Result = GetRHIRayTracingScene(Layer);
-	checkf(Result, TEXT("Ray tracing scene was not created. Perhaps Create() was not called."));
+	FRHIRayTracingScene* Result = GetRHIRayTracingScene(InLayer, InViewHandle);
+	checkf(Result, TEXT("Ray tracing scene was not created. Perhaps Update() was not called."));
 	return Result;
 }
 
-FShaderResourceViewRHIRef FRayTracingScene::CreateLayerViewRHI(FRHICommandListBase& RHICmdList, ERayTracingSceneLayer InLayer) const
+FShaderResourceViewRHIRef FRayTracingScene::CreateLayerViewRHI(FRHICommandListBase& RHICmdList, ERayTracingSceneLayer InLayer, FViewHandle InViewHandle) const
 {
-	const FLayer& Layer = Layers[uint8(InLayer)];
-	checkf(Layer.RayTracingScenePooledBuffer, TEXT("Ray tracing scene was not created. Perhaps Create() was not called."));
-	return RHICmdList.CreateShaderResourceView(FShaderResourceViewInitializer(Layer.RayTracingScenePooledBuffer->GetRHI(), Layer.RayTracingSceneRHI, 0));
+	const FLayerView& LayerView = Layers[uint8(InLayer)].Views[InViewHandle];
+	checkf(LayerView.RayTracingScenePooledBuffer, TEXT("Ray tracing scene was not created. Perhaps Update() was not called."));
+	return RHICmdList.CreateShaderResourceView(FShaderResourceViewInitializer(LayerView.RayTracingScenePooledBuffer->GetRHI(), LayerView.RayTracingSceneRHI, 0));
 }
 
-FRDGBufferSRVRef FRayTracingScene::GetLayerView(ERayTracingSceneLayer Layer) const
+FRDGBufferSRVRef FRayTracingScene::GetLayerView(ERayTracingSceneLayer InLayer, FViewHandle InViewHandle) const
 {
-	checkf(Layers[uint8(Layer)].RayTracingSceneBufferSRV, TEXT("Ray tracing scene SRV was not created. Perhaps Create() was not called."));
-	return Layers[uint8(Layer)].RayTracingSceneBufferSRV;
+	const FLayerView& LayerView = Layers[uint8(InLayer)].Views[InViewHandle];
+	checkf(LayerView.RayTracingSceneBufferSRV, TEXT("Ray tracing scene SRV was not created. Perhaps Update() was not called."));
+	return LayerView.RayTracingSceneBufferSRV;
 }
 
-FRDGBufferUAVRef FRayTracingScene::GetInstanceHitCountBufferUAV(ERayTracingSceneLayer Layer) const
+FRDGBufferUAVRef FRayTracingScene::GetInstanceHitCountBufferUAV(ERayTracingSceneLayer InLayer, FViewHandle InViewHandle) const
 {
-	return bTracingFeedbackEnabled ? Layers[uint8(Layer)].InstanceHitCountBufferUAV : nullptr;
+	return bTracingFeedbackEnabled ? Layers[uint8(InLayer)].Views[InViewHandle].InstanceHitCountBufferUAV : nullptr;
 }
 
-uint32 FRayTracingScene::GetNumNativeInstances(ERayTracingSceneLayer InLayer) const
+uint32 FRayTracingScene::GetNumNativeInstances(ERayTracingSceneLayer InLayer, FViewHandle InViewHandle) const
 {
-	const FLayer& Layer = Layers[uint8(InLayer)];
-	checkf(bInitializationDataBuilt, TEXT("Must call BuildInitializationData() or Create() before using GetNumNativeInstances()."));
-	return Layer.InstanceBufferBuilder.GetMaxNumInstances();
+	const FLayerView& LayerView = Layers[uint8(InLayer)].Views[InViewHandle];
+	checkf(bInitializationDataBuilt, TEXT("Must call BuildInitializationData() or Update() before using GetNumNativeInstances()."));
+	return LayerView.InstanceBufferBuilder.GetMaxNumInstances();
 }
 
 FRayTracingScene::FInstanceHandle FRayTracingScene::AddCachedInstance(FRayTracingGeometryInstance Instance, ERayTracingSceneLayer InLayer, const FPrimitiveSceneProxy* Proxy, bool bDynamic, int32 InGeometryHandle)
@@ -790,7 +902,6 @@ FRayTracingScene::FInstanceHandle FRayTracingScene::AddCachedInstance(FRayTracin
 	if (Layer.CachedInstancesFreeList.IsEmpty())
 	{
 		InstanceIndex = Layer.Instances.Add(MoveTemp(Instance));
-		Layer.VisibleInstances.Add(false);
 	}
 	else
 	{
@@ -875,7 +986,7 @@ FRHIRayTracingGeometry* FRayTracingScene::GetCachedInstanceGeometry(uint32 Packe
 	return GetCachedInstanceGeometry(FInstanceHandle(PackedHandle));
 }
 
-FRayTracingScene::FInstanceHandle FRayTracingScene::AddTransientInstance(FRayTracingGeometryInstance Instance, ERayTracingSceneLayer InLayer, const FPrimitiveSceneProxy* Proxy, bool bDynamic, int32 GeometryHandle)
+FRayTracingScene::FInstanceHandle FRayTracingScene::AddTransientInstance(FRayTracingGeometryInstance Instance, ERayTracingSceneLayer InLayer, FViewHandle InViewHandle, const FPrimitiveSceneProxy* Proxy, bool bDynamic, int32 GeometryHandle)
 {
 	ensure(bCachedInstancesLocked);
 
@@ -901,23 +1012,39 @@ FRayTracingScene::FInstanceHandle FRayTracingScene::AddTransientInstance(FRayTra
 		check(Layer.Instances.Num() == Layer.InstancesDebugData.Num());
 	}
 
-	Layer.VisibleInstances.Add(true);
+	FLayerView& LayerView = Layer.Views[InViewHandle];
+
+	int32 MinNumInstances = InstanceIndex + 1;
+	if (LayerView.VisibleInstances.Num() < MinNumInstances)
+	{
+		LayerView.VisibleInstances.SetNum(MinNumInstances, false);
+	}
+
+	LayerView.VisibleInstances[InstanceIndex] = true;
 
 	return { InLayer, InstanceIndex };
 }
 
-void FRayTracingScene::MarkInstanceVisible(FInstanceHandle Handle)
+void FRayTracingScene::MarkInstanceVisible(FInstanceHandle Handle, FViewHandle InViewHandle)
 {
 	FLayer& Layer = Layers[uint8(Handle.Layer)];
 
 	check(Layer.Instances[Handle.Index].GeometryRHI != nullptr);
 
-	Layer.VisibleInstances[Handle.Index] = true;
+	FLayerView& LayerView = Layer.Views[InViewHandle];
+
+	int32 MinNumInstances = Handle.Index + 1;
+	if (LayerView.VisibleInstances.Num() < MinNumInstances)
+	{
+		LayerView.VisibleInstances.SetNum(MinNumInstances, false);
+	}
+
+	LayerView.VisibleInstances[Handle.Index] = true;
 }
 
-void FRayTracingScene::MarkInstanceVisible(uint32 Handle)
+void FRayTracingScene::MarkInstanceVisible(uint32 Handle, FViewHandle InViewHandle)
 {
-	MarkInstanceVisible(FInstanceHandle(ERayTracingSceneLayer(Handle >> 24), Handle & 0xFFFFFF));
+	MarkInstanceVisible(FInstanceHandle(ERayTracingSceneLayer(Handle >> 24), Handle & 0xFFFFFF), InViewHandle);
 }
 
 uint32 FRayTracingScene::FLayer::GetCachedInstanceSectionSize()
@@ -929,6 +1056,22 @@ void FRayTracingScene::Reset()
 {
 	const uint8 NumLayers = uint8(ERayTracingSceneLayer::NUM);
 
+	for (int32 ViewIndex : TransientViewIndices)
+	{
+		check(ActiveViews.IsValidIndex(ViewIndex) && ActiveViews[ViewIndex] == ViewIndex);
+
+		// clear ViewIndex in Layers Views
+		for (uint32 LayerIndex = 0; LayerIndex < NumLayers; ++LayerIndex)
+		{
+			FLayer& Layer = Layers[LayerIndex];
+			Layer.Views[ViewIndex] = {};
+		}
+
+		ViewParameters[ViewIndex] = {};
+	}
+
+	TransientViewIndices.Empty();
+
 	for (uint32 LayerIndex = 0; LayerIndex < NumLayers; ++LayerIndex)
 	{
 		FLayer& Layer = Layers[LayerIndex];
@@ -938,25 +1081,32 @@ void FRayTracingScene::Reset()
 		Layer.Instances.SetNum(CachedInstanceSectionSize);
 		Layer.InstancesDebugData.SetNum(bInstanceDebugDataEnabled ? CachedInstanceSectionSize : 0);
 
-		Layer.VisibleInstances.Reset();
-		Layer.VisibleInstances.SetNum(CachedInstanceSectionSize, false);
+		Layer.GeometryHandleBuffer = nullptr;
+		Layer.GeometryHandles.SetNum(bTracingFeedbackEnabled ? CachedInstanceSectionSize : 0);
 
-		Layer.InstanceBufferBuilder = {};
+		Layer.InstanceDebugBuffer = nullptr;
 
-		Layer.RayTracingSceneRHI = nullptr;
-		Layer.RayTracingSceneBufferRDG = nullptr;
-		Layer.RayTracingSceneBufferSRV = nullptr;
+		for (int32 ViewIndex = 0; ViewIndex < Layer.Views.Num(); ++ViewIndex)
+		{
+			FLayerView& LayerView = Layer.Views[ViewIndex];
 
-		Layer.InstanceBuffer = nullptr;
-		Layer.BuildScratchBuffer = nullptr;
-		Layer.InstanceExtraDataBuffer = nullptr;
-		Layer.InstanceDebugBuffer = nullptr;
+			LayerView.VisibleInstances.Reset();
+			LayerView.VisibleInstances.SetNum(CachedInstanceSectionSize, false);
 
-		Layer.InstanceHitCountBuffer = nullptr;
-		Layer.InstanceHitCountBufferUAV = nullptr;
-		Layer.AccelerationStructureIndexBuffer = nullptr;
-		Layer.GeometryHandleBuffer = nullptr;
-		Layer.GeometryHandles.SetNum(bTracingFeedbackEnabled ? CachedInstanceSectionSize : 0);
+			LayerView.InstanceBufferBuilder = {};
+
+			LayerView.RayTracingSceneRHI = nullptr;
+			LayerView.RayTracingSceneBufferRDG = nullptr;
+			LayerView.RayTracingSceneBufferSRV = nullptr;
+
+			LayerView.InstanceBuffer = nullptr;
+			LayerView.BuildScratchBuffer = nullptr;
+			LayerView.InstanceExtraDataBuffer = nullptr;
+
+			LayerView.InstanceHitCountBuffer = nullptr;
+			LayerView.InstanceHitCountBufferUAV = nullptr;
+			LayerView.AccelerationStructureIndexBuffer = nullptr;
+		}
 	}
 
 	InstanceStatsBuffer = nullptr;
@@ -967,6 +1117,9 @@ void FRayTracingScene::Reset()
 
 	bUsesLightingChannels = false;
 
+	check(InitTask.IsCompleted());
+	InitTask = {};
+
 	bCachedInstancesLocked = false;
 }
 
@@ -983,7 +1136,11 @@ void FRayTracingScene::EndFrame()
 		{
 			FLayer& Layer = Layers[LayerIndex];
 
-			Layer.RayTracingScenePooledBuffer = nullptr;
+			for (int32 ViewIndex = 0; ViewIndex < Layer.Views.Num(); ++ViewIndex)
+			{
+				FLayerView& LayerView = Layer.Views[ViewIndex];
+				LayerView.RayTracingScenePooledBuffer = nullptr;
+			}
 		}
 
 		GeometriesToBuild.Empty();
diff --git a/Engine/Source/Runtime/Renderer/Private/RayTracing/RayTracingScene.h b/Engine/Source/Runtime/Renderer/Private/RayTracing/RayTracingScene.h
--- a/Engine/Source/Runtime/Renderer/Private/RayTracing/RayTracingScene.h
+++ b/Engine/Source/Runtime/Renderer/Private/RayTracing/RayTracingScene.h
@@ -88,6 +88,34 @@ public:
 
 	static const FInstanceHandle INVALID_INSTANCE_HANDLE;
 
+	struct FViewHandle
+	{
+		FViewHandle()
+			: Index(UINT32_MAX)
+		{}
+
+		bool IsValid() const
+		{
+			return Index != UINT32_MAX;
+		}
+
+	private:
+		explicit FViewHandle(uint32 InIndex)
+			: Index(InIndex)
+		{}
+
+		operator uint32()
+		{
+			return Index;
+		}
+
+		uint32 Index;
+
+		friend class FRayTracingScene;
+	};
+
+	static const FViewHandle INVALID_VIEW_HANDLE;
+
 	FRayTracingScene();
 	~FRayTracingScene();
 
@@ -101,18 +129,21 @@ public:
 	FRHIRayTracingGeometry* GetCachedInstanceGeometry(FInstanceHandle Handle) const;
 	FRHIRayTracingGeometry* GetCachedInstanceGeometry(uint32 PackedHandle) const;
 
-	FInstanceHandle AddTransientInstance(FRayTracingGeometryInstance Instance, ERayTracingSceneLayer Layer, const FPrimitiveSceneProxy* Proxy = nullptr, bool bDynamic = false, int32 GeometryHandle = INDEX_NONE);
+	FInstanceHandle AddTransientInstance(FRayTracingGeometryInstance Instance, ERayTracingSceneLayer Layer, FViewHandle ViewHandle, const FPrimitiveSceneProxy* Proxy = nullptr, bool bDynamic = false, int32 GeometryHandle = INDEX_NONE);
 
-	void MarkInstanceVisible(FInstanceHandle Handle);
-	void MarkInstanceVisible(uint32 PackedHandle);
+	void MarkInstanceVisible(FInstanceHandle Handle, FViewHandle ViewHandle);
+	void MarkInstanceVisible(uint32 PackedHandle, FViewHandle ViewHandle);
 
 	// Builds various metadata required to create the final scene.
 	// Must be done before calling Create(...).
 	void BuildInitializationData(bool bUseLightingChannels, bool bForceOpaque, bool bDisableTriangleCull);
 
+	FViewHandle AddView(uint32 ViewKey);
+	void RemoveView(uint32 ViewKey);
+
 	// Allocates GPU memory to fit at least the current number of instances.
 	// Kicks off instance buffer build to parallel thread along with RDG pass.
-	void Create(FRDGBuilder& GraphBuilder, FSceneUniformBuffer& SceneUniformBuffer, const FViewInfo& View, const FGPUScene* GPUScene, ERDGPassFlags ComputePassFlags);
+	void Update(FRDGBuilder& GraphBuilder, FSceneUniformBuffer& SceneUniformBuffer, const FGPUScene* GPUScene, ERDGPassFlags ComputePassFlags);
 
 	void Build(FRDGBuilder& GraphBuilder, ERDGPassFlags ComputePassFlags, FRDGBufferRef DynamicGeometryScratchBuffer);
 
@@ -149,32 +180,35 @@ public:
 	RENDERER_API bool IsCreated() const;
 
 	// Returns RayTracingSceneRHI object (may return null).
-	RENDERER_API  FRHIRayTracingScene* GetRHIRayTracingScene(ERayTracingSceneLayer Layer) const;
+	RENDERER_API  FRHIRayTracingScene* GetRHIRayTracingScene(ERayTracingSceneLayer Layer, FViewHandle ViewHandle) const;
 
 	// Similar to GetRayTracingScene, but checks that ray tracing scene RHI object is valid.
-	RENDERER_API  FRHIRayTracingScene* GetRHIRayTracingSceneChecked(ERayTracingSceneLayer Layer) const;
+	RENDERER_API  FRHIRayTracingScene* GetRHIRayTracingSceneChecked(ERayTracingSceneLayer Layer, FViewHandle ViewHandle) const;
 
 	// Creates new RHI view of a layer. Can only be used on valid ray tracing scene. 
-	RENDERER_API FShaderResourceViewRHIRef CreateLayerViewRHI(FRHICommandListBase& RHICmdList, ERayTracingSceneLayer Layer) const;
+	RENDERER_API FShaderResourceViewRHIRef CreateLayerViewRHI(FRHICommandListBase& RHICmdList, ERayTracingSceneLayer Layer, FViewHandle ViewHandle) const;
 
 	// Returns RDG view of a layer. Can only be used on valid ray tracing scene.
-	RENDERER_API FRDGBufferSRVRef GetLayerView(ERayTracingSceneLayer Layer) const;
+	RENDERER_API FRDGBufferSRVRef GetLayerView(ERayTracingSceneLayer Layer, FViewHandle ViewHandle) const;
 
 	// Feedback	
-	RENDERER_API FRDGBufferUAVRef GetInstanceHitCountBufferUAV(ERayTracingSceneLayer Layer) const;
+	RENDERER_API FRDGBufferUAVRef GetInstanceHitCountBufferUAV(ERayTracingSceneLayer Layer, FViewHandle ViewHandle) const;
 
-	FRDGBufferRef GetInstanceBuffer(ERayTracingSceneLayer Layer) const { return Layers[uint8(Layer)].InstanceBuffer; }
+	FRDGBufferRef GetInstanceBuffer(ERayTracingSceneLayer Layer, FViewHandle ViewHandle) const { return Layers[uint8(Layer)].Views[ViewHandle].InstanceBuffer; }
 
 	TConstArrayView<FRayTracingGeometryInstance> GetInstances(ERayTracingSceneLayer Layer) const { return MakeArrayView(Layers[uint8(Layer)].Instances); }
 
 	FRayTracingGeometryInstance& GetInstance(FInstanceHandle Handle) { return Layers[uint8(Handle.Layer)].Instances[Handle.Index]; }
 
-	uint32 GetNumNativeInstances(ERayTracingSceneLayer Layer) const;
+	uint32 GetNumNativeInstances(ERayTracingSceneLayer Layer, FViewHandle ViewHandle) const;
 
 	FRDGBufferRef GetInstanceDebugBuffer(ERayTracingSceneLayer Layer) const { return Layers[uint8(Layer)].InstanceDebugBuffer; }
-	FRDGBufferRef GetInstanceExtraDataBuffer(ERayTracingSceneLayer Layer) const { return Layers[uint8(Layer)].InstanceExtraDataBuffer; }
 
-	void InitPreViewTranslation(const FViewMatrices& ViewMatrices);
+	FRDGBufferRef GetInstanceExtraDataBuffer(ERayTracingSceneLayer Layer, FViewHandle ViewHandle) const { return Layers[uint8(Layer)].Views[ViewHandle].InstanceExtraDataBuffer; }
+
+	void SetViewParams(FViewHandle ViewHandle, const FViewMatrices& ViewMatrices, const FRayTracingCullingParameters& CullingParameters);
+
+	FVector GetPreViewTranslation(FViewHandle ViewHandle) const { return ViewParameters[ViewHandle].PreViewTranslation; }
 
 public:
 
@@ -185,9 +219,6 @@ public:
 
 	bool bUsesLightingChannels = false;
 
-	// Used for transforming to translated world space in which TLAS was built.
-	FVector PreViewTranslation {};
-
 	UE::Tasks::FTask InitTask; // Task to asynchronously call BuildInitializationData()
 
 private:
@@ -198,10 +229,8 @@ private:
 	void ReleaseReadbackBuffers();
 	void ReleaseFeedbackReadbackBuffers();
 
-	struct FLayer
+	struct FLayerView
 	{
-		uint32 GetCachedInstanceSectionSize();
-
 		FRayTracingInstanceBufferBuilder InstanceBufferBuilder;
 
 		FRayTracingSceneRHIRef RayTracingSceneRHI;
@@ -213,8 +242,6 @@ private:
 		FRDGBufferRef InstanceHitCountBuffer;
 		FRDGBufferUAVRef InstanceHitCountBufferUAV;
 		FRDGBufferRef AccelerationStructureIndexBuffer;
-		FRDGBufferRef GeometryHandleBuffer;
-		TArray<int32> GeometryHandles;
 
 		TRefCountPtr<FRDGPooledBuffer> RayTracingScenePooledBuffer;
 		FRDGBufferRef RayTracingSceneBufferRDG;
@@ -222,6 +249,20 @@ private:
 
 		FRDGBufferRef InstanceExtraDataBuffer = nullptr;
 
+		TBitArray<> VisibleInstances;
+
+		uint32 NumActiveInstances = 0;
+		uint32 MaxNumInstances = 0;
+	};
+
+	struct FLayer
+	{
+		uint32 GetCachedInstanceSectionSize();
+
+		// Feedback
+		FRDGBufferRef GeometryHandleBuffer;
+		TArray<int32> GeometryHandles;
+
 		// Special data for debugging purposes
 		FRDGBufferRef InstanceDebugBuffer = nullptr;
 
@@ -236,10 +277,7 @@ private:
 
 		uint32 NumCachedInstances = 0;
 
-		TBitArray<> VisibleInstances;
-
-		uint32 NumActiveInstances = 0;
-		uint32 MaxNumInstances = 0;
+		TArray<FLayerView> Views;
 
 		FName Name;
 	};
@@ -249,6 +287,20 @@ private:
 	// Transient memory allocator
 	FMemStackBase Allocator;
 
+	struct FViewParameters
+	{
+		const FRayTracingCullingParameters* CullingParameters;
+
+		// Used for transforming to translated world space in which TLAS was built.
+		FVector PreViewTranslation;
+	};
+
+	TArray<FViewParameters> ViewParameters;
+
+	TSparseArray<int32> ActiveViews;
+	TMap<uint32, int32> ViewIndexMap;
+	TArray<int32> TransientViewIndices;
+
 	bool bInstanceExtraDataBufferEnabled = false;
 	bool bTracingFeedbackEnabled = false;
 	bool bInstanceDebugDataEnabled = false;
diff --git a/Engine/Source/Runtime/Renderer/Private/RendererScene.cpp b/Engine/Source/Runtime/Renderer/Private/RendererScene.cpp
--- a/Engine/Source/Runtime/Renderer/Private/RendererScene.cpp
+++ b/Engine/Source/Runtime/Renderer/Private/RendererScene.cpp
@@ -469,6 +469,8 @@ void FScene::RemoveViewState_RenderThread(FSceneViewStateInterface* ViewState)
 			}
 			PerViewRayTracingData.Remove(ViewKey);
 		}
+
+		RayTracingScene.RemoveView(ViewKey);
 	}
 #endif // RHI_RAYTRACING
 }
diff --git a/Engine/Source/Runtime/Renderer/Private/SceneRendering.cpp b/Engine/Source/Runtime/Renderer/Private/SceneRendering.cpp
--- a/Engine/Source/Runtime/Renderer/Private/SceneRendering.cpp
+++ b/Engine/Source/Runtime/Renderer/Private/SceneRendering.cpp
@@ -1085,7 +1085,7 @@ FRHIRayTracingScene* FViewInfo::GetRayTracingSceneChecked(ERayTracingSceneLayer
 	{
 		if (FScene* Scene = Family->Scene->GetRenderScene())
 		{
-			FRHIRayTracingScene* Result = Scene->RayTracingScene.GetRHIRayTracingScene(Layer);
+			FRHIRayTracingScene* Result = Scene->RayTracingScene.GetRHIRayTracingScene(Layer, GetRayTracingSceneViewHandle());
 			checkf(Result, TEXT("Ray tracing scene is expected to be created at this point."));
 			return Result;
 		}
@@ -1101,7 +1101,7 @@ FRDGBufferSRVRef FViewInfo::GetRayTracingSceneLayerViewChecked(ERayTracingSceneL
 	{
 		if (FScene* Scene = Family->Scene->GetRenderScene())
 		{
-			Result = Scene->RayTracingScene.GetLayerView(Layer);
+			Result = Scene->RayTracingScene.GetLayerView(Layer, GetRayTracingSceneViewHandle());
 		}
 	}
 	checkf(Result, TEXT("Ray tracing scene SRV is expected to be created at this point."));
@@ -1115,7 +1115,7 @@ FRDGBufferUAVRef FViewInfo::GetRayTracingInstanceHitCountUAV(FRDGBuilder& GraphB
 	{
 		if (FScene* Scene = Family->Scene->GetRenderScene())
 		{
-			return Scene->RayTracingScene.GetInstanceHitCountBufferUAV(ERayTracingSceneLayer::Base);
+			return Scene->RayTracingScene.GetInstanceHitCountBufferUAV(ERayTracingSceneLayer::Base, GetRayTracingSceneViewHandle());
 		}
 	}	
 	return nullptr;
@@ -1364,7 +1364,8 @@ void FViewInfo::SetupUniformBufferParameters(
 		}
 
 #if RHI_RAYTRACING
-		const FDFVector3& RayTracingScenePreViewTranslation = Scene->RayTracingScene.PreViewTranslation;
+		const FRayTracingScene::FViewHandle& LocalRayTracingSceneViewHandle = GetRayTracingSceneViewHandle();
+		const FDFVector3& RayTracingScenePreViewTranslation = LocalRayTracingSceneViewHandle.IsValid() ? Scene->RayTracingScene.GetPreViewTranslation(LocalRayTracingSceneViewHandle) : FDFVector3{};
 		ViewUniformShaderParameters.TLASPreViewTranslationHigh = RayTracingScenePreViewTranslation.High;
 		ViewUniformShaderParameters.TLASPreViewTranslationLow = RayTracingScenePreViewTranslation.Low;
 #endif
diff --git a/Engine/Source/Runtime/Renderer/Private/SceneRendering.h b/Engine/Source/Runtime/Renderer/Private/SceneRendering.h
--- a/Engine/Source/Runtime/Renderer/Private/SceneRendering.h
+++ b/Engine/Source/Runtime/Renderer/Private/SceneRendering.h
@@ -56,7 +56,7 @@
 #include "DepthRendering.h"
 
 #if RHI_RAYTRACING
-#include "RayTracingInstanceBufferUtil.h"
+#include "RayTracing/RayTracingScene.h"
 #endif // RHI_RAYTRACING
 
 // Forward declarations.
@@ -1294,7 +1294,7 @@ public:
 	TArray<UE::Tasks::FTask> AddDynamicRayTracingMeshBatchTaskList;
 	TArray<FRayTracingShaderBindingDataOneFrameArray*, SceneRenderingAllocator> DynamicRayTracingShaderBindingsPerTask;
 	TArray<FDynamicRayTracingMeshCommandStorage*, SceneRenderingAllocator> DynamicRayTracingMeshCommandStoragePerTask;
-#endif
+#endif // RHI_RAYTRACING
 
 	// Used by mobile renderer to determine whether static meshes will be rendered with CSM shaders or not.
 	FMobileCSMVisibilityInfo MobileCSMVisibilityInfo;
@@ -1532,12 +1532,33 @@ public:
 	FLumenSceneData* ViewLumenSceneData;
 
 #if RHI_RAYTRACING
+private:
+	FRayTracingScene::FViewHandle RayTracingSceneViewHandle;
+
+public:
+
 	bool HasRayTracingScene() const;
-	FRHIRayTracingScene* GetRayTracingSceneChecked(ERayTracingSceneLayer Layer) const; // Soft-deprecated method, use FScene.RayTracingScene instead.
-	FRDGBufferSRVRef GetRayTracingSceneLayerViewChecked(ERayTracingSceneLayer Layer) const; // Soft-deprecated method, use FScene.RayTracingScene instead.
+	FRHIRayTracingScene* GetRayTracingSceneChecked(ERayTracingSceneLayer Layer) const;
+	FRDGBufferSRVRef GetRayTracingSceneLayerViewChecked(ERayTracingSceneLayer Layer) const;
 
 	FRDGBufferUAVRef GetRayTracingInstanceHitCountUAV(FRDGBuilder& GraphBuilder) const;
 
+	inline const FRayTracingScene::FViewHandle& GetRayTracingSceneViewHandle() const
+	{
+		if (UNLIKELY(IStereoRendering::IsStereoEyePass(StereoPass) && IStereoRendering::IsASecondaryPass(StereoPass)))
+		{
+			return GetPrimaryView()->RayTracingSceneViewHandle;
+		}
+
+		return RayTracingSceneViewHandle;
+	}
+
+	inline void SetRayTracingSceneViewHandle(const FRayTracingScene::FViewHandle& InHandle)
+	{
+		checkf(!IStereoRendering::IsStereoEyePass(StereoPass) || !IStereoRendering::IsASecondaryPass(StereoPass), TEXT("Stereo views should share a single ray tracing scene view."));
+		RayTracingSceneViewHandle = InHandle;
+	}
+
 	struct FRayTracingData
 	{
 		FRayTracingPipelineState* PipelineState = nullptr;
@@ -3129,4 +3150,4 @@ bool IsPrimitiveAlphaHoldoutEnabledForAnyView(TArrayView<const FViewInfo> Views)
 bool SceneCaptureRequiresAlphaChannel(const FSceneView& View);
 
 /** Checks whether the material and scene proxy combination will modify the mesh position. Useful for determining whether the material can be substituted with a default material for depth rendering etc. */
-bool DoMaterialAndPrimitiveModifyMeshPosition(const FMaterial& Material, const FPrimitiveSceneProxy* PrimitiveSceneProxy);
\ No newline at end of file
+bool DoMaterialAndPrimitiveModifyMeshPosition(const FMaterial& Material, const FPrimitiveSceneProxy* PrimitiveSceneProxy);
diff --git a/Engine/Source/Runtime/Renderer/Public/FXRenderingUtils.h b/Engine/Source/Runtime/Renderer/Public/FXRenderingUtils.h
--- a/Engine/Source/Runtime/Renderer/Public/FXRenderingUtils.h
+++ b/Engine/Source/Runtime/Renderer/Public/FXRenderingUtils.h
@@ -70,9 +70,27 @@ namespace UE::FXRenderingUtils
 #if RHI_RAYTRACING
 	namespace RayTracing
 	{
-		RENDERER_API bool HasRayTracingScene(const FSceneInterface* Scene);
-		RENDERER_API FRHIRayTracingScene* GetRayTracingScene(const FSceneInterface* Scene);
-		RENDERER_API FRHIShaderResourceView* GetRayTracingSceneView(FRHICommandListBase& RHICmdList, const FSceneInterface* Scene);
+		RENDERER_API bool HasRayTracingScene(const FSceneInterface& Scene);
+		RENDERER_API FRHIRayTracingScene* GetRayTracingScene(const FSceneInterface& Scene, const FSceneView& View);
+		RENDERER_API FRHIShaderResourceView* GetRayTracingSceneView(FRHICommandListBase& RHICmdList, const FSceneInterface& Scene, const FSceneView& View);
+
+		UE_DEPRECATED(5.7, "Please provide Scene reference instead.")
+		inline bool HasRayTracingScene(const FSceneInterface* Scene)
+		{
+			return HasRayTracingScene(*Scene);
+		}
+
+		UE_DEPRECATED(5.7, "Please provide Scene and View reference.")
+		inline FRHIRayTracingScene* GetRayTracingScene(const FSceneInterface* Scene)
+		{
+			return nullptr;
+		}
+
+		UE_DEPRECATED(5.7, "Please provide Scene and View reference.")
+		inline FRHIShaderResourceView* GetRayTracingSceneView(FRHICommandListBase& RHICmdList, const FSceneInterface* Scene)
+		{
+			return nullptr;
+		}
 
 		RENDERER_API FShaderBindingTableRHIRef CreateShaderBindingTable(FRHICommandListBase& RHICmdList, const FSceneInterface* Scene,  uint32 LocalBindingDataSize);		
 		RENDERER_API TConstArrayView<FRayTracingShaderBindingData> GetVisibleRayTracingShaderBindings(const FSceneView& View);
