[HWRT] Fix issues with readback logic in FRayTracingScene::FinishStats(...)
- need to store current MaxNumViews when enqueuing buffer readback and then use that number when processing the results.
- transient views were not being removed from ActiveViews array.
- prevent overwriting readback buffers that are not ready yet.
- avoid race condition where MaxNumViews is updated before RDG pass is executed which caused stale results to be processed.
- done by processing ready results before enqueing new ones.

diff --git a/Engine/Source/Runtime/Renderer/Private/RayTracing/RayTracingScene.cpp b/Engine/Source/Runtime/Renderer/Private/RayTracing/RayTracingScene.cpp
--- a/Engine/Source/Runtime/Renderer/Private/RayTracing/RayTracingScene.cpp
+++ b/Engine/Source/Runtime/Renderer/Private/RayTracing/RayTracingScene.cpp
@@ -229,7 +229,7 @@ void FRayTracingScene::Update(FRDGBuilder& GraphBuilder, FSceneUniformBuffer& Sc
 	FRHICommandListBase& RHICmdList = GraphBuilder.RHICmdList;
 
 	const uint8 NumLayers = uint8(ERayTracingSceneLayer::NUM);
-	const uint32 MaxNumViews = ActiveViews.GetMaxIndex() + 1;
+	const uint32 MaxNumViews = ActiveViews.GetMaxIndex();
 
 	FRDGBufferUAVRef InstanceStatsBufferUAV = nullptr;
 	{
@@ -715,8 +715,44 @@ void FRayTracingScene::FinishTracingFeedback(FRDGBuilder& GraphBuilder, ERDGPass
 		}
 	}
 
-	// copy feedback to readback buffer
+	// process ready results
+
+	while (FeedbackReadbackNumPending > 0)
 	{
+		uint32 Index = (FeedbackReadbackWriteIndex + MaxReadbackBuffers - FeedbackReadbackNumPending) % MaxReadbackBuffers;
+		FRHIGPUBufferReadback* GeometryHandleReadbackBuffer = FeedbackReadback[Index].GeometryHandleReadbackBuffer;
+		FRHIGPUBufferReadback* GeometryCountReadbackBuffer = FeedbackReadback[Index].GeometryCountReadbackBuffer;
+		check(GeometryHandleReadbackBuffer->IsReady() == GeometryCountReadbackBuffer->IsReady());
+		if (GeometryHandleReadbackBuffer->IsReady() && GeometryCountReadbackBuffer->IsReady())
+		{
+			FeedbackReadbackNumPending--;
+
+			const uint32* GeometryCountPtr = (const uint32*)GeometryCountReadbackBuffer->Lock(sizeof(uint32));
+			const uint32 GeometryCount = GeometryCountPtr[0];
+			GeometryCountReadbackBuffer->Unlock();
+
+			const int32* GeometryHandlesPtr = (const int32*)GeometryHandleReadbackBuffer->Lock(sizeof(int32) * GeometryCount);
+
+			for (uint32 i = 0; i < GeometryCount; i++)
+			{
+				if (ensure(GeometryHandlesPtr[i] != INDEX_NONE))
+				{
+					GRayTracingGeometryManager->AddVisibleGeometry(GeometryHandlesPtr[i]);
+				}
+			}
+
+			GeometryHandleReadbackBuffer->Unlock();
+		}
+		else
+		{
+			break;
+		}
+	}
+
+	//if (FeedbackReadbackNumPending < MaxReadbackBuffers) // TODO: need to prevent overwritng results?
+	{
+		// copy feedback to readback buffer
+
 		FFeedbackReadbackPassParameters* PassParameters = GraphBuilder.AllocParameters<FFeedbackReadbackPassParameters>();
 		PassParameters->HandleBuffer = GeometryHandleBuffer;
 		PassParameters->CountBuffer = GeometryHandleAllocatorBuffer;
@@ -733,83 +769,37 @@ void FRayTracingScene::FinishTracingFeedback(FRDGBuilder& GraphBuilder, ERDGPass
 
 		FeedbackReadbackWriteIndex = (FeedbackReadbackWriteIndex + 1u) % MaxReadbackBuffers;
 		FeedbackReadbackNumPending = FMath::Min(FeedbackReadbackNumPending + 1u, MaxReadbackBuffers);
-
-		while (FeedbackReadbackNumPending > 0)
-		{
-			uint32 Index = (FeedbackReadbackWriteIndex + MaxReadbackBuffers - FeedbackReadbackNumPending) % MaxReadbackBuffers;
-			FRHIGPUBufferReadback* GeometryHandleReadbackBuffer = FeedbackReadback[Index].GeometryHandleReadbackBuffer;
-			FRHIGPUBufferReadback* GeometryCountReadbackBuffer = FeedbackReadback[Index].GeometryCountReadbackBuffer;
-			check(GeometryHandleReadbackBuffer->IsReady() == GeometryCountReadbackBuffer->IsReady());
-			if (GeometryHandleReadbackBuffer->IsReady() && GeometryCountReadbackBuffer->IsReady())
-			{
-				FeedbackReadbackNumPending--;
-				
-				const uint32* GeometryCountPtr = (const uint32*)GeometryCountReadbackBuffer->Lock(sizeof(uint32));
-				const uint32 GeometryCount = GeometryCountPtr[0];
-				GeometryCountReadbackBuffer->Unlock();
-
-				const int32* GeometryHandlesPtr = (const int32*)GeometryHandleReadbackBuffer->Lock(sizeof(int32) * GeometryCount);
-
-				for (uint32 i = 0; i < GeometryCount; i++)
-				{					
-					if (ensure(GeometryHandlesPtr[i] != INDEX_NONE))
-					{
-						GRayTracingGeometryManager->AddVisibleGeometry(GeometryHandlesPtr[i]);
-					}				
-				}
-
-				GeometryHandleReadbackBuffer->Unlock();
-			}
-			else
-			{
-				break;
-			}
-		}
 	}
 }
 
 void FRayTracingScene::FinishStats(FRDGBuilder& GraphBuilder, ERDGPassFlags ComputePassFlags)
 {
 	//  if necessary create readback buffers
-	if (StatsReadbackBuffers.IsEmpty())
+	if (StatsReadback.IsEmpty())
 	{
-		StatsReadbackBuffers.SetNum(MaxReadbackBuffers);
+		StatsReadback.SetNum(MaxReadbackBuffers);
 
 		for (uint32 Index = 0; Index < MaxReadbackBuffers; ++Index)
 		{
-			StatsReadbackBuffers[Index] = new FRHIGPUBufferReadback(TEXT("FRayTracingScene::StatsReadbackBuffer"));
+			StatsReadback[Index].ReadbackBuffer = new FRHIGPUBufferReadback(TEXT("FRayTracingScene::StatsReadbackBuffer"));
 		}
 	}
 
-	// copy stats to readback buffer
-	if (InstanceStatsBuffer != nullptr)
-	{
-		AddReadbackBufferPass(GraphBuilder, RDG_EVENT_NAME("FRayTracingScene::StatsReadback"), InstanceStatsBuffer,
-			[ReadbackBuffer = StatsReadbackBuffers[StatsReadbackBuffersWriteIndex], InstanceStatsBuffer = InstanceStatsBuffer](FRDGAsyncTask, FRHICommandList& RHICmdList)
-			{
-				ReadbackBuffer->EnqueueCopy(RHICmdList, InstanceStatsBuffer->GetRHI(), 0u);
-			});
-
-		StatsReadbackBuffersWriteIndex = (StatsReadbackBuffersWriteIndex + 1u) % MaxReadbackBuffers;
-		StatsReadbackBuffersNumPending = FMath::Min(StatsReadbackBuffersNumPending + 1u, MaxReadbackBuffers);
-	}
-
 	uint32 TotalNumNativeInstances = 0;
 	uint32 TotalNumActiveInstances = 0;
 
 	const uint8 NumLayers = uint8(ERayTracingSceneLayer::NUM);
-	const uint32 MaxNumViews = ActiveViews.GetMaxIndex() + 1;
 
 	// process ready results
-	while (StatsReadbackBuffersNumPending > 0)
+	while (StatsReadbackNumPending > 0)
 	{
-		uint32 Index = (StatsReadbackBuffersWriteIndex + MaxReadbackBuffers - StatsReadbackBuffersNumPending) % MaxReadbackBuffers;
-		FRHIGPUBufferReadback* ReadbackBuffer = StatsReadbackBuffers[Index];
-		if (ReadbackBuffer->IsReady())
+		uint32 Index = (StatsReadbackWriteIndex + MaxReadbackBuffers - StatsReadbackNumPending) % MaxReadbackBuffers;
+		FStatsReadbackData& ReadbackData = StatsReadback[Index];
+		if (ReadbackData.ReadbackBuffer->IsReady())
 		{
-			StatsReadbackBuffersNumPending--;
+			StatsReadbackNumPending--;
 
-			auto ReadbackBufferPtr = (const FInstanceBufferStats*)ReadbackBuffer->Lock(sizeof(FInstanceBufferStats) * NumLayers * MaxNumViews);
+			auto ReadbackBufferPtr = (const FInstanceBufferStats*)ReadbackData.ReadbackBuffer->Lock(sizeof(FInstanceBufferStats) * NumLayers * ReadbackData.MaxNumViews);
 
 			for (uint32 LayerIndex = 0; LayerIndex < NumLayers; ++LayerIndex)
 			{
@@ -821,14 +811,14 @@ void FRayTracingScene::FinishStats(FRDGBuilder& GraphBuilder, ERDGPassFlags Comp
 
 					const uint32 LayerViewNumNativeInstances = LayerView.InstanceBufferBuilder.GetMaxNumInstances();
 
-					LayerView.NumActiveInstances = FMath::Min(ReadbackBufferPtr[LayerIndex * MaxNumViews + ViewIndex], LayerViewNumNativeInstances);
+					LayerView.NumActiveInstances = FMath::Min(ReadbackBufferPtr[LayerIndex * ReadbackData.MaxNumViews + ViewIndex], LayerViewNumNativeInstances);
 
 					TotalNumNativeInstances += LayerViewNumNativeInstances;
 					TotalNumActiveInstances += LayerView.NumActiveInstances;
 				}
 			}
 
-			ReadbackBuffer->Unlock();
+			ReadbackData.ReadbackBuffer->Unlock();
 		}
 		else
 		{
@@ -838,6 +828,21 @@ void FRayTracingScene::FinishStats(FRDGBuilder& GraphBuilder, ERDGPassFlags Comp
 
 	SET_DWORD_STAT(STAT_RayTracingTotalInstances, TotalNumNativeInstances);
 	SET_DWORD_STAT(STAT_RayTracingActiveInstances, TotalNumActiveInstances);
+
+	// copy stats to readback buffer
+	if (InstanceStatsBuffer != nullptr && StatsReadbackNumPending < MaxReadbackBuffers)
+	{
+		AddReadbackBufferPass(GraphBuilder, RDG_EVENT_NAME("FRayTracingScene::StatsReadback"), InstanceStatsBuffer,
+			[&ReadbackData = StatsReadback[StatsReadbackWriteIndex], InstanceStatsBuffer = InstanceStatsBuffer](FRDGAsyncTask, FRHICommandList& RHICmdList)
+			{
+				ReadbackData.ReadbackBuffer->EnqueueCopy(RHICmdList, InstanceStatsBuffer->GetRHI(), 0u);
+			});
+
+		StatsReadback[StatsReadbackWriteIndex].MaxNumViews = ActiveViews.GetMaxIndex();
+
+		StatsReadbackWriteIndex = (StatsReadbackWriteIndex + 1u) % MaxReadbackBuffers;
+		StatsReadbackNumPending = FMath::Min(StatsReadbackNumPending + 1u, MaxReadbackBuffers);
+	}
 }
 
 void FRayTracingScene::PostRender(FRDGBuilder& GraphBuilder, ERDGPassFlags ComputePassFlags)
@@ -1067,6 +1072,7 @@ void FRayTracingScene::Reset()
 			Layer.Views[ViewIndex] = {};
 		}
 
+		ActiveViews.RemoveAt(ViewIndex);
 		ViewParameters[ViewIndex] = {};
 	}
 
@@ -1224,15 +1230,15 @@ bool FRayTracingScene::SetInstanceDebugDataEnabled(bool bInEnabled)
 
 void FRayTracingScene::ReleaseReadbackBuffers()
 {
-	for (auto& ReadbackBuffer : StatsReadbackBuffers)
+	for (auto& ReadbackData : StatsReadback)
 	{
-		delete ReadbackBuffer;
+		delete ReadbackData.ReadbackBuffer;
 	}
 
-	StatsReadbackBuffers.Empty();
+	StatsReadback.Empty();
 
-	StatsReadbackBuffersWriteIndex = 0;
-	StatsReadbackBuffersNumPending = 0;
+	StatsReadbackWriteIndex = 0;
+	StatsReadbackNumPending = 0;
 }
 
 void FRayTracingScene::ReleaseFeedbackReadbackBuffers()
diff --git a/Engine/Source/Runtime/Renderer/Private/RayTracing/RayTracingScene.h b/Engine/Source/Runtime/Renderer/Private/RayTracing/RayTracingScene.h
--- a/Engine/Source/Runtime/Renderer/Private/RayTracing/RayTracingScene.h
+++ b/Engine/Source/Runtime/Renderer/Private/RayTracing/RayTracingScene.h
@@ -313,8 +313,8 @@ private:
 
 	struct FFeedbackReadbackData
 	{
-		FRHIGPUBufferReadback* GeometryHandleReadbackBuffer;
-		FRHIGPUBufferReadback* GeometryCountReadbackBuffer;		
+		FRHIGPUBufferReadback* GeometryHandleReadbackBuffer = nullptr;
+		FRHIGPUBufferReadback* GeometryCountReadbackBuffer = nullptr;
 	};
 
 	const uint32 MaxReadbackBuffers = 4;
@@ -325,9 +325,15 @@ private:
 	uint32 FeedbackReadbackWriteIndex = 0;
 	uint32 FeedbackReadbackNumPending = 0;
 
-	TArray<FRHIGPUBufferReadback*> StatsReadbackBuffers;
-	uint32 StatsReadbackBuffersWriteIndex = 0;
-	uint32 StatsReadbackBuffersNumPending = 0;
+	struct FStatsReadbackData
+	{
+		FRHIGPUBufferReadback* ReadbackBuffer = nullptr;
+		uint32 MaxNumViews = 0;
+	};
+
+	TArray<FStatsReadbackData> StatsReadback;
+	uint32 StatsReadbackWriteIndex = 0;
+	uint32 StatsReadbackNumPending = 0;
 };
 
 #endif // RHI_RAYTRACING
