Fix mGPU Corruption

Rearrange the order of operations in the VirtualShadowMapArray initialization so that the mismatched transitions, caused by per-GPU resource access state
divergence, don't cause visual corruption.

#rnx

diff --git a/Engine/Source/Runtime/Renderer/Private/VirtualShadowMaps/VirtualShadowMapArray.cpp b/Engine/Source/Runtime/Renderer/Private/VirtualShadowMaps/VirtualShadowMapArray.cpp
--- a/Engine/Source/Runtime/Renderer/Private/VirtualShadowMaps/VirtualShadowMapArray.cpp
+++ b/Engine/Source/Runtime/Renderer/Private/VirtualShadowMaps/VirtualShadowMapArray.cpp
@@ -2791,6 +2791,10 @@ void FVirtualShadowMapArray::BuildPageAllocations(
 	RDG_EVENT_SCOPE(GraphBuilder, "FVirtualShadowMapArray::BuildPageAllocation");
 	SCOPED_NAMED_EVENT(FVirtualShadowMapArray_BuildPageAllocation, FColor::Emerald);
 
+#if WITH_MGPU
+	bool bCopyOtherGPUsPhysicalPageLists = false;
+#endif
+
 	// Update cached or newly invalidated pages with respect to the new requests
 	{	
 		// Cached data from previous frames is available and valid.  Note that we currently don't support GPUMask varying within
@@ -2837,8 +2841,7 @@ void FVirtualShadowMapArray::BuildPageAllocations(
 		FRHIGPUMask InverseGPUMask;
 		if (CacheManager->GetPrevBuffers().PhysicalPageLists && Views[0].GPUMask.Invert(InverseGPUMask) && InverseGPUMask.Intersects(CacheManager->GetCacheValidGPUMask()))
 		{
-			RDG_GPU_MASK_SCOPE(GraphBuilder, InverseGPUMask);
-			AddCopyBufferPass(GraphBuilder, PhysicalPageListsRDG, GraphBuilder.RegisterExternalBuffer(CacheManager->GetPrevBuffers().PhysicalPageLists));
+			bCopyOtherGPUsPhysicalPageLists = true;
 		}
 
 		// Track which GPUs the cache has been initialized on.  Merges GPU mask if cache data was already available, otherwise sets mask to initialize it.
@@ -3026,6 +3029,26 @@ void FVirtualShadowMapArray::BuildPageAllocations(
 
 	UpdateCachedUniformBuffers(GraphBuilder);
 
+#if WITH_MGPU
+	if (bCopyOtherGPUsPhysicalPageLists)
+	{
+		//
+		// Do this at the end to prevent visual corruption as reported in [UE-295062]
+		//
+		// The RDG doesn't have any concept of there being a copy of these resources per GPU so when
+		// the GPUMask is anything other than All, the access states can diverge on a per-GPU basis.
+		// This can then result in incorrect transitions for some of the resource copies. This fix 
+		// isn't perfect and does not completely solve the problem, even in this limited case, but does
+		// appear to avoid the reported corruption and is light-weight compared to other potential fixes.
+		//
+
+		FRHIGPUMask InverseGPUMask;
+		Views[0].GPUMask.Invert(InverseGPUMask);
+		RDG_GPU_MASK_SCOPE(GraphBuilder, InverseGPUMask);
+		AddCopyBufferPass(GraphBuilder, PhysicalPageListsRDG, GraphBuilder.RegisterExternalBuffer(CacheManager->GetPrevBuffers().PhysicalPageLists));
+	}
+#endif
+
 #if !UE_BUILD_SHIPPING
 	// Only dump one frame of light data
 	GDumpVSMLightNames = false;
