[Exposure] Fixed compensation curve LUT not working on mobile renderer.

diff --git a/Engine/Shaders/Private/PostProcessMobile.usf b/Engine/Shaders/Private/PostProcessMobile.usf
--- a/Engine/Shaders/Private/PostProcessMobile.usf
+++ b/Engine/Shaders/Private/PostProcessMobile.usf
@@ -1403,7 +1403,11 @@ float4 BasicEyeAdaptation_Mobile()
 	// Convert LogLuminanceAverage to Average Intensity
 	const float AverageSceneLuminance = exp2(LogLumAve);
 
-	const float MiddleGreyExposureCompensation = EyeAdaptation_ExposureCompensationSettings * EyeAdaptation_ExposureCompensationCurve * EyeAdaptation_GreyMult;// we want the average luminance remapped to 0.18, not 1.0
+	const float EyeAdaptationEV100 = log2((AverageSceneLuminance / 0.18) / EyeAdaptation_LuminanceMax);
+	const float LUT_U = EyeAdaptationEV100 * EyeAdaptation_EV100ToExposureCompensationCurveLUTScaleBias.x + EyeAdaptation_EV100ToExposureCompensationCurveLUTScaleBias.y;
+	const float ExposureCompensationFromCurveLUT = Texture2DSampleLevel(EyeAdaptation_ExposureCompensationCurveLUT, EyeAdaptation_ExposureCompensationCurveSampler, float2(LUT_U, 0.5f), 0).x;
+
+	const float MiddleGreyExposureCompensation = EyeAdaptation_ExposureCompensationSettings * EyeAdaptation_ExposureCompensationCurve * ExposureCompensationFromCurveLUT * EyeAdaptation_GreyMult;// we want the average luminance remapped to 0.18, not 1.0
 
 	const float ClampedLumAve = clamp(AverageSceneLuminance, EyeAdaptation_MinAverageLuminance, EyeAdaptation_MaxAverageLuminance);
 	
@@ -1671,7 +1675,11 @@ void HistogramEyeAdaptationCS(uint2 DispatchThreadId : SV_DispatchThreadID)
 
 	const float OldExposureScale = EyeAdaptationBuffer[0].x;
 
-	const float MiddleGreyExposureCompensation = EyeAdaptation_ExposureCompensationSettings * EyeAdaptation_ExposureCompensationCurve; // we want the average luminance remapped to 0.18, not 1.0
+	const float EyeAdaptationEV100 = log2((AverageSceneLuminance / 0.18) / EyeAdaptation_LuminanceMax);
+	const float LUT_U = EyeAdaptationEV100 * EyeAdaptation_EV100ToExposureCompensationCurveLUTScaleBias.x + EyeAdaptation_EV100ToExposureCompensationCurveLUTScaleBias.y;
+	const float ExposureCompensationFromCurveLUT = Texture2DSampleLevel(EyeAdaptation_ExposureCompensationCurveLUT, EyeAdaptation_ExposureCompensationCurveSampler, float2(LUT_U, 0.5f), 0).x;
+
+	const float MiddleGreyExposureCompensation = EyeAdaptation_ExposureCompensationSettings * EyeAdaptation_ExposureCompensationCurve * ExposureCompensationFromCurveLUT; // we want the average luminance remapped to 0.18, not 1.0
 	const float OldExposure = MiddleGreyExposureCompensation / (OldExposureScale != 0 ? OldExposureScale : 1.0f);
 
 	// eye adaptation changes over time
diff --git a/Engine/Source/Runtime/Renderer/Private/MobileShadingRenderer.cpp b/Engine/Source/Runtime/Renderer/Private/MobileShadingRenderer.cpp
--- a/Engine/Source/Runtime/Renderer/Private/MobileShadingRenderer.cpp
+++ b/Engine/Source/Runtime/Renderer/Private/MobileShadingRenderer.cpp
@@ -1035,6 +1035,12 @@ void FMobileSceneRenderer::RenderCustomRenderPassBasePass(FRDGBuilder& GraphBuil
 	}
 }
 
+struct FUpdateExposureCompensationCurveLUTTaskData
+{
+	UE::Tasks::FTask Task;
+	FRHICommandList* CmdList = nullptr;
+};
+
 void FMobileSceneRenderer::Render(FRDGBuilder& GraphBuilder, const FSceneRenderUpdateInputs* SceneUpdateInputs)
 {
 	if (!ViewFamily.EngineShowFlags.Rendering)
@@ -1050,6 +1056,25 @@ void FMobileSceneRenderer::Render(FRDGBuilder& GraphBuilder, const FSceneRenderU
 
 	IVisibilityTaskData* VisibilityTaskData = OnRenderBegin(GraphBuilder, SceneUpdateInputs);
 
+	FUpdateExposureCompensationCurveLUTTaskData UpdateExposureCompensationCurveLUTTaskData;
+	UpdateExposureCompensationCurveLUTTaskData.Task = LaunchSceneRenderTask(TEXT("UpdateExposureCompensationCurveLUTs"),
+		[TaskData = &UpdateExposureCompensationCurveLUTTaskData, Views = MakeConstArrayView(Views)]
+		{
+			TaskData->CmdList = new FRHICommandList();
+			TaskData->CmdList->SwitchPipeline(ERHIPipeline::Graphics);
+
+			for (int32 ViewIndex = 0; ViewIndex < Views.Num(); ++ViewIndex)
+			{
+				const FViewInfo& View = Views[ViewIndex];
+				if (View.ViewState != nullptr)
+				{
+					View.ViewState->UpdateExposureCompensationCurveLUT(*TaskData->CmdList, View.FinalPostProcessSettings.AutoExposureBiasCurve);
+				}
+			}
+
+			TaskData->CmdList->FinishRecording();
+		});
+
 	FRDGExternalAccessQueue ExternalAccessQueue;
 
 	if (SceneUpdateInputs)
@@ -1559,6 +1584,15 @@ void FMobileSceneRenderer::Render(FRDGBuilder& GraphBuilder, const FSceneRenderU
 					RDG_GPU_STAT_SCOPE(GraphBuilder, Postprocessing);
 					SCOPE_CYCLE_COUNTER(STAT_FinishRenderViewTargetTime);
 
+					if (UpdateExposureCompensationCurveLUTTaskData.Task.IsValid())
+					{
+						UpdateExposureCompensationCurveLUTTaskData.Task.Wait();
+						UpdateExposureCompensationCurveLUTTaskData.Task = {};
+
+						GraphBuilder.RHICmdList.QueueAsyncCommandListSubmit(UpdateExposureCompensationCurveLUTTaskData.CmdList);
+						UpdateExposureCompensationCurveLUTTaskData.CmdList = nullptr;
+					}
+
 					FMobilePostProcessingInputs PostProcessingInputs;
 					PostProcessingInputs.ViewFamilyTexture = ViewFamilyTexture;
 					PostProcessingInputs.SceneTextures = CreateMobileSceneTextureUniformBuffer(GraphBuilder, &SceneTextures, EMobileSceneTextureSetupMode::All);
@@ -1601,6 +1635,15 @@ void FMobileSceneRenderer::Render(FRDGBuilder& GraphBuilder, const FSceneRenderU
 	GEngine->GetPostRenderDelegateEx().Broadcast(GraphBuilder);
 	GetSceneExtensionsRenderers().PostRender(GraphBuilder);
 
+	if (UpdateExposureCompensationCurveLUTTaskData.Task.IsValid())
+	{
+		UpdateExposureCompensationCurveLUTTaskData.Task.Wait();
+		UpdateExposureCompensationCurveLUTTaskData.Task = {};
+
+		GraphBuilder.RHICmdList.QueueAsyncCommandListSubmit(UpdateExposureCompensationCurveLUTTaskData.CmdList);
+		UpdateExposureCompensationCurveLUTTaskData.CmdList = nullptr;
+	}
+
 	if (bUseVirtualTexturing)
 	{
 		VirtualTexture::EndFeedback(GraphBuilder);
