Fix incorrect comment in Substrate.

diff --git a/Engine/Source/Runtime/Renderer/Private/Substrate/Substrate.h b/Engine/Source/Runtime/Renderer/Private/Substrate/Substrate.h
--- a/Engine/Source/Runtime/Renderer/Private/Substrate/Substrate.h
+++ b/Engine/Source/Runtime/Renderer/Private/Substrate/Substrate.h
@@ -147,7 +147,7 @@ struct FSubstrateViewData
 	uint32 MaxClosurePerPixel = 0;
 	uint32 MaxBytesPerPixel = 0;
 
-	// True if any material requires the complex special path (e.g. glints or SpecularLUT)
+	// Visible/Active tile type mask (1 bit per tile type)
 	uint8 UsesTileTypeMask = 0;
 	bool bUsesAnisotropy = false;
 
