[Exposure] Moved some compensation curve LUT logic to helper functions so it can be shared between renderers.

diff --git a/Engine/Source/Runtime/Renderer/Private/DeferredShadingRenderer.cpp b/Engine/Source/Runtime/Renderer/Private/DeferredShadingRenderer.cpp
--- a/Engine/Source/Runtime/Renderer/Private/DeferredShadingRenderer.cpp
+++ b/Engine/Source/Runtime/Renderer/Private/DeferredShadingRenderer.cpp
@@ -20,6 +20,7 @@
 #include "PostProcess/PostProcessSubsurface.h"
 #include "PostProcess/PostProcessVisualizeCalibrationMaterial.h"
 #include "PostProcess/TemporalAA.h"
+#include "PostProcess/PostProcessEyeAdaptation.h"
 #include "CompositionLighting/CompositionLighting.h"
 #include "FXSystem.h"
 #include "OneColorShader.h"
@@ -1346,12 +1347,6 @@ void FDeferredShadingSceneRenderer::CommitFinalPipelineState()
 	}
 }
 
-struct FUpdateExposureCompensationCurveLUTTaskData
-{
-	UE::Tasks::FTask Task;
-	FRHICommandList* CmdList = nullptr;
-};
-
 void FDeferredShadingSceneRenderer::RenderNanite(FRDGBuilder& GraphBuilder, const TArray<FViewInfo>& InViews, FSceneTextures& SceneTextures, bool bIsEarlyDepthComplete,
 	FNaniteBasePassVisibility& InNaniteBasePassVisibility,
 	TArray<Nanite::FRasterResults, TInlineAllocator<2>>& NaniteRasterResults,
@@ -1788,23 +1783,7 @@ void FDeferredShadingSceneRenderer::Render(FRDGBuilder& GraphBuilder, const FSce
 	FInitViewTaskDatas InitViewTaskDatas = OnRenderBegin(GraphBuilder, SceneUpdateInputs);
 
 	FUpdateExposureCompensationCurveLUTTaskData UpdateExposureCompensationCurveLUTTaskData;
-	UpdateExposureCompensationCurveLUTTaskData.Task = LaunchSceneRenderTask(TEXT("UpdateExposureCompensationCurveLUTs"),
-		[TaskData = &UpdateExposureCompensationCurveLUTTaskData, Views = MakeConstArrayView(Views)]
-		{
-			TaskData->CmdList = new FRHICommandList();
-			TaskData->CmdList->SwitchPipeline(ERHIPipeline::Graphics);
-
-			for (int32 ViewIndex = 0; ViewIndex < Views.Num(); ++ViewIndex)
-			{
-				const FViewInfo& View = Views[ViewIndex];
-				if (View.ViewState != nullptr)
-				{
-					View.ViewState->UpdateExposureCompensationCurveLUT(*TaskData->CmdList, View.FinalPostProcessSettings.AutoExposureBiasCurve);
-				}
-			}
-
-			TaskData->CmdList->FinishRecording();
-		});
+	BeginUpdateExposureCompensationCurveLUT(Views, &UpdateExposureCompensationCurveLUTTaskData);
 
 	FRDGExternalAccessQueue ExternalAccessQueue;
 	TUniquePtr<FVirtualTextureUpdater> VirtualTextureUpdater;
@@ -3883,14 +3862,7 @@ void FDeferredShadingSceneRenderer::Render(FRDGBuilder& GraphBuilder, const FSce
 			RDG_GPU_STAT_SCOPE(GraphBuilder, Postprocessing);
 			SCOPED_NAMED_EVENT(PostProcessing, FColor::Emerald);
 
-			if (UpdateExposureCompensationCurveLUTTaskData.Task.IsValid())
-			{
-				UpdateExposureCompensationCurveLUTTaskData.Task.Wait();
-				UpdateExposureCompensationCurveLUTTaskData.Task = {};
-
-				GraphBuilder.RHICmdList.QueueAsyncCommandListSubmit(UpdateExposureCompensationCurveLUTTaskData.CmdList);
-				UpdateExposureCompensationCurveLUTTaskData.CmdList = nullptr;
-			}
+			FinishUpdateExposureCompensationCurveLUT(GraphBuilder.RHICmdList, &UpdateExposureCompensationCurveLUTTaskData);
 
 			FPostProcessingInputs PostProcessingInputs;
 			PostProcessingInputs.ViewFamilyTexture = ViewFamilyTexture;
@@ -4008,14 +3980,7 @@ void FDeferredShadingSceneRenderer::Render(FRDGBuilder& GraphBuilder, const FSce
 		GEngine->GetPostRenderDelegateEx().Broadcast(GraphBuilder);
 	}
 
-	if (UpdateExposureCompensationCurveLUTTaskData.Task.IsValid())
-	{
-		UpdateExposureCompensationCurveLUTTaskData.Task.Wait();
-		UpdateExposureCompensationCurveLUTTaskData.Task = {};
-
-		GraphBuilder.RHICmdList.QueueAsyncCommandListSubmit(UpdateExposureCompensationCurveLUTTaskData.CmdList);
-		UpdateExposureCompensationCurveLUTTaskData.CmdList = nullptr;
-	}
+	FinishUpdateExposureCompensationCurveLUT(GraphBuilder.RHICmdList, &UpdateExposureCompensationCurveLUTTaskData);
 	
 	GetSceneExtensionsRenderers().PostRender(GraphBuilder);
 
diff --git a/Engine/Source/Runtime/Renderer/Private/MobileShadingRenderer.cpp b/Engine/Source/Runtime/Renderer/Private/MobileShadingRenderer.cpp
--- a/Engine/Source/Runtime/Renderer/Private/MobileShadingRenderer.cpp
+++ b/Engine/Source/Runtime/Renderer/Private/MobileShadingRenderer.cpp
@@ -1035,12 +1035,6 @@ void FMobileSceneRenderer::RenderCustomRenderPassBasePass(FRDGBuilder& GraphBuil
 	}
 }
 
-struct FUpdateExposureCompensationCurveLUTTaskData
-{
-	UE::Tasks::FTask Task;
-	FRHICommandList* CmdList = nullptr;
-};
-
 void FMobileSceneRenderer::Render(FRDGBuilder& GraphBuilder, const FSceneRenderUpdateInputs* SceneUpdateInputs)
 {
 	if (!ViewFamily.EngineShowFlags.Rendering)
@@ -1057,23 +1051,7 @@ void FMobileSceneRenderer::Render(FRDGBuilder& GraphBuilder, const FSceneRenderU
 	IVisibilityTaskData* VisibilityTaskData = OnRenderBegin(GraphBuilder, SceneUpdateInputs);
 
 	FUpdateExposureCompensationCurveLUTTaskData UpdateExposureCompensationCurveLUTTaskData;
-	UpdateExposureCompensationCurveLUTTaskData.Task = LaunchSceneRenderTask(TEXT("UpdateExposureCompensationCurveLUTs"),
-		[TaskData = &UpdateExposureCompensationCurveLUTTaskData, Views = MakeConstArrayView(Views)]
-		{
-			TaskData->CmdList = new FRHICommandList();
-			TaskData->CmdList->SwitchPipeline(ERHIPipeline::Graphics);
-
-			for (int32 ViewIndex = 0; ViewIndex < Views.Num(); ++ViewIndex)
-			{
-				const FViewInfo& View = Views[ViewIndex];
-				if (View.ViewState != nullptr)
-				{
-					View.ViewState->UpdateExposureCompensationCurveLUT(*TaskData->CmdList, View.FinalPostProcessSettings.AutoExposureBiasCurve);
-				}
-			}
-
-			TaskData->CmdList->FinishRecording();
-		});
+	BeginUpdateExposureCompensationCurveLUT(Views, &UpdateExposureCompensationCurveLUTTaskData);
 
 	FRDGExternalAccessQueue ExternalAccessQueue;
 
@@ -1584,14 +1562,7 @@ void FMobileSceneRenderer::Render(FRDGBuilder& GraphBuilder, const FSceneRenderU
 					RDG_GPU_STAT_SCOPE(GraphBuilder, Postprocessing);
 					SCOPE_CYCLE_COUNTER(STAT_FinishRenderViewTargetTime);
 
-					if (UpdateExposureCompensationCurveLUTTaskData.Task.IsValid())
-					{
-						UpdateExposureCompensationCurveLUTTaskData.Task.Wait();
-						UpdateExposureCompensationCurveLUTTaskData.Task = {};
-
-						GraphBuilder.RHICmdList.QueueAsyncCommandListSubmit(UpdateExposureCompensationCurveLUTTaskData.CmdList);
-						UpdateExposureCompensationCurveLUTTaskData.CmdList = nullptr;
-					}
+					FinishUpdateExposureCompensationCurveLUT(GraphBuilder.RHICmdList, &UpdateExposureCompensationCurveLUTTaskData);
 
 					FMobilePostProcessingInputs PostProcessingInputs;
 					PostProcessingInputs.ViewFamilyTexture = ViewFamilyTexture;
@@ -1635,14 +1606,7 @@ void FMobileSceneRenderer::Render(FRDGBuilder& GraphBuilder, const FSceneRenderU
 	GEngine->GetPostRenderDelegateEx().Broadcast(GraphBuilder);
 	GetSceneExtensionsRenderers().PostRender(GraphBuilder);
 
-	if (UpdateExposureCompensationCurveLUTTaskData.Task.IsValid())
-	{
-		UpdateExposureCompensationCurveLUTTaskData.Task.Wait();
-		UpdateExposureCompensationCurveLUTTaskData.Task = {};
-
-		GraphBuilder.RHICmdList.QueueAsyncCommandListSubmit(UpdateExposureCompensationCurveLUTTaskData.CmdList);
-		UpdateExposureCompensationCurveLUTTaskData.CmdList = nullptr;
-	}
+	FinishUpdateExposureCompensationCurveLUT(GraphBuilder.RHICmdList, &UpdateExposureCompensationCurveLUTTaskData);
 
 	if (bUseVirtualTexturing)
 	{
diff --git a/Engine/Source/Runtime/Renderer/Private/PostProcess/PostProcessEyeAdaptation.cpp b/Engine/Source/Runtime/Renderer/Private/PostProcess/PostProcessEyeAdaptation.cpp
--- a/Engine/Source/Runtime/Renderer/Private/PostProcess/PostProcessEyeAdaptation.cpp
+++ b/Engine/Source/Runtime/Renderer/Private/PostProcess/PostProcessEyeAdaptation.cpp
@@ -411,6 +411,39 @@ FVector2f GetExposureCompensationCurveLUTScaleBias()
 	return FVector2f(Scale, Bias);
 }
 
+void BeginUpdateExposureCompensationCurveLUT(TArrayView<FViewInfo> Views, FUpdateExposureCompensationCurveLUTTaskData* TaskData)
+{
+	TaskData->Task = LaunchSceneRenderTask(TEXT("UpdateExposureCompensationCurveLUTs"),
+		[TaskData, Views]
+		{
+			TaskData->CmdList = new FRHICommandList();
+			TaskData->CmdList->SwitchPipeline(ERHIPipeline::Graphics);
+
+			for (int32 ViewIndex = 0; ViewIndex < Views.Num(); ++ViewIndex)
+			{
+				const FViewInfo& View = Views[ViewIndex];
+				if (View.ViewState != nullptr)
+				{
+					View.ViewState->UpdateExposureCompensationCurveLUT(*TaskData->CmdList, View.FinalPostProcessSettings.AutoExposureBiasCurve);
+				}
+			}
+
+			TaskData->CmdList->FinishRecording();
+		});
+}
+
+void FinishUpdateExposureCompensationCurveLUT(FRHICommandListImmediate& RHICmdList, FUpdateExposureCompensationCurveLUTTaskData* TaskData)
+{
+	if (TaskData->Task.IsValid())
+	{
+		TaskData->Task.Wait();
+		TaskData->Task = {};
+
+		RHICmdList.QueueAsyncCommandListSubmit(TaskData->CmdList);
+		TaskData->CmdList = nullptr;
+	}
+}
+
 FEyeAdaptationParameters GetEyeAdaptationParameters(const FViewInfo& View)
 {
 	const bool bExtendedLuminanceRange = IsExtendLuminanceRangeEnabled();
diff --git a/Engine/Source/Runtime/Renderer/Private/PostProcess/PostProcessEyeAdaptation.h b/Engine/Source/Runtime/Renderer/Private/PostProcess/PostProcessEyeAdaptation.h
--- a/Engine/Source/Runtime/Renderer/Private/PostProcess/PostProcessEyeAdaptation.h
+++ b/Engine/Source/Runtime/Renderer/Private/PostProcess/PostProcessEyeAdaptation.h
@@ -48,7 +48,14 @@ int32 GetAutoExposureIlluminanceDownscaleFactor();
 // Returns the auto exposure method enabled by the view (including CVar override).
 EAutoExposureMethod GetAutoExposureMethod(const FViewInfo& View);
 
-void SetupExposureCompensationCurveLUT(FRHICommandListImmediate& RHICmdList, const FViewInfo& View);
+struct FUpdateExposureCompensationCurveLUTTaskData
+{
+	UE::Tasks::FTask Task;
+	FRHICommandList* CmdList = nullptr;
+};
+
+void BeginUpdateExposureCompensationCurveLUT(TArrayView<FViewInfo> Views, FUpdateExposureCompensationCurveLUTTaskData* TaskData);
+void FinishUpdateExposureCompensationCurveLUT(FRHICommandListImmediate& RHICmdList, FUpdateExposureCompensationCurveLUTTaskData* TaskData);
 
 BEGIN_SHADER_PARAMETER_STRUCT(FEyeAdaptationParameters, )
 	SHADER_PARAMETER(float, ExposureLowPercent)
