Fixed mobile renderer not properly setting the global view uniform buffer for occlusion queries when rendering multiple views.

#rnx

diff --git a/Engine/Source/Runtime/Renderer/Private/MobileShadingRenderer.cpp b/Engine/Source/Runtime/Renderer/Private/MobileShadingRenderer.cpp
--- a/Engine/Source/Runtime/Renderer/Private/MobileShadingRenderer.cpp
+++ b/Engine/Source/Runtime/Renderer/Private/MobileShadingRenderer.cpp
@@ -890,15 +890,20 @@ void FMobileSceneRenderer::RenderFullDepthPrepass(FRDGBuilder& GraphBuilder, TAr
 		AddDitheredStencilClearPass(GraphBuilder, SceneTextures.Depth.Target, ShaderPlatform, false);
 	}
 
-	for (FRenderViewContext& ViewContext : RenderViews)
+	FViewOcclusionQueriesPerView QueriesPerView;
+	bool bDoOcclusionQueries = DoOcclusionQueries() && !bIsSceneCaptureRenderPass;
+	if (bDoOcclusionQueries)
 	{
-		FViewInfo& View = *ViewContext.ViewInfo;
+		AllocateOcclusionTests(Scene, QueriesPerView, Views, VisibleLightInfos);
+		bDoOcclusionQueries = (QueriesPerView.Num() > 0);
+	}
 
-		// Render occlusion at the last view pass only, as they already loop through all views
-		// If this is scene capture render pass, don't render occlusion.
-		bool bDoOcclusionQueries = (ViewContext.bIsLastView && DoOcclusionQueries() && !bIsSceneCaptureRenderPass);
-		if (bDoOcclusionQueries)
+	if (bDoOcclusionQueries)
+	{
+		for (FRenderViewContext& ViewContext : RenderViews)
 		{
+			FViewInfo& View = *ViewContext.ViewInfo;
+
 			auto* PassParameters = GraphBuilder.AllocParameters<FMobileRenderPassParameters>();
 			PassParameters->RenderTargets.DepthStencil = FDepthStencilBinding(SceneTextures.Depth.Target, ERenderTargetLoadAction::ELoad, ERenderTargetLoadAction::ELoad, FExclusiveDepthStencil::DepthRead_StencilRead);
 			PassParameters->View = View.GetShaderParameters();
@@ -909,9 +914,10 @@ void FMobileSceneRenderer::RenderFullDepthPrepass(FRDGBuilder& GraphBuilder, TAr
 				RDG_EVENT_NAME("RenderOcclusion"),
 				PassParameters,
 				ERDGPassFlags::Raster | ERDGPassFlags::NeverCull,
-				[this](FRHICommandList& RHICmdList)
+				[this, ViewContext, ViewQueries = bDoOcclusionQueries ? MoveTemp(QueriesPerView[ViewContext.ViewIndex]) : FViewOcclusionQueries()](FRHICommandList& RHICmdList)
 				{
-					RenderOcclusion(RHICmdList);
+					FViewInfo& View = *ViewContext.ViewInfo;
+					RenderOcclusion(RHICmdList, View, ViewQueries);
 				});
 		}
 	}
@@ -1703,6 +1709,14 @@ void FMobileSceneRenderer::RenderForward(FRDGBuilder& GraphBuilder, FRDGTextureR
 	FRenderViewContextArray RenderViews;
 	GetRenderViews(Views, RenderViews);
 
+	bool bDoOcclusionQueries = !bIsFullDepthPrepassEnabled && DoOcclusionQueries();
+	FViewOcclusionQueriesPerView QueriesPerView;
+	if (bDoOcclusionQueries)
+	{
+		AllocateOcclusionTests(Scene, QueriesPerView, Views, VisibleLightInfos);
+		bDoOcclusionQueries = (QueriesPerView.Num() > 0);
+	}
+
 	for (FRenderViewContext& ViewContext : RenderViews)
 	{
 		FViewInfo& View = *ViewContext.ViewInfo;
@@ -1742,14 +1756,16 @@ void FMobileSceneRenderer::RenderForward(FRDGBuilder& GraphBuilder, FRDGTextureR
 		PassParameters->HalfResLocalFogVolumeViewSRV = View.LocalFogVolumeViewData.HalfResLocalFogVolumeViewSRV;
 		PassParameters->HalfResLocalFogVolumeDepthSRV = View.LocalFogVolumeViewData.HalfResLocalFogVolumeDepthSRV;
 
+		FViewOcclusionQueries* ViewQueries = bDoOcclusionQueries ? &QueriesPerView[ViewContext.ViewIndex] : nullptr;
+		
 		// Split if we need to render translucency in a separate render pass
 		if (bRequiresMultiPass)
 		{
-			RenderForwardMultiPass(GraphBuilder, PassParameters, ViewContext, SceneTextures, InstanceCullingManager);
+			RenderForwardMultiPass(GraphBuilder, PassParameters, ViewContext, SceneTextures, InstanceCullingManager, ViewQueries);
 		}
 		else
 		{
-			RenderForwardSinglePass(GraphBuilder, PassParameters, ViewContext, SceneTextures, InstanceCullingManager);
+			RenderForwardSinglePass(GraphBuilder, PassParameters, ViewContext, SceneTextures, InstanceCullingManager, ViewQueries);
 		}
 	}
 	
@@ -1759,7 +1775,7 @@ void FMobileSceneRenderer::RenderForward(FRDGBuilder& GraphBuilder, FRDGTextureR
 	}
 }
 
-void FMobileSceneRenderer::RenderForwardSinglePass(FRDGBuilder& GraphBuilder, FMobileRenderPassParameters* PassParameters, FRenderViewContext& ViewContext, FSceneTextures& SceneTextures, FInstanceCullingManager& InstanceCullingManager)
+void FMobileSceneRenderer::RenderForwardSinglePass(FRDGBuilder& GraphBuilder, FMobileRenderPassParameters* PassParameters, FRenderViewContext& ViewContext, FSceneTextures& SceneTextures, FInstanceCullingManager& InstanceCullingManager, FViewOcclusionQueries* ViewOcclusionQueries)
 {
 	struct FForwardSinglePassParameterCollection
 	{
@@ -1793,7 +1809,7 @@ void FMobileSceneRenderer::RenderForwardSinglePass(FRDGBuilder& GraphBuilder, FM
 	}
 		
 	PassParameters->RenderTargets.SubpassHint = bTonemapSubpassInline ? ESubpassHint::CustomResolveSubpass : ESubpassHint::DepthReadSubpass;
-	const bool bDoOcclusionQueries = (!bIsFullDepthPrepassEnabled && ViewContext.bIsLastView && DoOcclusionQueries());
+	const bool bDoOcclusionQueries = (ViewOcclusionQueries != nullptr);
 	const int32 NumOcclusionQueries = bDoOcclusionQueries ? ComputeNumOcclusionQueriesToBatch() : 0u;
 	const bool bVulkanAdrenoOcclusionMode = bAdrenoOcclusionMode && IsVulkanPlatform(ShaderPlatform);
 	const bool bDoOcclusionInMainPass = bDoOcclusionQueries && !bVulkanAdrenoOcclusionMode;
@@ -1804,7 +1820,7 @@ void FMobileSceneRenderer::RenderForwardSinglePass(FRDGBuilder& GraphBuilder, FM
 		PassParameters,
 		// the second view pass should not be merged with the first view pass on mobile since the subpass would not work properly.
 		ERDGPassFlags::Raster | ERDGPassFlags::NeverMerge,
-		[this, PassParameters, ParameterCollection, ViewContext, bDoOcclusionInMainPass, &SceneTextures](FRHICommandList& RHICmdList)
+		[this, PassParameters, ParameterCollection, ViewContext, bDoOcclusionInMainPass, &SceneTextures, ViewQueries = bDoOcclusionInMainPass ? MoveTemp(*ViewOcclusionQueries) : FViewOcclusionQueries()](FRHICommandList& RHICmdList)
 	{
 		FViewInfo& View = *ViewContext.ViewInfo;
 			
@@ -1847,7 +1863,7 @@ void FMobileSceneRenderer::RenderForwardSinglePass(FRDGBuilder& GraphBuilder, FM
 				// flush
 				RHICmdList.SubmitCommandsHint();
 			}
-			RenderOcclusion(RHICmdList);
+			RenderOcclusion(RHICmdList, View, ViewQueries);
 		}
 
 		// Pre-tonemap before MSAA resolve (iOS only)
@@ -1883,14 +1899,15 @@ void FMobileSceneRenderer::RenderForwardSinglePass(FRDGBuilder& GraphBuilder, FM
 			// The occlusion pass needs to be unique to be optimized properly by the driver, don't merge it.
 			// This pass has no observable outputs on the RenderGraph, so it needs to be marked as NeverCull.
 			ERDGPassFlags::Raster | ERDGPassFlags::NeverMerge | ERDGPassFlags::NeverCull,
-			[this](FRHICommandListImmediate& RHICmdList)
+			[this, ViewContext, ViewQueries = MoveTemp(*ViewOcclusionQueries)](FRHICommandListImmediate& RHICmdList)
 		{
-			RenderOcclusion(RHICmdList);
+			FViewInfo& View = *ViewContext.ViewInfo;
+			RenderOcclusion(RHICmdList, View, ViewQueries);
 		});
 	}
 }
 
-void FMobileSceneRenderer::RenderForwardMultiPass(FRDGBuilder& GraphBuilder, FMobileRenderPassParameters* PassParameters, FRenderViewContext& ViewContext, FSceneTextures& SceneTextures, FInstanceCullingManager& InstanceCullingManager)
+void FMobileSceneRenderer::RenderForwardMultiPass(FRDGBuilder& GraphBuilder, FMobileRenderPassParameters* PassParameters, FRenderViewContext& ViewContext, FSceneTextures& SceneTextures, FInstanceCullingManager& InstanceCullingManager, FViewOcclusionQueries* ViewOcclusionQueries)
 {
 	struct FForwardFirstPassParameterCollection
 	{
@@ -1975,7 +1992,7 @@ void FMobileSceneRenderer::RenderForwardMultiPass(FRDGBuilder& GraphBuilder, FMo
 	SecondPassParameters->RenderTargets.DepthStencil.SetStencilLoadAction(ERenderTargetLoadAction::ELoad);
 	SecondPassParameters->RenderTargets.DepthStencil.SetDepthStencilAccess(ExclusiveDepthStencil);
 	
-	const bool bDoOcclusionQueries = (!bIsFullDepthPrepassEnabled && ViewContext.bIsLastView && DoOcclusionQueries());
+	const bool bDoOcclusionQueries = (ViewOcclusionQueries != nullptr);
 	SecondPassParameters->RenderTargets.NumOcclusionQueries = bDoOcclusionQueries ? ComputeNumOcclusionQueriesToBatch() : 0u;
 
 	if (Scene->GPUScene.IsEnabled())
@@ -1988,7 +2005,7 @@ void FMobileSceneRenderer::RenderForwardMultiPass(FRDGBuilder& GraphBuilder, FMo
 		RDG_EVENT_NAME("DecalsAndTranslucency"),
 		SecondPassParameters,
 		ERDGPassFlags::Raster,
-		[this, SecondParameterCollection, ViewContext, bDoOcclusionQueries, &SceneTextures](FRHICommandList& RHICmdList)
+		[this, SecondParameterCollection, ViewContext, bDoOcclusionQueries, &SceneTextures, ViewQueries = (bDoOcclusionQueries) ? MoveTemp(*ViewOcclusionQueries) : FViewOcclusionQueries()](FRHICommandList& RHICmdList)
 	{
 		FViewInfo& View = *ViewContext.ViewInfo;
 		auto SecondPassParameters = &SecondParameterCollection->PassParameters;
@@ -2003,7 +2020,7 @@ void FMobileSceneRenderer::RenderForwardMultiPass(FRDGBuilder& GraphBuilder, FMo
 		if (bDoOcclusionQueries)
 		{
 			// Issue occlusion queries
-			RenderOcclusion(RHICmdList);
+			RenderOcclusion(RHICmdList, View, ViewQueries);
 		}
 
 		// Pre-tonemap before MSAA resolve (iOS only)
@@ -2158,6 +2175,14 @@ void FMobileSceneRenderer::RenderDeferredSinglePass(FRDGBuilder& GraphBuilder, F
 	FRenderViewContextArray RenderViews;
 	GetRenderViews(Views, RenderViews);
 
+	bool bDoOcclusionQueries = (!bIsFullDepthPrepassEnabled && DoOcclusionQueries());
+	FViewOcclusionQueriesPerView QueriesPerView;
+	if (bDoOcclusionQueries)
+	{
+		AllocateOcclusionTests(Scene, QueriesPerView, Views, VisibleLightInfos);
+		bDoOcclusionQueries = QueriesPerView.Num() > 0;
+	}
+
 	for (FRenderViewContext& ViewContext : RenderViews)
 	{
 		FViewInfo& View = *ViewContext.ViewInfo;
@@ -2209,8 +2234,7 @@ void FMobileSceneRenderer::RenderDeferredSinglePass(FRDGBuilder& GraphBuilder, F
 		PassParameters->RenderTargets = BasePassRenderTargets;
 		PassParameters->RenderTargets.SubpassHint = ESubpassHint::DeferredShadingSubpass;
 		const EMobileSSRQuality MobileSSRQuality = ActiveMobileSSRQuality(View, bShouldRenderVelocities);
-		const bool bDoOcclusionQueires = (!bIsFullDepthPrepassEnabled && ViewContext.bIsLastView && DoOcclusionQueries());
-		PassParameters->RenderTargets.NumOcclusionQueries = bDoOcclusionQueires ? ComputeNumOcclusionQueriesToBatch() : 0u;
+		PassParameters->RenderTargets.NumOcclusionQueries = bDoOcclusionQueries ? ComputeNumOcclusionQueriesToBatch() : 0u;
 
 		if (Scene->GPUScene.IsEnabled())
 		{
@@ -2230,7 +2254,14 @@ void FMobileSceneRenderer::RenderDeferredSinglePass(FRDGBuilder& GraphBuilder, F
 			PassParameters,
 			// the second view pass should not be merged with the first view pass on mobile since the subpass would not work properly.
 			ERDGPassFlags::Raster | ERDGPassFlags::NeverMerge,
-			[this, ParameterCollection, ViewContext, bDoOcclusionQueires, MobileSSRQuality, &SortedLightSet, bUsingPixelLocalStorage](FRHICommandList& RHICmdList)
+			[this, 
+			ParameterCollection, 
+			ViewContext, 
+			MobileSSRQuality, 
+			&SortedLightSet, 
+			bUsingPixelLocalStorage,
+			bDoOcclusionQueries, 
+			ViewQueries = (bDoOcclusionQueries) ? MoveTemp(QueriesPerView[ViewContext.ViewIndex]) : FViewOcclusionQueries()](FRHICommandList& RHICmdList)
 			{
 				FViewInfo& View = *ViewContext.ViewInfo;
 				auto PassParameters = &ParameterCollection->PassParameters;
@@ -2257,10 +2288,10 @@ void FMobileSceneRenderer::RenderDeferredSinglePass(FRDGBuilder& GraphBuilder, F
 				// Draw translucency.
 				RenderTranslucency(RHICmdList, View, Views, StandardTranslucencyPass, StandardTranslucencyMeshPass, &ParameterCollection->TranslucencyInstanceCullingDrawParams);
 
-				if (bDoOcclusionQueires)
+				if (bDoOcclusionQueries)
 				{
 					// Issue occlusion queries
-					RenderOcclusion(RHICmdList);
+					RenderOcclusion(RHICmdList, View, ViewQueries);
 				}
 			});
 	}
@@ -2275,6 +2306,14 @@ void FMobileSceneRenderer::RenderDeferredMultiPass(FRDGBuilder& GraphBuilder, FS
 	FRenderViewContextArray RenderViews;
 	GetRenderViews(Views, RenderViews);
 
+	bool bDoOcclusionQueries = (!bIsFullDepthPrepassEnabled && DoOcclusionQueries());
+	FViewOcclusionQueriesPerView QueriesPerView;
+	if (bDoOcclusionQueries)
+	{
+		AllocateOcclusionTests(Scene, QueriesPerView, Views, VisibleLightInfos);
+		bDoOcclusionQueries = QueriesPerView.Num() > 0;
+	}
+
 	for (FRenderViewContext& ViewContext : RenderViews)
 	{
 		FViewInfo& View = *ViewContext.ViewInfo;
@@ -2313,8 +2352,7 @@ void FMobileSceneRenderer::RenderDeferredMultiPass(FRDGBuilder& GraphBuilder, FS
 			PassParameters->RenderTargets.DepthStencil.SetDepthStencilAccess(bIsFullDepthPrepassEnabled ? FExclusiveDepthStencil::DepthRead_StencilWrite : FExclusiveDepthStencil::DepthWrite_StencilWrite);
 		}
 
-		const bool bDoOcclusionQueires = (!bIsFullDepthPrepassEnabled && ViewContext.bIsLastView && DoOcclusionQueries());
-		PassParameters->RenderTargets.NumOcclusionQueries = bDoOcclusionQueires ? ComputeNumOcclusionQueriesToBatch() : 0u;
+		PassParameters->RenderTargets.NumOcclusionQueries = bDoOcclusionQueries ? ComputeNumOcclusionQueriesToBatch() : 0u;
 
 		if (Scene->GPUScene.IsEnabled())
 		{
@@ -2335,7 +2373,13 @@ void FMobileSceneRenderer::RenderDeferredMultiPass(FRDGBuilder& GraphBuilder, FS
 			RDG_EVENT_NAME("BasePass"),
 			PassParameters,
 			ERDGPassFlags::Raster,
-			[this, ParameterCollection, ViewContext, &SceneTextures, bDoOcclusionQueires](FRHICommandList& RHICmdList)
+			[this, 
+			ParameterCollection, 
+			ViewContext, 
+			&SceneTextures, 
+			bDoOcclusionQueries, 
+			ViewQueries = (bDoOcclusionQueries) ? MoveTemp(QueriesPerView[ViewContext.ViewIndex]) : FViewOcclusionQueries()]
+			(FRHICommandList& RHICmdList)
 			{
 				auto PassParameters = &ParameterCollection->PassParameters;
 				FViewInfo& View = *ViewContext.ViewInfo;
@@ -2348,10 +2392,10 @@ void FMobileSceneRenderer::RenderDeferredMultiPass(FRDGBuilder& GraphBuilder, FS
 
 				PostRenderBasePass(RHICmdList, View);
 
-				if (bDoOcclusionQueires)
+				if (bDoOcclusionQueries)
 				{
 					// Issue occlusion queries
-					RenderOcclusion(RHICmdList);
+					RenderOcclusion(RHICmdList, View, ViewQueries);
 				}
 
 				if (bIsFullDepthPrepassEnabled)
diff --git a/Engine/Source/Runtime/Renderer/Private/SceneOcclusion.cpp b/Engine/Source/Runtime/Renderer/Private/SceneOcclusion.cpp
--- a/Engine/Source/Runtime/Renderer/Private/SceneOcclusion.cpp
+++ b/Engine/Source/Runtime/Renderer/Private/SceneOcclusion.cpp
@@ -1069,7 +1069,7 @@ static void TrimAllOcclusionHistory(TArrayView<FViewInfo> Views)
 	}
 }
 
-static void AllocateOcclusionTests(FViewOcclusionQueriesPerView& QueriesPerView, const FScene* Scene, TArrayView<const FVisibleLightInfo> VisibleLightInfos, TArrayView<FViewInfo> Views)
+void AllocateOcclusionTests(const FScene* Scene, FViewOcclusionQueriesPerView& QueriesPerView, TArrayView<FViewInfo> Views, TArrayView<const FVisibleLightInfo> VisibleLightInfos)
 {
 	SCOPED_NAMED_EVENT(FSceneRenderer_AllocateOcclusionTests, FColor::Emerald);
 
@@ -1226,190 +1226,219 @@ static void AllocateOcclusionTests(FViewOcclusionQueriesPerView& QueriesPerView,
 
 static void BeginOcclusionTests(
 	FRHICommandList& RHICmdList,
-	TArrayView<FViewInfo> Views,
-	ERHIFeatureLevel::Type FeatureLevel,
-	const FViewOcclusionQueriesPerView& QueriesPerView,
-	uint32 DownsampleFactor)
+	FViewInfo& View,
+	const FViewOcclusionQueries& ViewQuery,
+	ERHIFeatureLevel::Type FeatureLevel, 
+	uint32 DownsampleFactor,
+	FGraphicsPipelineStateInitializer* InGraphicsPSOInit = nullptr)
 {
 	check(RHICmdList.IsInsideRenderPass());
-	check(QueriesPerView.Num() == Views.Num());
 
 	SCOPE_CYCLE_COUNTER(STAT_BeginOcclusionTestsTime);
 	SCOPED_DRAW_EVENT(RHICmdList, BeginOcclusionTests);
 
-	FGraphicsPipelineStateInitializer GraphicsPSOInit;
-	RHICmdList.ApplyCachedRenderTargets(GraphicsPSOInit);
-	GraphicsPSOInit.PrimitiveType = PT_TriangleList;
-	GraphicsPSOInit.BlendState = TStaticBlendStateWriteMask<CW_NONE, CW_NONE, CW_NONE, CW_NONE, CW_NONE, CW_NONE, CW_NONE, CW_NONE>::GetRHI();
-	// Depth tests, no depth writes, no color writes, opaque
-	GraphicsPSOInit.DepthStencilState = TStaticDepthStencilState<false, CF_DepthNearOrEqual>::GetRHI();
-	GraphicsPSOInit.BoundShaderState.VertexDeclarationRHI = GetVertexDeclarationFVector3();
+	FSceneViewState* ViewState = (FSceneViewState*)View.State;
+	SCOPED_GPU_MASK(RHICmdList, View.GPUMask);
 
-	for (int32 ViewIndex = 0; ViewIndex < Views.Num(); ViewIndex++)
+	FGraphicsPipelineStateInitializer StackGraphicsPSOInit;
+	if (InGraphicsPSOInit == nullptr)
 	{
-		// Don't draw occlusion queries for both eyes in multiview
-		FViewInfo& View = Views[ViewIndex];
-		if (View.bIsSinglePassStereo && IStereoRendering::IsASecondaryPass(View.StereoPass))
-		{
-			continue;
-		}
-
-		SCOPED_DRAW_EVENTF(RHICmdList, ViewOcclusionTests, TEXT("ViewOcclusionTests %d"), ViewIndex);
+		RHICmdList.ApplyCachedRenderTargets(StackGraphicsPSOInit);
+		StackGraphicsPSOInit.PrimitiveType = PT_TriangleList;
+		StackGraphicsPSOInit.BlendState = TStaticBlendStateWriteMask<CW_NONE, CW_NONE, CW_NONE, CW_NONE, CW_NONE, CW_NONE, CW_NONE, CW_NONE>::GetRHI();
+		// Depth tests, no depth writes, no color writes, opaque
+		StackGraphicsPSOInit.DepthStencilState = TStaticDepthStencilState<false, CF_DepthNearOrEqual>::GetRHI();
+		StackGraphicsPSOInit.BoundShaderState.VertexDeclarationRHI = GetVertexDeclarationFVector3();
+	}
+	FGraphicsPipelineStateInitializer& GraphicsPSOInit = (InGraphicsPSOInit != nullptr) ? *InGraphicsPSOInit : StackGraphicsPSOInit;
 
-		const FViewOcclusionQueries& ViewQuery = QueriesPerView[ViewIndex];
-		FSceneViewState* ViewState = (FSceneViewState*)View.State;
-		SCOPED_GPU_MASK(RHICmdList, View.GPUMask);
 
-		// We only need to render the front-faces of the culling geometry (this halves the amount of pixels we touch)
-		GraphicsPSOInit.RasterizerState = View.bReverseCulling ? TStaticRasterizerState<FM_Solid, CM_CCW>::GetRHI() : TStaticRasterizerState<FM_Solid, CM_CW>::GetRHI();
+	// We only need to render the front-faces of the culling geometry (this halves the amount of pixels we touch)
+	GraphicsPSOInit.RasterizerState = View.bReverseCulling ? TStaticRasterizerState<FM_Solid, CM_CCW>::GetRHI() : TStaticRasterizerState<FM_Solid, CM_CW>::GetRHI();
 
-		// Set stereo viewports while accounting for the downsample factor
-		if (View.IsInstancedStereoPass())
+	// Set stereo viewports while accounting for the downsample factor
+	if (View.IsInstancedStereoPass())
+	{
+		if (View.bIsMultiViewportEnabled)
 		{
-			if (View.bIsMultiViewportEnabled)
-			{
-				const FIntRect ViewRect = GetDownscaledRect(View.ViewRect, DownsampleFactor);
-				const FIntRect SecondaryViewRect = GetDownscaledRect(View.GetInstancedView()->ViewRect, DownsampleFactor);
-				RHICmdList.SetStereoViewport(ViewRect.Min.X, SecondaryViewRect.Min.X, ViewRect.Min.Y, SecondaryViewRect.Min.Y, 0.0f, ViewRect.Max.X, SecondaryViewRect.Max.X, ViewRect.Max.Y, SecondaryViewRect.Max.Y, 1.0f);
-			}
-			else
-			{
-				const FIntRect ViewRect = GetDownscaledRect(View.ViewRectWithSecondaryViews, DownsampleFactor);
-				RHICmdList.SetViewport(ViewRect.Min.X, ViewRect.Min.Y, 0.0f, ViewRect.Max.X, ViewRect.Max.Y, 1.0f);
-			}
+			const FIntRect ViewRect = GetDownscaledRect(View.ViewRect, DownsampleFactor);
+			const FIntRect SecondaryViewRect = GetDownscaledRect(View.GetInstancedView()->ViewRect, DownsampleFactor);
+			RHICmdList.SetStereoViewport(ViewRect.Min.X, SecondaryViewRect.Min.X, ViewRect.Min.Y, SecondaryViewRect.Min.Y, 0.0f, ViewRect.Max.X, SecondaryViewRect.Max.X, ViewRect.Max.Y, SecondaryViewRect.Max.Y, 1.0f);
 		}
 		else
 		{
-			const FIntRect ViewRect = GetDownscaledRect(View.ViewRect, DownsampleFactor);
+			const FIntRect ViewRect = GetDownscaledRect(View.ViewRectWithSecondaryViews, DownsampleFactor);
 			RHICmdList.SetViewport(ViewRect.Min.X, ViewRect.Min.Y, 0.0f, ViewRect.Max.X, ViewRect.Max.Y, 1.0f);
 		}
+	}
+	else
+	{
+		const FIntRect ViewRect = GetDownscaledRect(View.ViewRect, DownsampleFactor);
+		RHICmdList.SetViewport(ViewRect.Min.X, ViewRect.Min.Y, 0.0f, ViewRect.Max.X, ViewRect.Max.Y, 1.0f);
+	}
 
-		// Lookup the vertex shader.
-		TShaderMapRef<FOcclusionQueryVS> VertexShader(View.ShaderMap);
-		GraphicsPSOInit.BoundShaderState.VertexShaderRHI = VertexShader.GetVertexShader();
+	// Lookup the vertex shader.
+	TShaderMapRef<FOcclusionQueryVS> VertexShader(View.ShaderMap);
+	GraphicsPSOInit.BoundShaderState.VertexShaderRHI = VertexShader.GetVertexShader();
 
-		if (View.Family->EngineShowFlags.OcclusionMeshes)
-		{
-			TShaderMapRef<FOcclusionQueryPS> PixelShader(View.ShaderMap);
-			GraphicsPSOInit.BoundShaderState.PixelShaderRHI = PixelShader.GetPixelShader();
-			GraphicsPSOInit.BlendState = TStaticBlendState<CW_RGBA>::GetRHI();
-		}
+	if (View.Family->EngineShowFlags.OcclusionMeshes)
+	{
+		TShaderMapRef<FOcclusionQueryPS> PixelShader(View.ShaderMap);
+		GraphicsPSOInit.BoundShaderState.PixelShaderRHI = PixelShader.GetPixelShader();
+		GraphicsPSOInit.BlendState = TStaticBlendState<CW_RGBA>::GetRHI();
+	}
 
-		SetGraphicsPipelineState(RHICmdList, GraphicsPSOInit, 0);
+	SetGraphicsPipelineState(RHICmdList, GraphicsPSOInit, 0);
 
-		if (FeatureLevel > ERHIFeatureLevel::ES3_1)
+	if (FeatureLevel > ERHIFeatureLevel::ES3_1)
+	{
+		SCOPED_DRAW_EVENT(RHICmdList, ShadowFrustumQueries);
+		for (int i = 0; i < ViewQuery.LocalLightQueries.Num(); i++)
 		{
-			SCOPED_DRAW_EVENT(RHICmdList, ShadowFrustumQueries);
-			for (int i = 0; i < ViewQuery.LocalLightQueries.Num(); i++)
-			{
-				ExecutePointLightShadowOcclusionQuery(RHICmdList, View, *ViewQuery.LocalLightQueryInfos[i], VertexShader, ViewQuery.LocalLightQueries[i]);
-			}
+			ExecutePointLightShadowOcclusionQuery(RHICmdList, View, *ViewQuery.LocalLightQueryInfos[i], VertexShader, ViewQuery.LocalLightQueries[i]);
 		}
+	}
 
-		uint32 NumVertices = ViewQuery.CSMQueries.Num() * 6 // Plane 
-			+ ViewQuery.ShadowQueries.Num() * 8 // Cube
-			+ ViewQuery.ReflectionQueries.Num() * 8; // Cube
+	uint32 NumVertices = ViewQuery.CSMQueries.Num() * 6 // Plane 
+		+ ViewQuery.ShadowQueries.Num() * 8 // Cube
+		+ ViewQuery.ReflectionQueries.Num() * 8; // Cube
 
-		if (NumVertices > 0)
-		{
-			uint32 BaseVertexOffset = 0;
+	if (NumVertices > 0)
+	{
+		uint32 BaseVertexOffset = 0;
 
-			const FRHIBufferCreateDesc CreateDesc =
-				FRHIBufferCreateDesc::CreateVertex<FVector3f>(TEXT("ViewOcclusionTests"), NumVertices)
-				.AddUsage(EBufferUsageFlags::Volatile)
-				.SetInitialState(ERHIAccess::VertexOrIndexBuffer)
-				.SetInitActionInitializer();
+		const FRHIBufferCreateDesc CreateDesc =
+			FRHIBufferCreateDesc::CreateVertex<FVector3f>(TEXT("ViewOcclusionTests"), NumVertices)
+			.AddUsage(EBufferUsageFlags::Volatile)
+			.SetInitialState(ERHIAccess::VertexOrIndexBuffer)
+			.SetInitActionInitializer();
 
-			TRHIBufferInitializer<FVector3f> Initializer = RHICmdList.CreateBufferInitializer(CreateDesc);
+		TRHIBufferInitializer<FVector3f> Initializer = RHICmdList.CreateBufferInitializer(CreateDesc);
 
+		{
+			for (FProjectedShadowInfo const* Query : ViewQuery.CSMQueryInfos)
 			{
-				for (FProjectedShadowInfo const* Query : ViewQuery.CSMQueryInfos)
-				{
-					PrepareDirectionalLightShadowOcclusionQuery(BaseVertexOffset, Initializer, View, *Query);
-					checkSlow(BaseVertexOffset <= NumVertices);
-				}
-
-				for (FProjectedShadowInfo const* Query : ViewQuery.ShadowQueryInfos)
-				{
-					PrepareProjectedShadowOcclusionQuery(BaseVertexOffset, Initializer, View, *Query);
-					checkSlow(BaseVertexOffset <= NumVertices);
-				}
-
-				for (FPlanarReflectionSceneProxy const* Query : ViewQuery.ReflectionQueryInfos)
-				{
-					PreparePlanarReflectionOcclusionQuery(BaseVertexOffset, Initializer, View, Query);
-					checkSlow(BaseVertexOffset <= NumVertices);
-				}
+				PrepareDirectionalLightShadowOcclusionQuery(BaseVertexOffset, Initializer, View, *Query);
+				checkSlow(BaseVertexOffset <= NumVertices);
 			}
 
-			FBufferRHIRef VertexBufferRHI = Initializer.Finalize();
+			for (FProjectedShadowInfo const* Query : ViewQuery.ShadowQueryInfos)
+			{
+				PrepareProjectedShadowOcclusionQuery(BaseVertexOffset, Initializer, View, *Query);
+				checkSlow(BaseVertexOffset <= NumVertices);
+			}
 
+			for (FPlanarReflectionSceneProxy const* Query : ViewQuery.ReflectionQueryInfos)
 			{
-				SCOPED_DRAW_EVENT(RHICmdList, ShadowFrustumQueries);
-				SetShaderParametersLegacyVS(RHICmdList, VertexShader, View);
-				RHICmdList.SetStreamSource(0, VertexBufferRHI, 0);
-				BaseVertexOffset = 0;
+				PreparePlanarReflectionOcclusionQuery(BaseVertexOffset, Initializer, View, Query);
+				checkSlow(BaseVertexOffset <= NumVertices);
+			}
+		}
 
-				for (FRHIRenderQuery* const& Query : ViewQuery.CSMQueries)
-				{
-					ExecuteDirectionalLightShadowOcclusionQuery(RHICmdList, BaseVertexOffset, Query);
-					checkSlow(BaseVertexOffset <= NumVertices);
-				}
+		FBufferRHIRef VertexBufferRHI = Initializer.Finalize();
 
-				for (FRHIRenderQuery* const& Query : ViewQuery.ShadowQueries)
-				{
-					ExecuteProjectedShadowOcclusionQuery(RHICmdList, BaseVertexOffset, Query);
-					checkSlow(BaseVertexOffset <= NumVertices);
-				}
-			}
+		{
+			SCOPED_DRAW_EVENT(RHICmdList, ShadowFrustumQueries);
+			SetShaderParametersLegacyVS(RHICmdList, VertexShader, View);
+			RHICmdList.SetStreamSource(0, VertexBufferRHI, 0);
+			BaseVertexOffset = 0;
 
-			if (FeatureLevel > ERHIFeatureLevel::ES3_1)
+			for (FRHIRenderQuery* const& Query : ViewQuery.CSMQueries)
 			{
-				SCOPED_DRAW_EVENT(RHICmdList, PlanarReflectionQueries);
-				for (FRHIRenderQuery* const& Query : ViewQuery.ReflectionQueries)
-				{
-					ExecutePlanarReflectionOcclusionQuery(RHICmdList, BaseVertexOffset, Query);
-					check(BaseVertexOffset <= NumVertices);
-				}
+				ExecuteDirectionalLightShadowOcclusionQuery(RHICmdList, BaseVertexOffset, Query);
+				checkSlow(BaseVertexOffset <= NumVertices);
 			}
 
-			VertexBufferRHI.SafeRelease();
+			for (FRHIRenderQuery* const& Query : ViewQuery.ShadowQueries)
+			{
+				ExecuteProjectedShadowOcclusionQuery(RHICmdList, BaseVertexOffset, Query);
+				checkSlow(BaseVertexOffset <= NumVertices);
+			}
 		}
 
-		if (ViewQuery.bFlushQueries)
+		if (FeatureLevel > ERHIFeatureLevel::ES3_1)
 		{
-			SetShaderParametersLegacyVS(RHICmdList, VertexShader, View);
+			SCOPED_DRAW_EVENT(RHICmdList, PlanarReflectionQueries);
+			for (FRHIRenderQuery* const& Query : ViewQuery.ReflectionQueries)
+			{
+				ExecutePlanarReflectionOcclusionQuery(RHICmdList, BaseVertexOffset, Query);
+				check(BaseVertexOffset <= NumVertices);
+			}
+		}
+
+		VertexBufferRHI.SafeRelease();
+	}
 
-			if (GOcclusionQueryDispatchOrder == 0)
+	if (ViewQuery.bFlushQueries)
+	{
+		SetShaderParametersLegacyVS(RHICmdList, VertexShader, View);
+
+		if (GOcclusionQueryDispatchOrder == 0)
+		{
 			{
-				{
-					SCOPED_DRAW_EVENT(RHICmdList, GroupedQueries);
-					View.GroupedOcclusionQueries.Flush(RHICmdList);
-				}
-				{
-					SCOPED_DRAW_EVENT(RHICmdList, IndividualQueries);
-					View.IndividualOcclusionQueries.Flush(RHICmdList);
-				}
+				SCOPED_DRAW_EVENT(RHICmdList, GroupedQueries);
+				View.GroupedOcclusionQueries.Flush(RHICmdList);
 			}
-			else
 			{
-				{
-					SCOPED_DRAW_EVENT(RHICmdList, IndividualQueries);
-					View.IndividualOcclusionQueries.Flush(RHICmdList);
-				}
-				{
-					SCOPED_DRAW_EVENT(RHICmdList, GroupedQueries);
-					View.GroupedOcclusionQueries.Flush(RHICmdList);
-				}
+				SCOPED_DRAW_EVENT(RHICmdList, IndividualQueries);
+				View.IndividualOcclusionQueries.Flush(RHICmdList);
 			}
-
-			if (FeatureLevel == ERHIFeatureLevel::ES3_1 && 
-				View.ViewState && View.ViewState->OcclusionFeedback.IsInitialized())
+		}
+		else
+		{
 			{
 				SCOPED_DRAW_EVENT(RHICmdList, IndividualQueries);
-				View.ViewState->OcclusionFeedback.SubmitOcclusionDraws(RHICmdList, View);
+				View.IndividualOcclusionQueries.Flush(RHICmdList);
+			}
+			{
+				SCOPED_DRAW_EVENT(RHICmdList, GroupedQueries);
+				View.GroupedOcclusionQueries.Flush(RHICmdList);
 			}
 		}
+
+		if (FeatureLevel == ERHIFeatureLevel::ES3_1 &&
+			View.ViewState && View.ViewState->OcclusionFeedback.IsInitialized())
+		{
+			SCOPED_DRAW_EVENT(RHICmdList, IndividualQueries);
+			View.ViewState->OcclusionFeedback.SubmitOcclusionDraws(RHICmdList, View);
+		}
+	}
+}
+
+static void BeginOcclusionTests(
+	FRHICommandList& RHICmdList,
+	TArrayView<FViewInfo> Views,
+	ERHIFeatureLevel::Type FeatureLevel,
+	const FViewOcclusionQueriesPerView& QueriesPerView,
+	uint32 DownsampleFactor)
+{
+	check(RHICmdList.IsInsideRenderPass());
+	check(QueriesPerView.Num() == Views.Num());
+
+	SCOPE_CYCLE_COUNTER(STAT_BeginOcclusionTestsTime);
+	SCOPED_DRAW_EVENT(RHICmdList, BeginOcclusionTests);
+
+	FGraphicsPipelineStateInitializer GraphicsPSOInit;
+	RHICmdList.ApplyCachedRenderTargets(GraphicsPSOInit);
+	GraphicsPSOInit.PrimitiveType = PT_TriangleList;
+	GraphicsPSOInit.BlendState = TStaticBlendStateWriteMask<CW_NONE, CW_NONE, CW_NONE, CW_NONE, CW_NONE, CW_NONE, CW_NONE, CW_NONE>::GetRHI();
+	// Depth tests, no depth writes, no color writes, opaque
+	GraphicsPSOInit.DepthStencilState = TStaticDepthStencilState<false, CF_DepthNearOrEqual>::GetRHI();
+	GraphicsPSOInit.BoundShaderState.VertexDeclarationRHI = GetVertexDeclarationFVector3();
+
+	for (int32 ViewIndex = 0; ViewIndex < Views.Num(); ViewIndex++)
+	{
+		// Don't draw occlusion queries for both eyes in multiview
+		FViewInfo& View = Views[ViewIndex];
+		if (View.bIsSinglePassStereo && IStereoRendering::IsASecondaryPass(View.StereoPass))
+		{
+			continue;
+		}
+
+		SCOPED_DRAW_EVENTF(RHICmdList, ViewOcclusionTests, TEXT("ViewOcclusionTests %d"), ViewIndex);
+		const FViewOcclusionQueries& ViewQuery = QueriesPerView[ViewIndex];
+
+		BeginOcclusionTests(RHICmdList, View, ViewQuery, FeatureLevel, DownsampleFactor, &GraphicsPSOInit);
 	}
 }
 
@@ -1455,7 +1484,7 @@ void FDeferredShadingSceneRenderer::RenderOcclusion(
 		GraphBuilder.AddSetupTask([this, PassParameters, &QueriesPerView]
 		{
 			// Issue occlusion queries. This is done after the downsampled depth buffer is created so that it can be used for issuing queries.
-			AllocateOcclusionTests(QueriesPerView, Scene, VisibleLightInfos, Views);
+			AllocateOcclusionTests(Scene, QueriesPerView, Views, VisibleLightInfos);
 
 			int32 NumQueriesForBatch = 0;
 
@@ -1520,23 +1549,10 @@ void FDeferredShadingSceneRenderer::RenderOcclusion(
 	}
 }
 
-void FMobileSceneRenderer::RenderOcclusion(FRHICommandList& RHICmdList)
+void FMobileSceneRenderer::RenderOcclusion(FRHICommandList& RHICmdList, FViewInfo& View, const FViewOcclusionQueries& ViewOcclusionQueries)
 {
-	if (!DoOcclusionQueries())
-	{
-		return;
-	}
-
-	{
-		SCOPED_NAMED_EVENT(FMobileSceneRenderer_BeginOcclusionTests, FColor::Emerald);
-		FViewOcclusionQueriesPerView QueriesPerView;
-		AllocateOcclusionTests(QueriesPerView, Scene, VisibleLightInfos, Views);
-
-		if (QueriesPerView.Num())
-		{
-			BeginOcclusionTests(RHICmdList, Views, FeatureLevel, QueriesPerView, 1.0f);
-		}
-	}
+	SCOPED_NAMED_EVENT(FMobileSceneRenderer_BeginOcclusionTests, FColor::Emerald);
+	BeginOcclusionTests(RHICmdList, View, ViewOcclusionQueries, View.GetFeatureLevel(), 1.0f);
 }
 
 DECLARE_CYCLE_STAT(TEXT("OcclusionSubmittedFence Dispatch"), STAT_OcclusionSubmittedFence_Dispatch, STATGROUP_SceneRendering);
diff --git a/Engine/Source/Runtime/Renderer/Private/SceneOcclusion.h b/Engine/Source/Runtime/Renderer/Private/SceneOcclusion.h
--- a/Engine/Source/Runtime/Renderer/Private/SceneOcclusion.h
+++ b/Engine/Source/Runtime/Renderer/Private/SceneOcclusion.h
@@ -39,6 +39,8 @@ struct FViewOcclusionQueries
 
 using FViewOcclusionQueriesPerView = TArray<FViewOcclusionQueries, TInlineAllocator<1, SceneRenderingAllocator>>;
 
+extern void AllocateOcclusionTests(const FScene* Scene, FViewOcclusionQueriesPerView& QueriesPerView, TArrayView<class FViewInfo> Views, TArrayView<const class FVisibleLightInfo> VisibleLightInfos);
+
 /**
 * A vertex shader for rendering a texture on a simple element.
 */
diff --git a/Engine/Source/Runtime/Renderer/Private/SceneRendering.h b/Engine/Source/Runtime/Renderer/Private/SceneRendering.h
--- a/Engine/Source/Runtime/Renderer/Private/SceneRendering.h
+++ b/Engine/Source/Runtime/Renderer/Private/SceneRendering.h
@@ -2901,8 +2901,8 @@ protected:
 	void RenderModulatedShadowProjections(FRHICommandList& RHICmdList, int32 ViewIndex, const FViewInfo& View);
 
 	/** Issues occlusion queries */
-	void RenderOcclusion(FRHICommandList& RHICmdList);
-	
+	void RenderOcclusion(FRHICommandList& RHICmdList, FViewInfo& View, const struct FViewOcclusionQueries& ViewOcclusionQueries);
+
 	bool ShouldRenderHZB(TArrayView<FViewInfo> InViews);
 
 
@@ -2933,8 +2933,8 @@ protected:
 	void UpdateSkyReflectionUniformBuffer(FRHICommandListBase& RHICmdList);
 	
 	void RenderForward(FRDGBuilder& GraphBuilder, FRDGTextureRef ViewFamilyTexture, FSceneTextures& SceneTextures, FDBufferTextures& DBufferTextures, FInstanceCullingManager& InstanceCullingManager);
-	void RenderForwardSinglePass(FRDGBuilder& GraphBuilder, class FMobileRenderPassParameters* PassParameters, struct FRenderViewContext& ViewContext, FSceneTextures& SceneTextures, FInstanceCullingManager& InstanceCullingManager);
-	void RenderForwardMultiPass(FRDGBuilder& GraphBuilder, class FMobileRenderPassParameters* PassParameters, struct FRenderViewContext& ViewContext, FSceneTextures& SceneTextures, FInstanceCullingManager& InstanceCullingManager);
+	void RenderForwardSinglePass(FRDGBuilder& GraphBuilder, class FMobileRenderPassParameters* PassParameters, struct FRenderViewContext& ViewContext, FSceneTextures& SceneTextures, FInstanceCullingManager& InstanceCullingManager, struct FViewOcclusionQueries* ViewOcclusionQueries);
+	void RenderForwardMultiPass(FRDGBuilder& GraphBuilder, class FMobileRenderPassParameters* PassParameters, struct FRenderViewContext& ViewContext, FSceneTextures& SceneTextures, FInstanceCullingManager& InstanceCullingManager, struct FViewOcclusionQueries* ViewOcclusionQueries);
 	
 	void RenderDeferredSinglePass(FRDGBuilder& GraphBuilder, FSceneTextures& SceneTextures, const FSortedLightSetSceneInfo& SortedLightSet, FDBufferTextures& DBufferTextures, FInstanceCullingManager& InstanceCullingManager);
 	void RenderDeferredMultiPass(FRDGBuilder& GraphBuilder, FSceneTextures& SceneTextures, const FSortedLightSetSceneInfo& SortedLightSet, FDBufferTextures& DBufferTextures, FInstanceCullingManager& InstanceCullingManager);
