[HWRT] Fix missing RayTracingScene.Reset() when there are multiple scene renderers.
- added post execute callback to call RayTracingScene.Reset() after the RDG graph of each renderer is executed.

diff --git a/Engine/Source/Runtime/Renderer/Private/DeferredShadingRenderer.cpp b/Engine/Source/Runtime/Renderer/Private/DeferredShadingRenderer.cpp
--- a/Engine/Source/Runtime/Renderer/Private/DeferredShadingRenderer.cpp
+++ b/Engine/Source/Runtime/Renderer/Private/DeferredShadingRenderer.cpp
@@ -1766,8 +1766,6 @@ void FDeferredShadingSceneRenderer::Render(FRDGBuilder& GraphBuilder, const FSce
 	FRayTracingScene& RayTracingScene = Scene->RayTracingScene;
 	RayTracingSBT.ResetMissAndCallableShaders();
 
-	RayTracingScene.Reset();
-
 	for (FViewInfo& View : Views)
 	{
 		if (IStereoRendering::IsStereoEyeView(View) && IStereoRendering::IsASecondaryView(View))
diff --git a/Engine/Source/Runtime/Renderer/Private/RayTracing/RayTracingInstanceBufferUtil.cpp b/Engine/Source/Runtime/Renderer/Private/RayTracing/RayTracingInstanceBufferUtil.cpp
--- a/Engine/Source/Runtime/Renderer/Private/RayTracing/RayTracingInstanceBufferUtil.cpp
+++ b/Engine/Source/Runtime/Renderer/Private/RayTracing/RayTracingInstanceBufferUtil.cpp
@@ -827,6 +827,8 @@ void BuildRayTracingInstanceBuffer(
 
 void FRayTracingInstanceBufferBuilder::Init(TConstArrayView<FRayTracingGeometryInstance> InInstances, FVector InPreViewTranslation)
 {
+	checkf(Instances.IsEmpty(), TEXT("RayTracingInstanceBufferBuilder was already initialized."));
+
 	Instances = InInstances;
 	PreViewTranslation = InPreViewTranslation;
 
diff --git a/Engine/Source/Runtime/Renderer/Private/RayTracing/RayTracingScene.h b/Engine/Source/Runtime/Renderer/Private/RayTracing/RayTracingScene.h
--- a/Engine/Source/Runtime/Renderer/Private/RayTracing/RayTracingScene.h
+++ b/Engine/Source/Runtime/Renderer/Private/RayTracing/RayTracingScene.h
@@ -235,17 +235,17 @@ private:
 
 		FRayTracingSceneRHIRef RayTracingSceneRHI;
 
-		FRDGBufferRef InstanceBuffer;
-		FRDGBufferRef BuildScratchBuffer;
+		FRDGBufferRef InstanceBuffer = nullptr;
+		FRDGBufferRef BuildScratchBuffer = nullptr;
 
 		// Feedback		
-		FRDGBufferRef InstanceHitCountBuffer;
+		FRDGBufferRef InstanceHitCountBuffer = nullptr;
 		FRDGBufferUAVRef InstanceHitCountBufferUAV;
-		FRDGBufferRef AccelerationStructureIndexBuffer;
+		FRDGBufferRef AccelerationStructureIndexBuffer = nullptr;
 
 		TRefCountPtr<FRDGPooledBuffer> RayTracingScenePooledBuffer;
-		FRDGBufferRef RayTracingSceneBufferRDG;
-		FRDGBufferSRVRef RayTracingSceneBufferSRV;
+		FRDGBufferRef RayTracingSceneBufferRDG = nullptr;
+		FRDGBufferSRVRef RayTracingSceneBufferSRV = nullptr;
 
 		FRDGBufferRef InstanceExtraDataBuffer = nullptr;
 
@@ -260,7 +260,7 @@ private:
 		uint32 GetCachedInstanceSectionSize();
 
 		// Feedback
-		FRDGBufferRef GeometryHandleBuffer;
+		FRDGBufferRef GeometryHandleBuffer = nullptr;
 		TArray<int32> GeometryHandles;
 
 		// Special data for debugging purposes
@@ -319,7 +319,7 @@ private:
 
 	const uint32 MaxReadbackBuffers = 4;
 
-	FRDGBufferRef InstanceStatsBuffer;
+	FRDGBufferRef InstanceStatsBuffer = nullptr;
 
 	TArray<FFeedbackReadbackData> FeedbackReadback;
 	uint32 FeedbackReadbackWriteIndex = 0;
diff --git a/Engine/Source/Runtime/Renderer/Private/SceneRendering.cpp b/Engine/Source/Runtime/Renderer/Private/SceneRendering.cpp
--- a/Engine/Source/Runtime/Renderer/Private/SceneRendering.cpp
+++ b/Engine/Source/Runtime/Renderer/Private/SceneRendering.cpp
@@ -4680,6 +4680,13 @@ void FSceneRenderer::OnRenderFinish(FRDGBuilder& GraphBuilder, FRDGTextureRef Vi
 		}
 	});
 
+#if RHI_RAYTRACING
+	GraphBuilder.AddPostExecuteCallback([this]
+		{
+			Scene->RayTracingScene.Reset();;
+		});
+#endif // RHI_RAYTRACING
+
 	// Sync all mesh pass setup tasks prior to starting RDG execution.
 	FGraphEventArray TasksToWait;
 
