[HWRT] Updated ray tracing instance gathering to support multiple views.

diff --git a/Engine/Source/Runtime/Renderer/Private/DeferredShadingRenderer.cpp b/Engine/Source/Runtime/Renderer/Private/DeferredShadingRenderer.cpp
--- a/Engine/Source/Runtime/Renderer/Private/DeferredShadingRenderer.cpp
+++ b/Engine/Source/Runtime/Renderer/Private/DeferredShadingRenderer.cpp
@@ -1084,19 +1084,9 @@ void FDeferredShadingSceneRenderer::SetupRayTracingRenderingData(FRDGBuilder& Gr
 	{
 		SCOPE_CYCLE_COUNTER(STAT_WaitRayTracingDynamicBindings);
 
-		for (FViewInfo& View : Views)
-		{
-			if (View.AddDynamicRayTracingMeshBatchTaskList.IsEmpty())
-			{
-				continue;
-			}
-
-			// need to wait for dynamic mesh batches tasks to finish before executing SetupRayTracingPipelineStatesAndSBT(...)
-			// since they can request new materials that need to be included in RTPSO 
-			UE::Tasks::Wait(View.AddDynamicRayTracingMeshBatchTaskList);
-
-			View.AddDynamicRayTracingMeshBatchTaskList.Empty();
-		}
+		// need to wait for dynamic mesh batches tasks to finish before executing SetupRayTracingPipelineStatesAndSBT(...)
+		// since they can request new materials that need to be included in RTPSO 
+		RayTracing::WaitForDynamicBindings(RayTracingGatherInstancesTaskData);
 	}
 
 	bool bIsUsingFallbackRTPSO = false;
@@ -1943,16 +1933,14 @@ void FDeferredShadingSceneRenderer::Render(FRDGBuilder& GraphBuilder, const FSce
 				}
 			}
 
-			const int32 ReferenceViewIndex = 0;
-			FViewInfo& ReferenceView = Views[ReferenceViewIndex];
+			InitViewTaskDatas.RayTracingGatherInstances = RayTracing::CreateGatherInstancesTaskData(Allocator, *Scene, Views.Num());
 
-			InitViewTaskDatas.RayTracingGatherInstances = RayTracing::CreateGatherInstancesTaskData(
-				Allocator,
-				*Scene,
-				ReferenceView,
-				ViewFamily,
-				GetViewPipelineState(ReferenceView).DiffuseIndirectMethod,
-				GetViewPipelineState(ReferenceView).ReflectionsMethod);
+			for (FViewInfo& View : Views)
+			{
+				const FPerViewPipelineState& ViewPipelineState = GetViewPipelineState(View);
+
+				RayTracing::AddView(*InitViewTaskDatas.RayTracingGatherInstances, View, ViewPipelineState.DiffuseIndirectMethod, ViewPipelineState.ReflectionsMethod);
+			}
 
 			RayTracing::BeginGatherInstances(*InitViewTaskDatas.RayTracingGatherInstances, InitViewTaskDatas.VisibilityTaskData->GetFrustumCullTask());
 		}
diff --git a/Engine/Source/Runtime/Renderer/Private/RayTracing/RayTracing.cpp b/Engine/Source/Runtime/Renderer/Private/RayTracing/RayTracing.cpp
--- a/Engine/Source/Runtime/Renderer/Private/RayTracing/RayTracing.cpp
+++ b/Engine/Source/Runtime/Renderer/Private/RayTracing/RayTracing.cpp
@@ -186,6 +186,7 @@ namespace RayTracing
 		EReflectionsMethod ReflectionsMethod)
 	{
 		bTranslucentGeometry = false;
+
 		LumenHardwareRayTracing::SetRayTracingSceneOptions(View, DiffuseIndirectMethod, ReflectionsMethod, *this);
 		RayTracingShadows::SetRayTracingSceneOptions(View.bHasRayTracingShadows, *this);
 
@@ -261,12 +262,25 @@ namespace RayTracing
 		TRange<int32> GeometriesToUpdateRange;
 	};
 
+	struct FDynamicPrimitiveIndex
+	{
+		FDynamicPrimitiveIndex() = default;
+
+		FDynamicPrimitiveIndex(int32 InIndex, uint8 InViewMask)
+			: Index(InIndex)
+			, ViewMask(InViewMask)
+		{}
+
+		uint32 Index : 24;
+		uint32 ViewMask : 8;
+	};
+
 	class FDynamicRayTracingInstancesContext
 	{
 	public:
-		FDynamicRayTracingInstancesContext(FScene& Scene, FViewInfo& InView, FSceneRenderingBulkObjectAllocator& InBulkAllocator, const FSceneOptions& InSceneOptions);
+		FDynamicRayTracingInstancesContext(FScene& InScene, TArrayView<FViewInfo*> InViews, TConstArrayView<const FSceneOptions*> InViewSceneOptions, FSceneRenderingBulkObjectAllocator& InBulkAllocator);
 
-		void GatherDynamicRayTracingInstances(TConstArrayView<int32> InDynamicPrimitives);
+		void GatherDynamicRayTracingInstances(TConstArrayView<FDynamicPrimitiveIndex> InDynamicPrimitives);
 
 		void GatherDynamicRayTracingInstances_RenderThread();
 
@@ -278,14 +292,15 @@ namespace RayTracing
 
 	private:
 
-		void GatherDynamicRayTracingInstances_Internal(int32 PrimitiveIndex);
+		void GatherDynamicRayTracingInstances_Internal(const FDynamicPrimitiveIndex& PrimitiveIndex);
 
 		FScene& Scene;
-		FViewInfo& View;
+		TArrayView<FViewInfo*> Views;
+		TConstArrayView<const FSceneOptions*> ViewSceneOptions;
 		FSceneRenderingBulkObjectAllocator& BulkAllocator;
-		const FSceneOptions& SceneOptions;
 
-		bool bTrackReferencedGeometryGroups;
+		const bool bTrackReferencedGeometryGroups;
+		bool bAnyViewRequiresTranslucentGeometry;
 
 		FRHICommandList* RHICmdList;
 		FGlobalDynamicVertexBuffer DynamicVertexBuffer;
@@ -295,18 +310,15 @@ namespace RayTracing
 
 		TArray<FDynamicRayTracingPrimitive> DynamicRayTracingPrimitives;
 
-		TArray<int32> RenderThreadDynamicPrimitives;
+		TArray<FDynamicPrimitiveIndex> RenderThreadDynamicPrimitives;
 	};
 
-	struct FGatherInstancesTaskData
+	struct FGatherInstancesViewTaskData
 	{
-		UE_NONCOPYABLE(FGatherInstancesTaskData)
-
-		FGatherInstancesTaskData(FScene& InScene, FViewInfo& InView, FSceneRenderingBulkObjectAllocator& InBulkAllocator, FSceneOptions InSceneOptions)
+		FGatherInstancesViewTaskData(FScene& InScene, FViewInfo& InView, FSceneRenderingBulkObjectAllocator& InBulkAllocator, FSceneOptions InSceneOptions)
 			: Scene(InScene)
 			, View(InView)
 			, SceneOptions(MoveTemp(InSceneOptions))
-			, DynamicRayTracingInstancesContext(Scene, View, InBulkAllocator, SceneOptions)
 		{
 
 		}
@@ -322,49 +334,95 @@ namespace RayTracing
 		TArray<FRelevantPrimitive> StaticPrimitives;
 		TArray<FRelevantCachedPrimitive> CachedStaticPrimitives;
 
-		FDynamicRayTracingInstancesContext DynamicRayTracingInstancesContext;
+		TArray<TSet<RayTracing::FGeometryGroupHandle>> ReferencedGeometryGroups;
+
+		// Array of primitives that need their cached ray tracing instance updated via FPrimitiveSceneInfo::UpdateCachedRaytracingData()
+		TArray<FPrimitiveSceneInfo*> DirtyCachedRayTracingPrimitives;
 
 		// This task must complete before accessing StaticPrimitivesIndices/DynamicPrimitivesIndices.
 		UE::Tasks::FTask GatherRelevantPrimitivesTask;
 
+		// This task must complete before accessing RayTracingScene/RaytracingSBT when processing dynamic instances.
+		UE::Tasks::FTask FinalizeGatherRelevantPrimitivesTask;
+
 		// This task must complete before accessing StaticPrimitives/CachedStaticPrimitives.
 		UE::Tasks::FTask GatherRelevantStaticPrimitivesTask;
 
+		// Used coarse mesh streaming handles during the last TLAS build
+		TArray<Nanite::CoarseMeshStreamingHandle> UsedCoarseMeshStreamingHandles; // TODO: Should be a set
+
+		int32 NumCachedStaticVisibleShaderBindings = 0; // TODO: Could remove this but it's used to Reserve
+
+		bool bUsesLightingChannels = false;
+	};
+
+	struct FGatherInstancesTaskData
+	{
+		UE_NONCOPYABLE(FGatherInstancesTaskData)
+
+		FGatherInstancesTaskData(
+			FSceneRenderingBulkObjectAllocator& InAllocator,
+			FScene& InScene,
+			uint32 NumViews)
+			: Scene(InScene)
+			, Allocator(InAllocator)
+		{
+			ViewTaskDatas.Reserve(NumViews);
+			Views.Reserve(NumViews);
+			ViewSceneOptions.Reserve(NumViews);
+		}
+
+		FScene& Scene;
+		FSceneRenderingBulkObjectAllocator& Allocator;
+
+		TArray<FGatherInstancesViewTaskData> ViewTaskDatas;
+		TArray<FViewInfo*> Views;
+		TArray<const FSceneOptions*> ViewSceneOptions;
+
+		UE::Tasks::FPipe AddInstancesPipe{ UE_SOURCE_LOCATION };
+		UE::Tasks::FPipe FinalizeGatherRelevantPrimitivesPipe{ UE_SOURCE_LOCATION };
+
+		FDynamicRayTracingInstancesContext* DynamicRayTracingInstancesContext;
+
 		// This task must complete before accessing DynamicRayTracingInstancesContext.
 		UE::Tasks::FTask GatherDynamicRayTracingInstancesTask;
+		UE::Tasks::FTaskEvent GatherDynamicRayTracingInstancesPrerequisites{ UE_SOURCE_LOCATION };
+		bool bGatherDynamicRayTracingInstancesPrerequisitesTriggered = false;
 
 		// This task must complete before PostRenderAllViewports().
 		UE::Tasks::FTask AddUsedStreamingHandlesTask;
 
-		UE::Tasks::FTaskEvent GatherDynamicRayTracingInstancesPrerequisites{ UE_SOURCE_LOCATION };
-		bool bGatherDynamicRayTracingInstancesPrerequisitesTriggered = false;
-
 		UE::Tasks::FTask VisibleRayTracingShaderBindingsFinalizeTask;
 		FRayTracingShaderBindingDataOneFrameArray VisibleShaderBindings;
 
-		// Array of primitives that should update their cached ray tracing instances via FPrimitiveSceneInfo::UpdateCachedRaytracingData()
-		TArray<FPrimitiveSceneInfo*> DirtyCachedRayTracingPrimitives; // TODO: remove this since it seems to be transient
-
-		// Used coarse mesh streaming handles during the last TLAS build
-		TArray<Nanite::CoarseMeshStreamingHandle> UsedCoarseMeshStreamingHandles; // TODO: Should be a set
-
-		int32 NumCachedStaticVisibleShaderBindings = 0; // TODO: Could remove this but it's used to Reserve
-
-		bool bUsesLightingChannels = false;
-
 		// Indicates that this object has been fully produced (for validation)
 		bool bValid = false;
 	};
 
 	FGatherInstancesTaskData* CreateGatherInstancesTaskData(
-		FSceneRenderingBulkObjectAllocator& InAllocator,
+		FSceneRenderingBulkObjectAllocator& Allocator,
 		FScene& Scene,
+		uint32 NumViews)
+	{
+		return Allocator.Create<FGatherInstancesTaskData>(Allocator, Scene, NumViews);
+	}
+
+	void AddView(
+		FGatherInstancesTaskData& TaskData,
 		FViewInfo& View,
-		const FViewFamilyInfo& ViewFamily,
 		EDiffuseIndirectMethod DiffuseIndirectMethod,
 		EReflectionsMethod ReflectionsMethod)
 	{
-		return InAllocator.Create<FGatherInstancesTaskData>(Scene, View, InAllocator, FSceneOptions(Scene, ViewFamily, View, DiffuseIndirectMethod, ReflectionsMethod));
+		if (IStereoRendering::IsStereoEyeView(View) && IStereoRendering::IsASecondaryView(View))
+		{
+			return;
+		}
+
+		const FViewFamilyInfo* ViewFamily = static_cast<const FViewFamilyInfo*>(View.Family);
+
+		TaskData.ViewTaskDatas.Add(FGatherInstancesViewTaskData(TaskData.Scene, View, TaskData.Allocator, FSceneOptions(TaskData.Scene, *ViewFamily, View, DiffuseIndirectMethod, ReflectionsMethod)));
+		TaskData.Views.Add(&View);
+		TaskData.ViewSceneOptions.Add(&TaskData.ViewTaskDatas.Last().SceneOptions);
 	}
 
 	void OnRenderBegin(const FSceneRenderUpdateInputs& SceneUpdateInputs)
@@ -507,7 +565,7 @@ namespace RayTracing
 		const FRHIRayTracingGeometry* RayTracingGeometry = nullptr;
 		TArray<FMeshBatch> MeshBatchesOwned;
 		TArrayView<const FMeshBatch> MeshBatchesView;
-		FRayTracingSBTAllocation* SBTAllocation;
+		FRayTracingSBTAllocation* SBTAllocation = nullptr;
 
 		TArrayView<const FMeshBatch> GetMeshBatches() const
 		{
@@ -533,6 +591,13 @@ namespace RayTracing
 		FRayTracingMeshBatchTaskPage* Next = nullptr;
 	};
 
+	struct FRayTracingMeshBatchTaskData
+	{
+		FRayTracingMeshBatchTaskPage* Head = nullptr;
+		FRayTracingMeshBatchTaskPage* Page = nullptr;
+		uint32 NumPendingMeshBatches = 0;
+	};
+
 	void DispatchRayTracingMeshBatchTask(FSceneRenderingBulkObjectAllocator& InBulkAllocator, FScene& Scene, FViewInfo& View, FRayTracingMeshBatchTaskPage* MeshBatchTaskHead, uint32 NumPendingMeshBatches)
 	{
 		FDynamicRayTracingMeshCommandStorage* TaskDynamicCommandStorage = InBulkAllocator.Create<FDynamicRayTracingMeshCommandStorage>();
@@ -540,6 +605,7 @@ namespace RayTracing
 
 		FRayTracingShaderBindingDataOneFrameArray* TaskVisibleShaderBindings = InBulkAllocator.Create<FRayTracingShaderBindingDataOneFrameArray>();
 		TaskVisibleShaderBindings->Reserve(NumPendingMeshBatches);
+
 		View.DynamicRayTracingShaderBindingsPerTask.Add(TaskVisibleShaderBindings);
 
 		View.AddDynamicRayTracingMeshBatchTaskList.Add(UE::Tasks::Launch(UE_SOURCE_LOCATION,
@@ -572,11 +638,11 @@ namespace RayTracing
 			}, UE::Tasks::ETaskPriority::High));
 	};
 
-	FDynamicRayTracingInstancesContext::FDynamicRayTracingInstancesContext(FScene& Scene, FViewInfo& InView, FSceneRenderingBulkObjectAllocator& InBulkAllocator, const FSceneOptions& InSceneOptions)
+	FDynamicRayTracingInstancesContext::FDynamicRayTracingInstancesContext(FScene& Scene, TArrayView<FViewInfo*> InViews, TConstArrayView<const FSceneOptions*> InViewSceneOptions, FSceneRenderingBulkObjectAllocator& InBulkAllocator)
 		: Scene(Scene)
-		, View(InView)
+		, Views(InViews)
+		, ViewSceneOptions(InViewSceneOptions)
 		, BulkAllocator(InBulkAllocator)
-		, SceneOptions(InSceneOptions)
 		, bTrackReferencedGeometryGroups(IsRayTracingUsingReferenceBasedResidency())
 		, RHICmdList(new FRHICommandList(FRHIGPUMask::All()))
 		, DynamicVertexBuffer(*RHICmdList)
@@ -592,35 +658,45 @@ namespace RayTracing
 			FSceneRenderer::DynamicReadBufferForRayTracing
 		);
 
-		RayTracingInstanceCollector.AddViewMeshArrays(&View, &View.RayTracingDynamicPrimitiveCollector);
+		bAnyViewRequiresTranslucentGeometry = false;
+
+		for (int32 ViewIndex = 0; ViewIndex < Views.Num(); ViewIndex++)
+		{
+			FViewInfo& View = *Views[ViewIndex];
+
+			RayTracingInstanceCollector.AddViewMeshArrays(&View, &View.DynamicPrimitiveCollector);
+
+			bAnyViewRequiresTranslucentGeometry |= ViewSceneOptions[ViewIndex]->bTranslucentGeometry;
+		}
 	}
 
-	void FDynamicRayTracingInstancesContext::GatherDynamicRayTracingInstances_Internal(int32 PrimitiveIndex)
+	void FDynamicRayTracingInstancesContext::GatherDynamicRayTracingInstances_Internal(const FDynamicPrimitiveIndex& PrimitiveIndex)
 	{
-		FPrimitiveSceneProxy* SceneProxy = Scene.PrimitiveSceneProxies[PrimitiveIndex];
-		FPrimitiveSceneInfo* SceneInfo = Scene.Primitives[PrimitiveIndex];
+		FPrimitiveSceneProxy* SceneProxy = Scene.PrimitiveSceneProxies[PrimitiveIndex.Index];
+		FPrimitiveSceneInfo* SceneInfo = Scene.Primitives[PrimitiveIndex.Index];
 		const FPersistentPrimitiveIndex PersistentPrimitiveIndex = SceneInfo->GetPersistentIndex();
 
+		//RayTracingInstanceCollector.SetPrimitive(SceneProxy, PrimitiveIndex);
 		RayTracingInstanceCollector.SetPrimitive(SceneProxy, FHitProxyId::InvisibleHitProxyId);
-		RayTracingInstanceCollector.SetVisibilityMap(0x1);
+		RayTracingInstanceCollector.SetVisibilityMap(PrimitiveIndex.ViewMask);
 
 		int32 BaseRayTracingInstance = RayTracingInstanceCollector.RayTracingInstances.Num();
 		int32 BaseGeometryToUpdate = RayTracingInstanceCollector.RayTracingGeometriesToUpdate.Num();
 
-		if (SceneOptions.bTranslucentGeometry || SceneProxy->IsOpaqueOrMasked())
+		if (bAnyViewRequiresTranslucentGeometry || SceneProxy->IsOpaqueOrMasked())
 		{
 			SceneProxy->GetDynamicRayTracingInstances(RayTracingInstanceCollector);
 		}
 
 		FDynamicRayTracingPrimitive Tmp;
-		Tmp.PrimitiveIndex = PrimitiveIndex;
+		Tmp.PrimitiveIndex = PrimitiveIndex.Index;
 		Tmp.InstancesRange = TRange<int32>(BaseRayTracingInstance, RayTracingInstanceCollector.RayTracingInstances.Num());
 		Tmp.GeometriesToUpdateRange = TRange<int32>(BaseGeometryToUpdate, RayTracingInstanceCollector.RayTracingGeometriesToUpdate.Num());
 
 		DynamicRayTracingPrimitives.Add(MoveTemp(Tmp));
 	}
 
-	void FDynamicRayTracingInstancesContext::GatherDynamicRayTracingInstances(TConstArrayView<int32> InDynamicPrimitives)
+	void FDynamicRayTracingInstancesContext::GatherDynamicRayTracingInstances(TConstArrayView<FDynamicPrimitiveIndex> InDynamicPrimitives)
 	{
 		TRACE_CPUPROFILER_EVENT_SCOPE(RayTracing_GatherDynamicRayTracingInstances);
 
@@ -634,9 +710,9 @@ namespace RayTracing
 
 		// TODO: Could filter primitives whose proxy supports ParallelGDRTI during GatherRayTracingRelevantPrimitives_Parallel
 
-		for (int32 PrimitiveIndex : InDynamicPrimitives)
+		for (const FDynamicPrimitiveIndex& PrimitiveIndex : InDynamicPrimitives)
 		{
-			FPrimitiveSceneProxy* SceneProxy = Scene.PrimitiveSceneProxies[PrimitiveIndex];
+			FPrimitiveSceneProxy* SceneProxy = Scene.PrimitiveSceneProxies[PrimitiveIndex.Index];
 
 			if (SceneProxy->SupportsParallelGDRTI())
 			{
@@ -655,7 +731,7 @@ namespace RayTracing
 
 		check(IsInRenderingThread());
 
-		for (int32 PrimitiveIndex : RenderThreadDynamicPrimitives)
+		for (const FDynamicPrimitiveIndex& PrimitiveIndex : RenderThreadDynamicPrimitives)
 		{
 			GatherDynamicRayTracingInstances_Internal(PrimitiveIndex);
 		}
@@ -680,13 +756,13 @@ namespace RayTracing
 			const FPersistentPrimitiveIndex PersistentPrimitiveIndex = SceneInfo->GetPersistentIndex();
 
 			for (int32 Index = DynamicRayTracingPrimitive.GeometriesToUpdateRange.GetLowerBoundValue(); Index < DynamicRayTracingPrimitive.GeometriesToUpdateRange.GetUpperBoundValue(); ++Index)
-			{				
+			{
 				const FRayTracingInstanceCollector::FRayTracingDynamicGeometryUpdateRequest& UpdateRequest = RayTracingInstanceCollector.RayTracingGeometriesToUpdate[Index];
-
+				
 				DynamicGeometryUpdateManager->AddDynamicGeometryToUpdate(
 					*RHICmdList,
 					&Scene,
-					&View,
+					Views[UpdateRequest.ViewIndex],
 					SceneProxy,
 					UpdateRequest.Params,
 					PersistentPrimitiveIndex.Index
@@ -726,10 +802,12 @@ namespace RayTracing
 	{
 		for (int32 ViewIndex = 0; ViewIndex < RayTracingInstanceCollector.RDGPooledBuffers.Num(); ++ViewIndex)
 		{
-			check(RayTracingInstanceCollector.Views[ViewIndex] == &View);
+			check(RayTracingInstanceCollector.Views[ViewIndex] == Views[ViewIndex]);
 
 			const TSet<FRDGPooledBuffer*>& ViewRDGPooledBuffers = RayTracingInstanceCollector.RDGPooledBuffers[ViewIndex];
 
+			FViewInfo& View = *Views[ViewIndex];
+
 			for (FRDGPooledBuffer* PooledBuffer : ViewRDGPooledBuffers)
 			{
 				FRDGBufferRef RDGBuffer = RDGBuilder.RegisterExternalBuffer(PooledBuffer);
@@ -742,30 +820,25 @@ namespace RayTracing
 	{
 		TRACE_CPUPROFILER_EVENT_SCOPE(RayTracing_AddDynamicInstancesToScene);
 
-		const float CurrentWorldTime = View.Family->Time.GetWorldTimeSeconds();
-
 		const bool bParallelMeshBatchSetup = GRayTracingParallelMeshBatchSetup && FApp::ShouldUseThreadingForPerformance();
-
-		FRayTracingMeshBatchTaskPage* MeshBatchTaskHead = nullptr;
-		FRayTracingMeshBatchTaskPage* MeshBatchTaskPage = nullptr;
-		uint32 NumPendingMeshBatches = 0;
+		const float LastRenderTimeUpdateDistance = CVarRayTracingDynamicGeometryLastRenderTimeUpdateDistance.GetValueOnRenderThread();
 		const uint32 RayTracingParallelMeshBatchSize = GRayTracingParallelMeshBatchSize;
 
-		auto KickRayTracingMeshBatchTask = [this, &MeshBatchTaskHead, &MeshBatchTaskPage, &NumPendingMeshBatches]()
+		TArray<FRayTracingMeshBatchTaskData, TInlineAllocator<2, SceneRenderingAllocator>> MeshBatchTaskData;
+		MeshBatchTaskData.AddDefaulted(Views.Num());
+
+		auto KickRayTracingMeshBatchTask = [this](FRayTracingMeshBatchTaskData& MeshBatchTaskData, FViewInfo& View)
 			{
-				if (MeshBatchTaskHead)
+				if (MeshBatchTaskData.Head)
 				{
-					DispatchRayTracingMeshBatchTask(BulkAllocator, Scene, View, MeshBatchTaskHead, NumPendingMeshBatches);
+					DispatchRayTracingMeshBatchTask(BulkAllocator, Scene, View, MeshBatchTaskData.Head, MeshBatchTaskData.NumPendingMeshBatches);
 				}
 
-				MeshBatchTaskHead = nullptr;
-				MeshBatchTaskPage = nullptr;
-				NumPendingMeshBatches = 0;
+				MeshBatchTaskData.Head = nullptr;
+				MeshBatchTaskData.Page = nullptr;
+				MeshBatchTaskData.NumPendingMeshBatches = 0;
 			};
 
-		const int32 ViewDynamicPrimitiveId = View.RayTracingDynamicPrimitiveCollector.GetPrimitiveIdRange().GetLowerBoundValue();
-		const int32 ViewInstanceSceneDataOffset = View.RayTracingDynamicPrimitiveCollector.GetInstanceSceneDataOffset();
-
 		for (const FDynamicRayTracingPrimitive& DynamicRayTracingPrimitive : DynamicRayTracingPrimitives)
 		{
 			const int32 PrimitiveIndex = DynamicRayTracingPrimitive.PrimitiveIndex;
@@ -783,6 +856,10 @@ namespace RayTracing
 				for (FRayTracingInstanceCollector::FRayTracingInstanceAndViewIndex& InstanceAndViewIndex : TempRayTracingInstances)
 				{
 					FRayTracingInstance& Instance = InstanceAndViewIndex.Instance;
+					FViewInfo& View = *Views[InstanceAndViewIndex.ViewIndex];
+					 
+					const int32 ViewDynamicPrimitiveId = View.RayTracingDynamicPrimitiveCollector.GetPrimitiveIdRange().GetLowerBoundValue();
+					const int32 ViewInstanceSceneDataOffset = View.RayTracingDynamicPrimitiveCollector.GetInstanceSceneDataOffset();
 
 					const FRayTracingGeometry* Geometry = Instance.Geometry;
 
@@ -847,7 +924,7 @@ namespace RayTracing
 						SceneInfo->bDynamicRayTracingInstanceCachedDataDirty = false;
 					}
 
-					if (!SceneOptions.bTranslucentGeometry && SceneInfo->DynamicRayTracingInstanceCachedData.bAllSegmentsTranslucent)
+					if (!ViewSceneOptions[InstanceAndViewIndex.ViewIndex]->bTranslucentGeometry && SceneInfo->DynamicRayTracingInstanceCachedData.bAllSegmentsTranslucent)
 					{
 						continue;
 					}
@@ -989,29 +1066,31 @@ namespace RayTracing
 
 					if (bParallelMeshBatchSetup)
 					{
-						if (NumPendingMeshBatches >= RayTracingParallelMeshBatchSize)
+						FRayTracingMeshBatchTaskData& ViewMeshBatchTaskData = MeshBatchTaskData[InstanceAndViewIndex.ViewIndex];
+
+						if (ViewMeshBatchTaskData.NumPendingMeshBatches >= RayTracingParallelMeshBatchSize)
 						{
-							KickRayTracingMeshBatchTask();
+							KickRayTracingMeshBatchTask(ViewMeshBatchTaskData, View);
 						}
 
-						if (MeshBatchTaskPage == nullptr || MeshBatchTaskPage->NumWorkItems == FRayTracingMeshBatchTaskPage::MaxWorkItems)
+						if (ViewMeshBatchTaskData.Page == nullptr || ViewMeshBatchTaskData.Page->NumWorkItems == FRayTracingMeshBatchTaskPage::MaxWorkItems)
 						{
 							FRayTracingMeshBatchTaskPage* NextPage = BulkAllocator.Create<FRayTracingMeshBatchTaskPage>();
-							if (MeshBatchTaskHead == nullptr)
+							if (ViewMeshBatchTaskData.Head == nullptr)
 							{
-								MeshBatchTaskHead = NextPage;
+								ViewMeshBatchTaskData.Head = NextPage;
 							}
-							if (MeshBatchTaskPage)
+							if (ViewMeshBatchTaskData.Page)
 							{
-								MeshBatchTaskPage->Next = NextPage;
+								ViewMeshBatchTaskData.Page->Next = NextPage;
 							}
-							MeshBatchTaskPage = NextPage;
+							ViewMeshBatchTaskData.Page = NextPage;
 						}
 
-						FRayTracingMeshBatchWorkItem& WorkItem = MeshBatchTaskPage->WorkItems[MeshBatchTaskPage->NumWorkItems];
-						MeshBatchTaskPage->NumWorkItems++;
+						FRayTracingMeshBatchWorkItem& WorkItem = ViewMeshBatchTaskData.Page->WorkItems[ViewMeshBatchTaskData.Page->NumWorkItems];
+						ViewMeshBatchTaskData.Page->NumWorkItems++;
 
-						NumPendingMeshBatches += Instance.GetMaterials().Num();
+						ViewMeshBatchTaskData.NumPendingMeshBatches += Instance.GetMaterials().Num();
 
 						if (Instance.OwnsMaterials())
 						{
@@ -1039,9 +1118,24 @@ namespace RayTracing
 					}
 				}
 
-				if (CVarRayTracingDynamicGeometryLastRenderTimeUpdateDistance.GetValueOnRenderThread() > 0.0f)
+				if (LastRenderTimeUpdateDistance > 0.0f)
 				{
-					if (FVector::Distance(SceneProxy->GetActorPosition(), View.ViewMatrices.GetViewOrigin()) < CVarRayTracingDynamicGeometryLastRenderTimeUpdateDistance.GetValueOnRenderThread())
+					float CurrentWorldTime = 0;
+					double DistanceToView = std::numeric_limits<double>::infinity();
+
+					for (int32 ViewIndex = 0; ViewIndex < Views.Num(); ++ViewIndex)
+					{
+						const FViewInfo& View = *Views[ViewIndex];
+
+						CurrentWorldTime = FMath::Max(CurrentWorldTime, View.Family->Time.GetWorldTimeSeconds());
+
+						if (LastRenderTimeUpdateDistance > 0.0f)
+						{
+							DistanceToView = FMath::Max(DistanceToView, FVector::Distance(SceneProxy->GetActorPosition(), View.ViewMatrices.GetViewOrigin()));
+						}
+					}
+
+					if (DistanceToView < LastRenderTimeUpdateDistance)
 					{
 						// Update LastRenderTime for components so that visibility based ticking (like skeletal meshes) can get updated
 						// We are only doing this for dynamic geometries now
@@ -1051,20 +1145,24 @@ namespace RayTracing
 				}
 			}
 		}
-
-		KickRayTracingMeshBatchTask();
+		
+		for (int32 ViewIndex = 0; ViewIndex < Views.Num(); ++ViewIndex)
+		{
+			KickRayTracingMeshBatchTask(MeshBatchTaskData[ViewIndex], *Views[ViewIndex]);
+		}
 
 		RayTracingInstanceCollector.RayTracingInstances.Empty();
 	}
 
-	void GatherRelevantPrimitives(FGatherInstancesTaskData& TaskData, bool bUsingReferenceBasedResidency)
+	void GatherRelevantPrimitives(FGatherInstancesViewTaskData& TaskData, bool bUsingReferenceBasedResidency)
 	{
 		FScene& Scene = TaskData.Scene;
 		FViewInfo& View = TaskData.View;
 
 		const bool bGameView = View.bIsGameView || View.Family->EngineShowFlags.Game;
 
-		bool bPerformRayTracing = View.State != nullptr && !View.bIsReflectionCapture && View.IsRayTracingAllowedForView();
+		const bool bPerformRayTracing = View.State != nullptr && !View.bIsReflectionCapture && View.IsRayTracingAllowedForView();
+
 		if (bPerformRayTracing)
 		{
 			TRACE_CPUPROFILER_EVENT_SCOPE(GatherRayTracingRelevantPrimitives);
@@ -1187,6 +1285,8 @@ namespace RayTracing
 				TaskData.UsedCoarseMeshStreamingHandles.Reserve(NumUsedCoarseMeshStreamingHandles);
 				TaskData.DirtyCachedRayTracingPrimitives.Reserve(NumDirtyCachedRayTracingPrimitives);
 
+				TaskData.ReferencedGeometryGroups.Reserve(Contexts.Num());
+
 				for (auto& Context : Contexts)
 				{
 					Context.StaticPrimitives.CopyToLinearArray(TaskData.StaticPrimitivesIndices);
@@ -1194,19 +1294,13 @@ namespace RayTracing
 					Context.UsedCoarseMeshStreamingHandles.CopyToLinearArray(TaskData.UsedCoarseMeshStreamingHandles);
 					Context.DirtyCachedRayTracingPrimitives.CopyToLinearArray(TaskData.DirtyCachedRayTracingPrimitives);
 
-					if(bUsingReferenceBasedResidency)
-					{
-						((FRayTracingGeometryManager*)GRayTracingGeometryManager)->AddReferencedGeometryGroups(Context.ReferencedGeometryGroups);
-					}
+					TaskData.ReferencedGeometryGroups.Add(MoveTemp(Context.ReferencedGeometryGroups));
 				}
 			}
 		}
-
-		// TODO: could be moved to GatherRelevantStaticPrimitivesTask so that launching GatherDynamicRayTracingInstancesTask isn't delayed
-		FPrimitiveSceneInfo::UpdateCachedRaytracingData(&TaskData.Scene, TaskData.DirtyCachedRayTracingPrimitives);
 	}
 
-	void GatherRelevantStaticPrimitives(FGatherInstancesTaskData& TaskData, float GlobalLODScale, int32 ForcedLODLevel, bool bUsingReferenceBasedResidency)
+	void GatherRelevantStaticPrimitives(FGatherInstancesViewTaskData& TaskData, float GlobalLODScale, int32 ForcedLODLevel, bool bUsingReferenceBasedResidency)
 	{
 		TRACE_CPUPROFILER_EVENT_SCOPE(GatherRayTracingRelevantStaticPrimitives);
 
@@ -1701,46 +1795,89 @@ namespace RayTracing
 	{
 		TRACE_CPUPROFILER_EVENT_SCOPE(RayTracingScene_AddDynamicInstances);
 
-		TaskData.DynamicRayTracingInstancesContext.AddInstancesToScene(RayTracingScene, RayTracingSBT, SharedBufferGenerationID);
+		TaskData.DynamicRayTracingInstancesContext->AddInstancesToScene(RayTracingScene, RayTracingSBT, SharedBufferGenerationID);
 	}
 
 	void BeginGatherInstances(FGatherInstancesTaskData& TaskData, UE::Tasks::FTask FrustumCullTask)
 	{
 		const bool bUsingReferenceBasedResidency = IsRayTracingUsingReferenceBasedResidency();
 
+		static const auto ICVarStaticMeshLODDistanceScale = IConsoleManager::Get().FindConsoleVariable(TEXT("r.StaticMeshLODDistanceScale"));
+		const float LODScaleCVarValue = ICVarStaticMeshLODDistanceScale->GetFloat();
+		const int32 ForcedLODLevel = GetCVarForceLOD();
+
+		const bool bMultiView = TaskData.Views.Num() > 1;
+
+		// When there are multiple views, we use a set to avoid duplicate updates of cached ray tracing primitives
+		TSet<FPrimitiveSceneInfo*>* UpdatedDirtyCachedRayTracingPrimitives = bMultiView ? TaskData.Allocator.Create<TSet<FPrimitiveSceneInfo*>>() : nullptr;
+
 		// Use high priority tasks to reduce stalls on the critical path.
 		const UE::Tasks::ETaskPriority TaskPriority = UE::Tasks::ETaskPriority::High;
 
 		UE::Tasks::FTask CacheRayTracingPrimitivesTask = TaskData.Scene.GetCacheRayTracingPrimitivesTask();
 
-		TaskData.GatherRelevantPrimitivesTask = UE::Tasks::Launch(UE_SOURCE_LOCATION, [&TaskData, bUsingReferenceBasedResidency]
-			{
-				FTaskTagScope Scope(ETaskTag::EParallelRenderingThread);
-				GatherRelevantPrimitives(TaskData, bUsingReferenceBasedResidency);
-			}, MakeArrayView({ CacheRayTracingPrimitivesTask, FrustumCullTask }), TaskPriority);
+		TArray<UE::Tasks::FTask> GatherRelevantPrimitivesTasks;
 
-		static const auto ICVarStaticMeshLODDistanceScale = IConsoleManager::Get().FindConsoleVariable(TEXT("r.StaticMeshLODDistanceScale"));
-		const float LODScaleCVarValue = ICVarStaticMeshLODDistanceScale->GetFloat();
-		const int32 ForcedLODLevel = GetCVarForceLOD();
+		for (FGatherInstancesViewTaskData& ViewTaskData : TaskData.ViewTaskDatas)
+		{
+			ViewTaskData.GatherRelevantPrimitivesTask = UE::Tasks::Launch(UE_SOURCE_LOCATION, [&ViewTaskData, bUsingReferenceBasedResidency]
+				{
+					FTaskTagScope Scope(ETaskTag::EParallelRenderingThread);
+					GatherRelevantPrimitives(ViewTaskData, bUsingReferenceBasedResidency);
+				}, MakeArrayView({ CacheRayTracingPrimitivesTask, FrustumCullTask }), TaskPriority);
 
-		TaskData.GatherRelevantStaticPrimitivesTask = UE::Tasks::Launch(UE_SOURCE_LOCATION, [&TaskData, LODScaleCVarValue, ForcedLODLevel, bUsingReferenceBasedResidency]
-			{
-				FTaskTagScope Scope(ETaskTag::EParallelRenderingThread);
-				GatherRelevantStaticPrimitives(TaskData, LODScaleCVarValue, ForcedLODLevel, bUsingReferenceBasedResidency);
-			}, TaskData.GatherRelevantPrimitivesTask, TaskPriority);
+			// Finalize logic can't run in parallel so a pipe is used to serialize work.
+			ViewTaskData.FinalizeGatherRelevantPrimitivesTask = TaskData.FinalizeGatherRelevantPrimitivesPipe.Launch(UE_SOURCE_LOCATION,
+				[&ViewTaskData, UpdatedDirtyCachedRayTracingPrimitives, bMultiView]
+				{
+					FTaskTagScope Scope(ETaskTag::EParallelRenderingThread);
+					TRACE_CPUPROFILER_EVENT_SCOPE(FinalizeGatherRelevantPrimitives);
 
-		TaskData.GatherDynamicRayTracingInstancesPrerequisites.AddPrerequisites(TaskData.GatherRelevantPrimitivesTask);
-		TaskData.GatherDynamicRayTracingInstancesPrerequisites.AddPrerequisites(TaskData.Scene.GetGPUSkinCacheTask());
+					for (const TSet<RayTracing::FGeometryGroupHandle>& ReferencedGeometryGroups : ViewTaskData.ReferencedGeometryGroups)
+					{
+						((FRayTracingGeometryManager*)GRayTracingGeometryManager)->AddReferencedGeometryGroups(ReferencedGeometryGroups);
+					}
 
-		// TODO: Could gather dynamic ray tracing instances using multiple tasks / FDynamicRayTracingInstancesContext
-		TaskData.GatherDynamicRayTracingInstancesTask = UE::Tasks::Launch(UE_SOURCE_LOCATION, [&TaskData]
-			{
-				FTaskTagScope Scope(ETaskTag::EParallelRenderingThread);
-				TaskData.DynamicRayTracingInstancesContext.GatherDynamicRayTracingInstances(TaskData.DynamicPrimitivesIndices);
-			}, TaskData.GatherDynamicRayTracingInstancesPrerequisites, TaskPriority);
+					if (bMultiView)
+					{
+						if (UpdatedDirtyCachedRayTracingPrimitives->IsEmpty())
+						{
+							// First view simply add DirtyCachedRayTracingPrimitives
+							UpdatedDirtyCachedRayTracingPrimitives->Append(ViewTaskData.DirtyCachedRayTracingPrimitives);
+						}
+						else
+						{
+							// Then need to filter DirtyCachedRayTracingPrimitives to avoid duplicate updates
+							TArray<FPrimitiveSceneInfo*> DirtyCachedRayTracingPrimitives = MoveTemp(ViewTaskData.DirtyCachedRayTracingPrimitives);
+							ViewTaskData.DirtyCachedRayTracingPrimitives.Empty(DirtyCachedRayTracingPrimitives.Num());
+
+							for (FPrimitiveSceneInfo* SceneInfo : DirtyCachedRayTracingPrimitives)
+							{
+								if (!UpdatedDirtyCachedRayTracingPrimitives->Contains(SceneInfo))
+								{
+									UpdatedDirtyCachedRayTracingPrimitives->Add(SceneInfo);
+
+									ViewTaskData.DirtyCachedRayTracingPrimitives.Add(SceneInfo);
+								}
+							}
+						}
+					}
+
+					FPrimitiveSceneInfo::UpdateCachedRaytracingData(&ViewTaskData.Scene, ViewTaskData.DirtyCachedRayTracingPrimitives);
+				}, ViewTaskData.GatherRelevantPrimitivesTask);
 
+			ViewTaskData.GatherRelevantStaticPrimitivesTask = UE::Tasks::Launch(UE_SOURCE_LOCATION, [&ViewTaskData, LODScaleCVarValue, ForcedLODLevel, bUsingReferenceBasedResidency]
+				{
+					FTaskTagScope Scope(ETaskTag::EParallelRenderingThread);
+					GatherRelevantStaticPrimitives(ViewTaskData, LODScaleCVarValue, ForcedLODLevel, bUsingReferenceBasedResidency);
+				}, ViewTaskData.FinalizeGatherRelevantPrimitivesTask, TaskPriority);
+
+			TaskData.GatherDynamicRayTracingInstancesPrerequisites.AddPrerequisites(ViewTaskData.GatherRelevantPrimitivesTask);
 
-		TaskData.AddUsedStreamingHandlesTask = UE::Tasks::Launch(UE_SOURCE_LOCATION, [&TaskData, LODScaleCVarValue, ForcedLODLevel, bUsingReferenceBasedResidency]
+			GatherRelevantPrimitivesTasks.Add(ViewTaskData.GatherRelevantPrimitivesTask);
+		}
+
+		TaskData.AddUsedStreamingHandlesTask = UE::Tasks::Launch(UE_SOURCE_LOCATION, [&TaskData]
 			{
 				FTaskTagScope Scope(ETaskTag::EParallelRenderingThread);
 
@@ -1748,9 +1885,58 @@ namespace RayTracing
 				Nanite::FCoarseMeshStreamingManager* CoarseMeshSM = IStreamingManager::Get().GetNaniteCoarseMeshStreamingManager();
 				if (CoarseMeshSM)
 				{
-					CoarseMeshSM->AddUsedStreamingHandles(TaskData.UsedCoarseMeshStreamingHandles);
+					for (FGatherInstancesViewTaskData& ViewTaskData : TaskData.ViewTaskDatas)
+					{
+						CoarseMeshSM->AddUsedStreamingHandles(ViewTaskData.UsedCoarseMeshStreamingHandles);
+					}
 				}
-			}, TaskData.GatherRelevantPrimitivesTask, TaskPriority);
+			}, GatherRelevantPrimitivesTasks, TaskPriority);
+
+		// Dynamic instance gathering
+		{
+			TaskData.DynamicRayTracingInstancesContext = TaskData.Allocator.Create<FDynamicRayTracingInstancesContext>(TaskData.Scene, TaskData.Views, TaskData.ViewSceneOptions, TaskData.Allocator);
+
+			TaskData.GatherDynamicRayTracingInstancesPrerequisites.AddPrerequisites(TaskData.Scene.GetGPUSkinCacheTask());
+
+			// TODO: Could gather dynamic ray tracing instances using multiple tasks / FDynamicRayTracingInstancesContext
+			TaskData.GatherDynamicRayTracingInstancesTask = UE::Tasks::Launch(UE_SOURCE_LOCATION, [&TaskData]
+				{
+					FTaskTagScope Scope(ETaskTag::EParallelRenderingThread);
+
+					// Build array of FDynamicPrimitiveIndex (includes ViewMasks for each dynamic primitive) by merging ViewTaskData.DynamicPrimitivesIndices.
+					//	Alternatively ViewMasks array could be built during GatherRelevantPrimitives(...) with atomic OR operations,
+					//	and then here we could just loop over all primitives to collect the relevant dynamic ones.
+					//	(loop over all primitives vs loop over dynamic primitives for each view and combine using map + array)
+
+					TArray<FDynamicPrimitiveIndex> DynamicPrimitivesIndices;
+					TMap<int32, int32> PrimitiveIndexMap;
+
+					for (int32 ViewIndex = 0; ViewIndex < TaskData.ViewTaskDatas.Num(); ++ViewIndex)
+					{
+						const FGatherInstancesViewTaskData& ViewTaskData = TaskData.ViewTaskDatas[ViewIndex];
+
+						DynamicPrimitivesIndices.Reserve(DynamicPrimitivesIndices.Num() + ViewTaskData.DynamicPrimitivesIndices.Num());
+
+						for (int32 PrimitiveIndex : ViewTaskData.DynamicPrimitivesIndices)
+						{
+							if (PrimitiveIndexMap.Contains(PrimitiveIndex))
+							{
+								const int32 DynamicPrimitiveIndex = PrimitiveIndexMap[PrimitiveIndex];
+								DynamicPrimitivesIndices[DynamicPrimitiveIndex].ViewMask |= (1 << ViewIndex);
+
+								check(DynamicPrimitivesIndices[DynamicPrimitiveIndex].Index == PrimitiveIndex);
+							}
+							else
+							{
+								const int32 DynamicPrimitiveIndex = DynamicPrimitivesIndices.Emplace(PrimitiveIndex, (1 << ViewIndex));
+								PrimitiveIndexMap.Add(PrimitiveIndex, DynamicPrimitiveIndex);
+							}
+						}
+					}
+
+					TaskData.DynamicRayTracingInstancesContext->GatherDynamicRayTracingInstances(DynamicPrimitivesIndices);
+				}, TaskData.GatherDynamicRayTracingInstancesPrerequisites, TaskPriority);
+		}
 
 		TaskData.bValid = true;
 	}
@@ -1776,10 +1962,12 @@ namespace RayTracing
 		SCOPE_CYCLE_COUNTER(STAT_RayTracing_FinishGatherInstances);
 		RDG_CSV_STAT_EXCLUSIVE_SCOPE(GraphBuilder, RayTracing_FinishGatherInstances);
 
-		TaskData.GatherRelevantPrimitivesTask.Wait();
+		for (FGatherInstancesViewTaskData& ViewTaskData : TaskData.ViewTaskDatas)
+		{
+			ViewTaskData.FinalizeGatherRelevantPrimitivesTask.Wait();
 
-		FScene& Scene = TaskData.Scene;
-		FViewInfo& View = TaskData.View;
+			INC_DWORD_STAT_BY(STAT_VisibleRayTracingPrimitives, ViewTaskData.StaticPrimitives.Num() + ViewTaskData.CachedStaticPrimitives.Num() + ViewTaskData.DynamicPrimitivesIndices.Num());
+		}
 
 		// Prepare ray tracing scene instance list
 		checkf(TaskData.bValid, TEXT("Ray tracing relevant primitive list is expected to have been created before GatherRayTracingWorldInstancesForView() is called."));
@@ -1793,77 +1981,110 @@ namespace RayTracing
 		RayTracingSBT.ResetDynamicAllocationData();
 		RayTracingScene.LockCachedInstances();
 
-		INC_DWORD_STAT_BY(STAT_VisibleRayTracingPrimitives, TaskData.StaticPrimitives.Num() + TaskData.CachedStaticPrimitives.Num() + TaskData.DynamicPrimitivesIndices.Num());
-
-		FRayTracingDynamicGeometryUpdateManager* DynamicGeometryUpdateManager = Scene.GetRayTracingDynamicGeometryUpdateManager();
+		FRayTracingDynamicGeometryUpdateManager* DynamicGeometryUpdateManager = TaskData.Scene.GetRayTracingDynamicGeometryUpdateManager();
 		const int64 SharedBufferGenerationID = DynamicGeometryUpdateManager->BeginUpdate();
 
-		TaskData.GatherDynamicRayTracingInstancesTask.Wait();
-		TaskData.DynamicRayTracingInstancesContext.GatherDynamicRayTracingInstances_RenderThread();
-		TaskData.DynamicRayTracingInstancesContext.Finish(GraphBuilder.RHICmdList);
+		{
+			TaskData.GatherDynamicRayTracingInstancesTask.Wait();
+			TaskData.DynamicRayTracingInstancesContext->GatherDynamicRayTracingInstances_RenderThread();
+			TaskData.DynamicRayTracingInstancesContext->Finish(GraphBuilder.RHICmdList);
+
+			for (FGatherInstancesViewTaskData& ViewTaskData : TaskData.ViewTaskDatas)
+			{
+				ViewTaskData.Scene.GPUScene.UploadDynamicPrimitiveShaderDataForView(GraphBuilder, ViewTaskData.View, /*bRayTracing*/ true);
+			}
+
+			TaskData.DynamicRayTracingInstancesContext->CollectRDGResources(GraphBuilder);
+		}
 
-		Scene.GPUScene.UploadDynamicPrimitiveShaderDataForView(GraphBuilder, View, /*bRayTracing*/ true);
+		bool bAnyViewLightingChannelsUsingAHS = false;
 
-		TaskData.DynamicRayTracingInstancesContext.CollectRDGResources(GraphBuilder);
+		UE::Tasks::FTaskEvent AddInstancesTaskEvent{ UE_SOURCE_LOCATION };
 
-		// This adds final instances to the ray tracing scene and must be done before FRayTracingScene::BuildInitializationData().
-		UE::Tasks::FTask AddInstancesTask = GraphBuilder.AddSetupTask(
-			[&TaskData, &Scene, &View, &RayTracingScene, &RayTracingSBT, SharedBufferGenerationID]
+		// This adds final dynamic instances to the ray tracing scene and must be done before FRayTracingScene::BuildInitializationData().
+		AddInstancesTaskEvent.AddPrerequisites(TaskData.AddInstancesPipe.Launch(UE_SOURCE_LOCATION,
+			[&TaskData, &RayTracingScene, &RayTracingSBT, SharedBufferGenerationID]
 			{
 				FTaskTagScope TaskTagScope(ETaskTag::EParallelRenderingThread);
 
 				AddDynamicInstancesToRayTracingScene(TaskData, RayTracingScene, RayTracingSBT, SharedBufferGenerationID);
+			}));
+
+		for (FGatherInstancesViewTaskData& ViewTaskData : TaskData.ViewTaskDatas)
+		{
+			// This adds final instances to the ray tracing scene and must be done before FRayTracingScene::BuildInitializationData().
+			AddInstancesTaskEvent.AddPrerequisites(TaskData.AddInstancesPipe.Launch(UE_SOURCE_LOCATION,
+				[&ViewTaskData, &RayTracingScene, &RayTracingSBT, SharedBufferGenerationID]
+				{
+					FTaskTagScope TaskTagScope(ETaskTag::EParallelRenderingThread);
+
+					RayTracingScene.bUsesLightingChannels |= ViewTaskData.bUsesLightingChannels && ViewTaskData.SceneOptions.bLightingChannelsUsingAHS;
 
-				RayTracingScene.bUsesLightingChannels |= TaskData.bUsesLightingChannels && TaskData.SceneOptions.bLightingChannelsUsingAHS;
+					ViewTaskData.View.VisibleRayTracingShaderBindings.Reserve(ViewTaskData.StaticPrimitives.Num() + ViewTaskData.NumCachedStaticVisibleShaderBindings);
 
-				View.VisibleRayTracingShaderBindings.Reserve(TaskData.StaticPrimitives.Num() + TaskData.NumCachedStaticVisibleShaderBindings);
+					AddStaticInstancesToRayTracingScene(
+						ViewTaskData.Scene,
+						ViewTaskData.View,
+						ViewTaskData.SceneOptions,
+						ViewTaskData.StaticPrimitives,
+						RayTracingScene,
+						ViewTaskData.View.VisibleRayTracingShaderBindings);
 
-				AddStaticInstancesToRayTracingScene(
-					Scene,
-					View,
-					TaskData.SceneOptions,
-					TaskData.StaticPrimitives,
-					RayTracingScene,
-					View.VisibleRayTracingShaderBindings);
+					AddVisibleCachedInstances(
+						ViewTaskData.View,
+						ViewTaskData.Scene.RayTracingScene,
+						ViewTaskData.Scene.RayTracingSBT,
+						ViewTaskData.CachedStaticPrimitives,
+						ViewTaskData.View.VisibleRayTracingShaderBindings);
 
-				AddVisibleCachedInstances(
-					View,
-					RayTracingScene,
-					RayTracingSBT,
-					TaskData.CachedStaticPrimitives,
-					View.VisibleRayTracingShaderBindings);
-			}, TaskData.GatherRelevantStaticPrimitivesTask);
+				}, ViewTaskData.GatherRelevantStaticPrimitivesTask));
+
+			bAnyViewLightingChannelsUsingAHS |= ViewTaskData.SceneOptions.bLightingChannelsUsingAHS;
+		}
+
+		AddInstancesTaskEvent.Trigger();
 
 		// Scene init task can run only when all pre-init tasks are complete (including culling tasks that are spawned while adding instances)
-		RayTracingScene.InitTask = GraphBuilder.AddSetupTask(
-			[&RayTracingScene, bLightingChannelsUsingAHS = TaskData.SceneOptions.bLightingChannelsUsingAHS]
+		RayTracingScene.InitTask = UE::Tasks::Launch(UE_SOURCE_LOCATION,
+			[&RayTracingScene, bLightingChannelsUsingAHS = bAnyViewLightingChannelsUsingAHS]()
 			{
 				FTaskTagScope TaskTagScope(ETaskTag::EParallelRenderingThread);
 
 				TRACE_CPUPROFILER_EVENT_SCOPE(RayTracingSceneInitTask);
 				RayTracingScene.BuildInitializationData(bLightingChannelsUsingAHS, GRayTracingDebugForceOpaque != 0, GRayTracingDebugDisableTriangleCull != 0);
-			}, AddInstancesTask);
+			}, AddInstancesTaskEvent);
 
 		// Finalizing VisibleRayTracingShaderBindings can run only when all pre-init tasks are complete (including culling tasks that are spawned while adding instances)
-		TaskData.VisibleRayTracingShaderBindingsFinalizeTask = GraphBuilder.AddSetupTask(
-			[&TaskData, &View, &RayTracingSBT]
+		TaskData.VisibleRayTracingShaderBindingsFinalizeTask = UE::Tasks::Launch(UE_SOURCE_LOCATION,
+			[&TaskData, &RayTracingSBT]
 			{
 				FTaskTagScope TaskTagScope(ETaskTag::EParallelRenderingThread);
 
 				{
 					TRACE_CPUPROFILER_EVENT_SCOPE(DeduplicateVisibleShaderBindings);
 
+					int32 TotalNumBindings = 0;
+
+					for (const FGatherInstancesViewTaskData& ViewTaskData : TaskData.ViewTaskDatas)
+					{
+						TotalNumBindings += ViewTaskData.View.VisibleRayTracingShaderBindings.Num();
+					}
+
 					// Deduplicate all the written SBT record indices by using bit array and checking the written indices into the SBT table
 					TBitArray<> ProcessedSBTAllocations(false, RayTracingSBT.GetNumGeometrySegments() * RAY_TRACING_NUM_SHADER_SLOTS);
 					TArray<FRayTracingShaderBindingData> DeduplicatedVisibleShaderBindingData;
-					DeduplicatedVisibleShaderBindingData.Reserve(View.VisibleRayTracingShaderBindings.Num());
-					for (FRayTracingShaderBindingData& VisibleBinding : View.VisibleRayTracingShaderBindings)
+					DeduplicatedVisibleShaderBindingData.Reserve(TotalNumBindings);
+
+					for (const FGatherInstancesViewTaskData& ViewTaskData : TaskData.ViewTaskDatas)
 					{
-						FBitReference BitReference = ProcessedSBTAllocations[VisibleBinding.SBTRecordIndex];
-						if (!BitReference)
+						for (FRayTracingShaderBindingData& VisibleBinding : ViewTaskData.View.VisibleRayTracingShaderBindings)
 						{
-							BitReference = true;
-							DeduplicatedVisibleShaderBindingData.Add(VisibleBinding);
+							FBitReference BitReference = ProcessedSBTAllocations[VisibleBinding.SBTRecordIndex];
+							if (!BitReference)
+							{
+								BitReference = true;
+								DeduplicatedVisibleShaderBindingData.Add(VisibleBinding);
+							}
 						}
 					}
 					TaskData.VisibleShaderBindings = MoveTemp(DeduplicatedVisibleShaderBindingData);
@@ -1871,7 +2092,7 @@ namespace RayTracing
 
 				// reset static allocation lock again now that the final array of dirty shader bindings is build
 				RayTracingSBT.ResetStaticAllocationLock();
-			}, AddInstancesTask);
+			}, AddInstancesTaskEvent);
 
 		// wait for this task here, although it could be done later in the frame
 		// since it's only consumed by FCoarseMeshStreamingManager::UpdateResourceStates() during PostRenderAllViewports_RenderThread
@@ -1880,6 +2101,15 @@ namespace RayTracing
 		return true;
 	}
 
+	void WaitForDynamicBindings(FGatherInstancesTaskData& TaskData)
+	{
+		for (FViewInfo* View : TaskData.Views)
+		{
+			UE::Tasks::Wait(View->AddDynamicRayTracingMeshBatchTaskList);
+			View->AddDynamicRayTracingMeshBatchTaskList.Empty();
+		}
+	}
+
 	bool FinishGatherVisibleShaderBindings(FGatherInstancesTaskData& TaskData)
 	{
 		TRACE_CPUPROFILER_EVENT_SCOPE(RayTracing_FinishGatherVisibleShaderBindings);
@@ -1887,15 +2117,23 @@ namespace RayTracing
 		TaskData.VisibleRayTracingShaderBindingsFinalizeTask.Wait();
 
 		// merge dynamic bindings
+		for (FGatherInstancesViewTaskData& ViewTaskData : TaskData.ViewTaskDatas)
 		{
-			for (int32 TaskIndex = 0; TaskIndex < TaskData.View.DynamicRayTracingShaderBindingsPerTask.Num(); TaskIndex++)
+			for (int32 TaskIndex = 0; TaskIndex < ViewTaskData.View.DynamicRayTracingShaderBindingsPerTask.Num(); TaskIndex++)
 			{
-				TaskData.VisibleShaderBindings.Append(*TaskData.View.DynamicRayTracingShaderBindingsPerTask[TaskIndex]);
+				ViewTaskData.View.VisibleRayTracingShaderBindings.Append(*ViewTaskData.View.DynamicRayTracingShaderBindingsPerTask[TaskIndex]);
 			}
 
-			TaskData.View.DynamicRayTracingShaderBindingsPerTask.Empty();
+			ViewTaskData.View.DynamicRayTracingShaderBindingsPerTask.Empty();
+
+			TaskData.VisibleShaderBindings.Append(ViewTaskData.View.VisibleRayTracingShaderBindings);
 		}
 
+		// Even though task dependencies are setup so all work is done by this point, we still have to wait on the
+		// pipe to clear out its internal state. Otherwise it can assert that it still has work at shutdown.
+		TaskData.AddInstancesPipe.WaitUntilEmpty();
+		TaskData.FinalizeGatherRelevantPrimitivesPipe.WaitUntilEmpty();
+
 		return true;
 	}
 
diff --git a/Engine/Source/Runtime/Renderer/Private/RayTracing/RayTracing.h b/Engine/Source/Runtime/Renderer/Private/RayTracing/RayTracing.h
--- a/Engine/Source/Runtime/Renderer/Private/RayTracing/RayTracing.h
+++ b/Engine/Source/Runtime/Renderer/Private/RayTracing/RayTracing.h
@@ -49,8 +49,11 @@ namespace RayTracing
 	FGatherInstancesTaskData* CreateGatherInstancesTaskData(
 		FSceneRenderingBulkObjectAllocator& InAllocator,
 		FScene& Scene,
+		uint32 NumViews);
+
+	void AddView(
+		FGatherInstancesTaskData& TaskData,
 		FViewInfo& View,
-		const FViewFamilyInfo& ViewFamily,
 		EDiffuseIndirectMethod DiffuseIndirectMethod,
 		EReflectionsMethod ReflectionsMethod);
 	
@@ -68,6 +71,8 @@ namespace RayTracing
 		FGlobalDynamicReadBuffer& InDynamicReadBuffer,
 		FSceneRenderingBulkObjectAllocator& InBulkAllocator);
 
+	void WaitForDynamicBindings(FGatherInstancesTaskData& TaskData);
+
 	// Wait for shader bindings related tasks to complete and combine bindings from all views
 	// This function must run on render thread
 	bool FinishGatherVisibleShaderBindings(FGatherInstancesTaskData& TaskData);
