[HWRT] Fix incorrect DistanceToView calculation across multiple views.

diff --git a/Engine/Source/Runtime/Renderer/Private/RayTracing/RayTracing.cpp b/Engine/Source/Runtime/Renderer/Private/RayTracing/RayTracing.cpp
--- a/Engine/Source/Runtime/Renderer/Private/RayTracing/RayTracing.cpp
+++ b/Engine/Source/Runtime/Renderer/Private/RayTracing/RayTracing.cpp
@@ -1121,7 +1121,7 @@ namespace RayTracing
 				if (LastRenderTimeUpdateDistance > 0.0f)
 				{
 					float CurrentWorldTime = 0;
-					double DistanceToView = std::numeric_limits<double>::infinity();
+					double DistanceToView = 0;
 
 					for (int32 ViewIndex = 0; ViewIndex < Views.Num(); ++ViewIndex)
 					{
