[HWRT] Fix incorrect DistanceToView calculation across multiple views.

diff --git a/Engine/Source/Runtime/Renderer/Private/RayTracing/RayTracing.cpp b/Engine/Source/Runtime/Renderer/Private/RayTracing/RayTracing.cpp
--- a/Engine/Source/Runtime/Renderer/Private/RayTracing/RayTracing.cpp
+++ b/Engine/Source/Runtime/Renderer/Private/RayTracing/RayTracing.cpp
@@ -1121,7 +1121,7 @@ namespace RayTracing
 				if (LastRenderTimeUpdateDistance > 0.0f)
 				{
 					float CurrentWorldTime = 0;
-					double DistanceToView = 0;
+					double DistanceToView = std::numeric_limits<double>::infinity();
 
 					for (int32 ViewIndex = 0; ViewIndex < Views.Num(); ++ViewIndex)
 					{
@@ -1131,7 +1131,7 @@ namespace RayTracing
 
 						if (LastRenderTimeUpdateDistance > 0.0f)
 						{
-							DistanceToView = FMath::Max(DistanceToView, FVector::Distance(SceneProxy->GetActorPosition(), View.ViewMatrices.GetViewOrigin()));
+							DistanceToView = FMath::Min(DistanceToView, FVector::Distance(SceneProxy->GetActorPosition(), View.ViewMatrices.GetViewOrigin()));
 						}
 					}
 
