Adaptive virtual texure: Separate logic for when to process an adaptive request and when to continue processing the standard VT feedback request.
#rnx

diff --git a/Engine/Source/Runtime/Renderer/Private/VT/VirtualTextureSystem.cpp b/Engine/Source/Runtime/Renderer/Private/VT/VirtualTextureSystem.cpp
--- a/Engine/Source/Runtime/Renderer/Private/VT/VirtualTextureSystem.cpp
+++ b/Engine/Source/Runtime/Renderer/Private/VT/VirtualTextureSystem.cpp
@@ -1631,12 +1631,21 @@ void FVirtualTextureSystem::GatherRequestsTask(const FGatherRequestsParameters&
 		const FAdaptiveVirtualTexture* RESTRICT AdaptiveVT = AdaptiveVTs[ID];
 		if (AdaptiveVT != nullptr && vLevelPlusOne <= 1)
 		{
+			// If vLevelPlusOne == 0 this is a request for higher resolution adaptive VT.
+			// If vLevelPlusOne == 1 then this indicates the current resolution is correct, and so we mark to keep residency.
+			// (In future we could consider 2, or even higher, as being enough to keep current residency to reduce adaptive reallocation).
 			uint32 AdaptiveAllocationRequest = AdaptiveVT->GetPackedAllocationRequest(vAddress, vLevelPlusOne, Frame);
-			if (AdaptiveAllocationRequest == 0)
+			if (AdaptiveAllocationRequest != 0)
 			{
+				// This is a valid request to pass on to the adaptive VT.
+				RequestList->AddAdaptiveAllocationRequest(AdaptiveAllocationRequest);
+			}
+			if (vLevelPlusOne == 0)
+			{
+				// The feedback is for a page that won't exist until the adaptive VT is reallocated.
+				// So we can skip further processing here.
 				continue;
 			}
-			RequestList->AddAdaptiveAllocationRequest(AdaptiveAllocationRequest);
 		}
 
 		uint32 PageTableLayersToLoad[VIRTUALTEXTURE_SPACE_MAXLAYERS] = { 0 };
