[VSM] Decrease font size of cluster light stats

diff --git a/Engine/Source/Runtime/Engine/Classes/Engine/Font.h b/Engine/Source/Runtime/Engine/Classes/Engine/Font.h
--- a/Engine/Source/Runtime/Engine/Classes/Engine/Font.h
+++ b/Engine/Source/Runtime/Engine/Classes/Engine/Font.h
@@ -219,10 +219,7 @@ public:
 	}
 
 	/** Get the info needed to use this UFont with Slate, using the fallback data for legacy Canvas APIs */
-	inline FSlateFontInfo GetLegacySlateFontInfo() const
-	{
-		return FSlateFontInfo(this, static_cast<float>(LegacyFontSize), LegacyFontName);
-	}
+	ENGINE_API FSlateFontInfo GetLegacySlateFontInfo() const;
 
 	/**
 	 * Returns the size of the object/ resource for display to artists/ LDs in the Editor.
diff --git a/Engine/Source/Runtime/Engine/Private/Font.cpp b/Engine/Source/Runtime/Engine/Private/Font.cpp
--- a/Engine/Source/Runtime/Engine/Private/Font.cpp
+++ b/Engine/Source/Runtime/Engine/Private/Font.cpp
@@ -355,6 +355,11 @@ void UFont::GetStringHeightAndWidth( const TCHAR *Text, int32& Height, int32& Wi
 	Height = FMath::CeilToInt( TotalHeight );
 }
 
+FSlateFontInfo UFont::GetLegacySlateFontInfo() const
+{
+	return FSlateFontInfo(this, static_cast<float>(LegacyFontSize), LegacyFontName);
+}
+
 void UFont::GetResourceSizeEx(FResourceSizeEx& CumulativeResourceSize)
 {
 	Super::GetResourceSizeEx(CumulativeResourceSize);
diff --git a/Engine/Source/Runtime/Renderer/Private/VirtualShadowMaps/VirtualShadowMapCacheManager.cpp b/Engine/Source/Runtime/Renderer/Private/VirtualShadowMaps/VirtualShadowMapCacheManager.cpp
--- a/Engine/Source/Runtime/Renderer/Private/VirtualShadowMaps/VirtualShadowMapCacheManager.cpp
+++ b/Engine/Source/Runtime/Renderer/Private/VirtualShadowMaps/VirtualShadowMapCacheManager.cpp
@@ -956,12 +956,12 @@ void FVirtualShadowMapArrayCacheManager::InitExtension(FScene& InScene)
 						FSceneRenderer::GetLightNameForDrawEvent(LightSceneInfo->Proxy, LightNameWithLevel);
 					}
 				}
-				NPFDiagnosticMessages.Add(FString::Printf(TEXT("%8d | %48s | %12d"), VSMId, *LightNameWithLevel, ClusterCount));
+				NPFDiagnosticMessages.Add(FString::Printf(TEXT("%6d | %100s | %10d"), VSMId, *LightNameWithLevel, ClusterCount));
 			}
 		};
 
 		NPFDiagnosticMessages.SetNum(0, EAllowShrinking::No);
-		NPFDiagnosticMessages.Add(FString::Printf(TEXT("%8s | %48s | %12s"), TEXT("VSM Id"), TEXT("Light"), TEXT("# Clusters")));
+		NPFDiagnosticMessages.Add(FString::Printf(TEXT("%6s | %100s | %10s"), TEXT("VSM Id"), TEXT("Light"), TEXT("# Clusters")));
 		TConstArrayView<uint32> NPFDiagsSinglePage = Message.ReadCount(FVirtualShadowMapArray::MaxNPFDiagnosticSlotsSinglePage * 2U);
 		for (int32 Index = 0; Index < NPFDiagsSinglePage.Num(); Index += 2)
 		{
@@ -1066,22 +1066,34 @@ FScreenPassTexture FVirtualShadowMapArrayCacheManager::AddVisualizePass(FRDGBuil
 					Canvas.DrawTile(X / DPIScale, Y / DPIScale, Width / DPIScale, Height / DPIScale, 0, 0, 0, 0, Color);
 				};
 
-				auto MeasureString = [&](FStringView Text, const UFont* Font = GetStatsFont())
+				auto MeasureString = [&](FStringView Text, int FontSize = -1, const UFont* Font = GetStatsFont())
 				{
 					FCanvasTextStringViewItem TextItem(FVector2D::ZeroVector, Text, Font, FLinearColor::Transparent);
 					
 					TextItem.bOutlined = true;
 					TextItem.OutlineColor = FLinearColor::Transparent;
+					if (FontSize > 0)
+					{
+						TextItem.SlateFontInfo = Font->GetLegacySlateFontInfo();
+						TextItem.SlateFontInfo->Size = FontSize;
+					}
+
 					Canvas.DrawItem(TextItem);
 					return TextItem.DrawnSize;
 				};
 
-				auto DrawString = [&](float X, float Y, FStringView Text, const FLinearColor& Color = FLinearColor::White, const UFont* Font = GetStatsFont())
+				auto DrawString = [&](float X, float Y, FStringView Text, int FontSize = -1, const FLinearColor& Color = FLinearColor::White, const UFont* Font = GetStatsFont())
 				{
 					FCanvasTextStringViewItem TextItem(FVector2D(X / DPIScale, Y / DPIScale), Text, Font, Color);
 					
 					TextItem.bOutlined = true;
 					TextItem.OutlineColor = FLinearColor::Black;
+					if (FontSize > 0)
+					{
+						TextItem.SlateFontInfo = Font->GetLegacySlateFontInfo();
+						TextItem.SlateFontInfo->Size = FontSize;
+					}
+
 					Canvas.DrawItem(TextItem);
 					return TextItem.DrawnSize;
 				};
@@ -1094,7 +1106,7 @@ FScreenPassTexture FVirtualShadowMapArrayCacheManager::AddVisualizePass(FRDGBuil
 						MessageBoxCharacters.X = FMath::Max(MessageBoxCharacters.X, Message.Len());
 					}
 
-					FVector2D CharSize = MeasureString(TEXT("A"), GEngine->GetMonospaceFont());
+					FVector2D CharSize = MeasureString(TEXT("A"), 8, GEngine->GetMonospaceFont());
 					MessageBoxSize = CharSize * MessageBoxCharacters + FVector2D(10, 10);
 				}
 				
@@ -1107,7 +1119,7 @@ FScreenPassTexture FVirtualShadowMapArrayCacheManager::AddVisualizePass(FRDGBuil
 				int Y = 0;
 				for (const FString& Message : NPFDiagnosticMessages)
 				{
-					Y += DrawString(BaselineX, BaselineY + Y, Message, FLinearColor::White, GEngine->GetMonospaceFont()).Y;
+					Y += DrawString(BaselineX, BaselineY + Y, Message, 8, FLinearColor::White, GEngine->GetMonospaceFont()).Y;
 				}
 			});
 		NPFDiagnosticMessages.SetNum(0, EAllowShrinking::No);
