Draw InstanceCullingOcclusionQuery debug view into debug primitives to avoid TAA smearing. Restructured EditorPrimitives and DebugPrimitives. DebugPrimitives now supports MSAA and partial occlusion of hidden lines during compositing.

diff --git a/Engine/Shaders/Private/InstanceCulling/InstanceCullingOcclusionQuery.usf b/Engine/Shaders/Private/InstanceCulling/InstanceCullingOcclusionQuery.usf
--- a/Engine/Shaders/Private/InstanceCulling/InstanceCullingOcclusionQuery.usf
+++ b/Engine/Shaders/Private/InstanceCulling/InstanceCullingOcclusionQuery.usf
@@ -318,20 +318,20 @@ void DebugMainPS(
 	case IVS_Hidden:
 		// Definitely invisible due to HZB test.
 		// No occlusion query performed.
-		OutColor = float4(0.1, 0.0, 0.0, 0.0); // red
+		OutColor = float4(1.0, 0.0, 0.0, 0.5); // red
 		break;
 	case IVS_Visible:
 		// Definitely visible due to near plane intersection or other heuristics.
 		// No occlusion query performed.
-		OutColor = float4(0.1, 0.1, 0.0, 0.0); // yellow
+		OutColor = float4(1.0, 1.0, 0.0, 0.5); // yellow
 		break;
 	case IVS_PossiblyVisible:
 		// Visible due to occlusion query.
-		OutColor = float4(0.0, 0.025, 0.0, 0.0); // green
+		OutColor = float4(0.0, 1.0, 0.0, 0.2); // green
 		break;
 	case IVS_Incompatible:
 		// Invisible due to occlusion query.
-		OutColor = float4(0.25, 0.0, 0.25, 0.0); // magenta
+		OutColor = float4(1.0, 0.0, 1.0, 0.5); // magenta
 		break;
 	default:
 		OutColor = (float4)0;
diff --git a/Engine/Shaders/Private/PostProcessCompositePrimitives.usf b/Engine/Shaders/Private/PostProcessCompositePrimitives.usf
--- a/Engine/Shaders/Private/PostProcessCompositePrimitives.usf
+++ b/Engine/Shaders/Private/PostProcessCompositePrimitives.usf
@@ -306,7 +306,7 @@ void ResolveEditorPrimitiveColor(int2 PixelPos, out float4 OutColor, out float O
 	OutColor.rgb *= OutColor.a;
 }
 
-void MainCompositeEditorPrimitivesPS(
+void MainCompositePostProcessPrimitivesPS(
 	float4 SvPosition : SV_POSITION,
 #if MSAA_DONT_RESOLVE && MSAA_SAMPLE_COUNT > 1
 	uint TargetSampleIndex : SV_SampleIndex,
diff --git a/Engine/Source/Runtime/Renderer/Private/DeferredShadingRenderer.cpp b/Engine/Source/Runtime/Renderer/Private/DeferredShadingRenderer.cpp
--- a/Engine/Source/Runtime/Renderer/Private/DeferredShadingRenderer.cpp
+++ b/Engine/Source/Runtime/Renderer/Private/DeferredShadingRenderer.cpp
@@ -3699,15 +3699,6 @@ void FDeferredShadingSceneRenderer::Render(FRDGBuilder& GraphBuilder, const FSce
 			ShadingEnergyConservation::Debug(GraphBuilder, View, SceneTextures);
 		}
 
-		if (ViewFamily.EngineShowFlags.VisualizeInstanceOcclusionQueries
-			&& Scene->InstanceCullingOcclusionQueryRenderer)
-		{
-			for (FViewInfo& View : Views)
-			{
-				Scene->InstanceCullingOcclusionQueryRenderer->RenderDebug(GraphBuilder, Scene->GPUScene, View, SceneTextures);
-			}
-		}
-
 		if (!bHasRayTracedOverlay && ViewFamily.EngineShowFlags.LightShafts)
 		{
 			SCOPE_CYCLE_COUNTER(STAT_FDeferredShadingSceneRenderer_RenderLightShaftBloom);
diff --git a/Engine/Source/Runtime/Renderer/Private/InstanceCulling/InstanceCullingOcclusionQuery.cpp b/Engine/Source/Runtime/Renderer/Private/InstanceCulling/InstanceCullingOcclusionQuery.cpp
--- a/Engine/Source/Runtime/Renderer/Private/InstanceCulling/InstanceCullingOcclusionQuery.cpp
+++ b/Engine/Source/Runtime/Renderer/Private/InstanceCulling/InstanceCullingOcclusionQuery.cpp
@@ -878,6 +878,7 @@ END_SHADER_PARAMETER_STRUCT()
 static void RenderInstanceOcclusionCullingDebug(
 	FRHICommandList& RHICmdList,
 	const FViewInfo& View,
+	const FIntRect& ViewRect,
 	FOcclusionInstanceCullingDebugParameters* PassParameters,
 	int32 NumInstances)
 {
@@ -887,14 +888,14 @@ static void RenderInstanceOcclusionCullingDebug(
 	FGraphicsPipelineStateInitializer GraphicsPSOInit;
 	RHICmdList.ApplyCachedRenderTargets(GraphicsPSOInit);
 
-	FIntVector4 ViewRect = FIntVector4(View.ViewRect.Min.X, View.ViewRect.Min.Y, View.ViewRect.Max.X, View.ViewRect.Max.Y);
-	RHICmdList.SetViewport(ViewRect.X, ViewRect.Y, 0.0f, ViewRect.Z, ViewRect.W, 1.0f);
+	RHICmdList.SetViewport(ViewRect.Min.X, ViewRect.Min.Y, 0.0f, ViewRect.Max.X, ViewRect.Max.Y, 1.0f);
 
 	GraphicsPSOInit.BoundShaderState.VertexDeclarationRHI = GInstanceCullingOcclusionQueryBox.VertexDeclaration;
 	GraphicsPSOInit.BoundShaderState.VertexShaderRHI = VertexShader.GetVertexShader();
 	GraphicsPSOInit.BoundShaderState.PixelShaderRHI = PixelShader.GetPixelShader();
-	GraphicsPSOInit.DepthStencilState = TStaticDepthStencilState<false, CF_Always>::GetRHI(); // No depth test or write
-	GraphicsPSOInit.BlendState = TStaticBlendState<CW_RGBA, BO_Add, BF_One, BF_InverseSourceAlpha, BO_Add, BF_Zero, BF_One>::GetRHI(); // Premultiplied
+	GraphicsPSOInit.DepthStencilState = TStaticDepthStencilState<true, CF_DepthNearOrEqual>::GetRHI();
+	GraphicsPSOInit.DepthStencilAccess = FExclusiveDepthStencil::DepthWrite_StencilNop;
+	GraphicsPSOInit.BlendState = TStaticBlendState<CW_RGBA, BO_Add, BF_SourceAlpha, BF_InverseSourceAlpha, BO_Add, BF_SourceAlpha, BF_DestAlpha>::GetRHI();
 	GraphicsPSOInit.RasterizerState = TStaticRasterizerState<FM_Solid, CM_None>::GetRHI();
 	GraphicsPSOInit.PrimitiveType = PT_LineList;
 
@@ -911,17 +912,17 @@ static void RenderInstanceOcclusionCullingDebug(
 	RHICmdList.DrawIndexedPrimitive(GInstanceCullingOcclusionQueryBox.IndexBuffer, 0, 0, 24, 36, 12, NumInstances);
 }
 
-void FInstanceCullingOcclusionQueryRenderer::RenderDebug(FRDGBuilder& GraphBuilder, FGPUScene& GPUScene, const FViewInfo& View, FSceneTextures& SceneTextures)
+void FInstanceCullingOcclusionQueryRenderer::RenderDebug(FRDGBuilder& GraphBuilder, FGPUScene& GPUScene, const FViewInfo& View, const FIntRect& ViewRect, FRDGTextureRef SceneColor, FRDGTextureRef SceneDepth)
 {
 	if (!IsCompatibleWithView(View) || !InstanceOcclusionQueryBuffer)
 	{
 		return;
 	}
 
+	check(SceneColor);
+
 	const uint32 ViewMask = FindOrAddViewSlot(View);
 
-	FRDGTextureRef SceneColor = View.GetSceneTextures().Color.Target;
-	FRDGTextureRef SceneDepth = View.GetSceneTextures().Depth.Target;
 	FRDGBufferRef InstanceOcclusionQueryBufferRDG = GraphBuilder.RegisterExternalBuffer(InstanceOcclusionQueryBuffer);
 
 	FRDGTextureRef DepthTexture = View.GetSceneTextures().Depth.Target;
@@ -932,8 +933,6 @@ void FInstanceCullingOcclusionQueryRenderer::RenderDebug(FRDGBuilder& GraphBuild
 	const int32 NumInstances = GPUScene.GetNumInstances();
 	const FGPUSceneResourceParameters GPUSceneParameters = GPUScene.GetShaderParameters(GraphBuilder);
 
-	const FIntPoint ViewRectSize = View.ViewRect.Size();
-
 	FOcclusionInstanceCullingDebugParameters* PassParameters = GraphBuilder.AllocParameters<FOcclusionInstanceCullingDebugParameters>();
 
 	PassParameters->VS.OcclusionSlop = OCCLUSION_SLOP;
@@ -943,17 +942,19 @@ void FInstanceCullingOcclusionQueryRenderer::RenderDebug(FRDGBuilder& GraphBuild
 	PassParameters->VS.HZBParameters = GetHZBParameters(GraphBuilder, View, EHZBType::FurthestHZB);
 	PassParameters->VS.ViewMask = ViewMask;
 	PassParameters->RenderTargets[0] = FRenderTargetBinding(SceneColor, ERenderTargetLoadAction::ELoad);
-	PassParameters->RenderTargets.DepthStencil = FDepthStencilBinding(SceneDepth,
-		ERenderTargetLoadAction::ELoad, ERenderTargetLoadAction::ENoAction,
-		FExclusiveDepthStencil::DepthRead_StencilNop);
+	if (SceneDepth)
+	{
+		PassParameters->RenderTargets.DepthStencil = FDepthStencilBinding(SceneDepth,
+			ERenderTargetLoadAction::ELoad, ERenderTargetLoadAction::ENoAction,
+			FExclusiveDepthStencil::DepthWrite_StencilNop);
+	}
 
 	GraphBuilder.AddPass(
 		RDG_EVENT_NAME("InstanceCullingOcclusionQueryRenderer_Draw"),
 		PassParameters, ERDGPassFlags::Raster | ERDGPassFlags::NeverCull,
-		[PassParameters, NumInstances, &View](FRDGAsyncTask, FRHICommandList& RHICmdList)
+		[PassParameters, NumInstances, &View, ViewRect](FRDGAsyncTask, FRHICommandList& RHICmdList)
 		{
-			RenderInstanceOcclusionCullingDebug(RHICmdList, View, PassParameters, NumInstances);
+			RenderInstanceOcclusionCullingDebug(RHICmdList, View, ViewRect, PassParameters, NumInstances);
 		});
-
 }
 
diff --git a/Engine/Source/Runtime/Renderer/Private/InstanceCulling/InstanceCullingOcclusionQuery.h b/Engine/Source/Runtime/Renderer/Private/InstanceCulling/InstanceCullingOcclusionQuery.h
--- a/Engine/Source/Runtime/Renderer/Private/InstanceCulling/InstanceCullingOcclusionQuery.h
+++ b/Engine/Source/Runtime/Renderer/Private/InstanceCulling/InstanceCullingOcclusionQuery.h
@@ -30,7 +30,7 @@ public:
 	/**
 	* Render object bounding boxes, color-coded based on visibility
 	*/
-	void RenderDebug(FRDGBuilder& GraphBuilder, FGPUScene& GPUScene, const FViewInfo& View, FSceneTextures& SceneTextures);
+	void RenderDebug(FRDGBuilder& GraphBuilder, FGPUScene& GPUScene, const FViewInfo& View, const FIntRect& ViewRect, FRDGTextureRef SceneColor, FRDGTextureRef SceneDepth);
 
 	/**
 	* Mark blocks of occlusion query buffer entries as visible in all views.
diff --git a/Engine/Source/Runtime/Renderer/Private/MobileBasePassRendering.cpp b/Engine/Source/Runtime/Renderer/Private/MobileBasePassRendering.cpp
--- a/Engine/Source/Runtime/Renderer/Private/MobileBasePassRendering.cpp
+++ b/Engine/Source/Runtime/Renderer/Private/MobileBasePassRendering.cpp
@@ -633,14 +633,17 @@ void FMobileSceneRenderer::RenderMobileDebugPrimitives(FRHICommandList& RHICmdLi
 	QUICK_SCOPE_CYCLE_COUNTER(STAT_DebugDynamicPrimitiveDrawTime);
 	SCOPED_DRAW_EVENT(RHICmdList, DynamicDebug);
 
-	FMeshPassProcessorRenderState DrawRenderState;
-	DrawRenderState.SetDepthStencilAccess(FExclusiveDepthStencil::DepthWrite_StencilWrite);
-	DrawRenderState.SetBlendState(TStaticBlendState<CW_RGBA, BO_Add, BF_One, BF_One, BO_Add, BF_One, BF_One>::GetRHI());
-		
-	DrawRenderState.SetDepthStencilState(TStaticDepthStencilState<true, CF_DepthNearOrEqual>::GetRHI());
-	View.DebugSimpleElementCollector.DrawBatchedElements(RHICmdList, DrawRenderState, View, EBlendModeFilter::OpaqueAndMasked, SDPG_World);
+	if (View.DebugSimpleElementCollector.HasAnyPrimitives())
+	{
+		FMeshPassProcessorRenderState DrawRenderState;
+		DrawRenderState.SetDepthStencilAccess(FExclusiveDepthStencil::DepthWrite_StencilWrite);
+		DrawRenderState.SetBlendState(TStaticBlendState<CW_RGBA, BO_Add, BF_One, BF_One, BO_Add, BF_One, BF_One>::GetRHI());
+
+		DrawRenderState.SetDepthStencilState(TStaticDepthStencilState<true, CF_DepthNearOrEqual>::GetRHI());
+		View.DebugSimpleElementCollector.DrawBatchedElements(RHICmdList, DrawRenderState, View, EBlendModeFilter::OpaqueAndMasked, SDPG_World);
 
-	DrawRenderState.SetDepthStencilState(TStaticDepthStencilState<true, CF_Always>::GetRHI());
-	View.DebugSimpleElementCollector.DrawBatchedElements(RHICmdList, DrawRenderState, View, EBlendModeFilter::OpaqueAndMasked, SDPG_Foreground);
+		DrawRenderState.SetDepthStencilState(TStaticDepthStencilState<true, CF_Always>::GetRHI());
+		View.DebugSimpleElementCollector.DrawBatchedElements(RHICmdList, DrawRenderState, View, EBlendModeFilter::OpaqueAndMasked, SDPG_Foreground);
+	}
 }
 #endif
\ No newline at end of file
diff --git a/Engine/Source/Runtime/Renderer/Private/MobileShadingRenderer.cpp b/Engine/Source/Runtime/Renderer/Private/MobileShadingRenderer.cpp
--- a/Engine/Source/Runtime/Renderer/Private/MobileShadingRenderer.cpp
+++ b/Engine/Source/Runtime/Renderer/Private/MobileShadingRenderer.cpp
@@ -82,6 +82,7 @@
 #include "DistanceFieldLightingShared.h"
 #include "DepthCopy.h"
 #include "SingleLayerWaterRendering.h"
+#include "PostProcess/PostProcessCompositeDebugPrimitives.h"
 
 uint32 GetShadowQuality();
 
@@ -1548,15 +1549,6 @@ void FMobileSceneRenderer::Render(FRDGBuilder& GraphBuilder, const FSceneRenderU
 		{
 			RenderMeshDistanceFieldVisualization(GraphBuilder, SceneTextures);
 		}
-	
-		if (ViewFamily.EngineShowFlags.VisualizeInstanceOcclusionQueries
-			&& Scene->InstanceCullingOcclusionQueryRenderer)
-		{
-			for (FViewInfo& View : Views)
-			{
-				Scene->InstanceCullingOcclusionQueryRenderer->RenderDebug(GraphBuilder, Scene->GPUScene, View, SceneTextures);
-			}
-		}
 
 		if (ViewFamily.bResolveScene)
 		{
@@ -1848,7 +1840,7 @@ void FMobileSceneRenderer::RenderForwardSinglePass(FRDGBuilder& GraphBuilder, FM
 		RenderTranslucency(RHICmdList, View, Views, StandardTranslucencyPass, StandardTranslucencyMeshPass, &ParameterCollection->TranslucencyInstanceCullingDrawParams);
 		
 #if UE_ENABLE_DEBUG_DRAWING
-		if ((!IsMobileHDR() || bTonemapSubpass) && FSceneRenderer::ShouldCompositeDebugPrimitivesInPostProcess(View))
+		if ((!IsMobileHDR() || bTonemapSubpass) && IsDebugPrimitivePassEnabled(View))
 		{
 			// Draw debug primitives after translucency for LDR as we do not have a post processing pass
 			RenderMobileDebugPrimitives(RHICmdList, View);
@@ -1874,6 +1866,13 @@ void FMobileSceneRenderer::RenderForwardSinglePass(FRDGBuilder& GraphBuilder, FM
 			RenderMobileCustomResolve(RHICmdList, View, NumMSAASamples, SceneTextures);
 		}
 	});
+
+#if UE_ENABLE_DEBUG_DRAWING
+	if (ViewFamily.EngineShowFlags.VisualizeInstanceOcclusionQueries && Scene->InstanceCullingOcclusionQueryRenderer)
+	{
+		Scene->InstanceCullingOcclusionQueryRenderer->RenderDebug(GraphBuilder, Scene->GPUScene, View, View.ViewRect, SceneTextures.Color.Resolve, nullptr);
+	}
+#endif
 	
 	// resolve MSAA depth
 	if (!GRHISupportsDepthStencilResolve && !bIsFullDepthPrepassEnabled)
diff --git a/Engine/Source/Runtime/Renderer/Private/PostProcess/PostProcessCompositeDebugPrimitives.cpp b/Engine/Source/Runtime/Renderer/Private/PostProcess/PostProcessCompositeDebugPrimitives.cpp
--- a/Engine/Source/Runtime/Renderer/Private/PostProcess/PostProcessCompositeDebugPrimitives.cpp
+++ b/Engine/Source/Runtime/Renderer/Private/PostProcess/PostProcessCompositeDebugPrimitives.cpp
@@ -3,74 +3,115 @@
 #include "PostProcess/PostProcessCompositeDebugPrimitives.h"
 #if UE_ENABLE_DEBUG_DRAWING
 #include "ScenePrivate.h"
+#include "PixelShaderUtils.h"
+#include "PostProcess/PostProcessing.h"
+
+#include "InstanceCulling/InstanceCullingOcclusionQuery.h"
 
 BEGIN_SHADER_PARAMETER_STRUCT(FDebugPrimitivesPassParameters, )
 	SHADER_PARAMETER_STRUCT_INCLUDE(FViewShaderParameters, View)
 	RENDER_TARGET_BINDING_SLOTS()
 END_SHADER_PARAMETER_STRUCT()
 
-FScreenPassTexture AddDebugPrimitivePass(
-	FRDGBuilder& GraphBuilder,
-	const FViewInfo& View,
-	const FCompositePrimitiveInputs& Inputs)
+struct FDebugPrimitivesDrawingContext
 {
-	check(Inputs.SceneColor.IsValid());
-	check(Inputs.SceneDepth.IsValid());
+	const FViewInfo* DebugView = nullptr;
+	FRDGTextureRef DebugPrimitiveColor = nullptr;
+	FRDGTextureRef DebugPrimitiveDepth = nullptr;
+};
 
-	RDG_EVENT_SCOPE(GraphBuilder, "CompositeDebugPrimitives");
-
-	FScreenPassRenderTarget Output = Inputs.OverrideOutput;
-	if (!Output.IsValid())
-	{
-		Output = FScreenPassRenderTarget::CreateFromInput(GraphBuilder, Inputs.SceneColor, View.GetOverwriteLoadAction(), TEXT("Debug.DrawPrimitivesColor"));
-	}
-	const uint32 NumMSAASamples = Output.Texture->Desc.NumSamples;
-	//Set sizing values to match the SceneColor size
+static FDebugPrimitivesDrawingContext InitializeDebugPrimitiveTextures(
+	FRDGBuilder& GraphBuilder, 
+	const FViewInfo& View,
+	const FCompositePrimitiveInputs& Inputs,
+	const FScreenPassRenderTarget& Output)
+{
+	// Setup view
 	const FIntRect ViewRect = Output.ViewRect;
 	FIntPoint Extent = Output.Texture->Desc.Extent;
+
+	FViewFamilyInfo* ViewFamily = (FViewFamilyInfo*)View.Family;
+	int NumMSAASamples = 1;
+#if WITH_EDITOR
+	NumMSAASamples = ViewFamily->SceneTexturesConfig.EditorPrimitiveNumSamples;
+#endif
+
 	const FViewInfo* DebugView = CreateCompositePrimitiveView(View, ViewRect, NumMSAASamples);	
 
-	//Prepare output textures for composite draw
-	const FScreenPassTextureViewport OutputViewport(Output);
-	FRDGTextureRef DebugPrimitiveDepth = CreateCompositeDepthTexture(GraphBuilder, Extent, NumMSAASamples);
-	//Inputs is const so create a over-ridable texture reference
-	FScreenPassTexture SceneDepth = Inputs.SceneDepth;
-	FVector2f SceneDepthJitter = FVector2f(View.TemporalJitterPixels);
+	// Setup textures
+	const FRDGTextureDesc ColorDesc(FRDGTextureDesc::Create2D(Extent, PF_B8G8R8A8, FClearValueBinding::Transparent, TexCreate_ShaderResource | TexCreate_RenderTargetable, 1, NumMSAASamples));
+	FRDGTextureRef DebugPrimitiveColor = GraphBuilder.CreateTexture(ColorDesc, TEXT("Debug.PrimitivesColor"));
+
+	const FRDGTextureDesc DepthDesc(FRDGTextureDesc::Create2D(Extent, PF_DepthStencil, FClearValueBinding::DepthFar, TexCreate_ShaderResource | TexCreate_DepthStencilTargetable, 1, NumMSAASamples));
+	FRDGTextureRef DebugPrimitiveDepth = GraphBuilder.CreateTexture(DepthDesc, TEXT("Debug.PrimitivesDepth"));
+
+	// Depth testing with scenedepth currently occurs after msaa resolve, so no depth copy/upscale here.
+	// This means we can perform depth sorting/testing amongst the debug primitives that are drawn, while
+	// 'disabling' depth testing of the result when compositing with the rest of the view, i.e. draw occluded lines.
+	const bool bCopyDepth = false;
+	if (bCopyDepth)
+	{
+		FScreenPassTexture UpsampledDepth = Inputs.SceneDepth;
+		FVector2f SceneDepthJitter = FVector2f(View.TemporalJitterPixels);
+
+		if (IsTemporalAccumulationBasedMethod(View.AntiAliasingMethod))
+		{
+			TemporalUpscaleDepthPass(GraphBuilder,
+									 *DebugView,
+									 Inputs.SceneColor,
+									 UpsampledDepth,
+									 SceneDepthJitter);
+		}
 
-	if (IsTemporalAccumulationBasedMethod(View.AntiAliasingMethod))
+		//Simple element pixel shaders do not output background color for composite, 
+		//so this allows the background to be drawn to the RT at the same time as depth without adding extra draw calls
+		PopulateDepthPass(GraphBuilder,
+						  *DebugView,
+						  Inputs.SceneColor,
+						  UpsampledDepth,
+						  DebugPrimitiveColor,
+						  DebugPrimitiveDepth,
+						  SceneDepthJitter,
+						  NumMSAASamples,
+						  true,
+						  Inputs.bUseMetalMSAAHDRDecode);
+	}
+	else
 	{
-		TemporalUpscaleDepthPass(GraphBuilder,
-			*DebugView,
-			Inputs.SceneColor,
-			SceneDepth,
-			SceneDepthJitter);
+		AddClearRenderTargetPass(GraphBuilder, DebugPrimitiveColor);
+		AddClearDepthStencilPass(GraphBuilder, DebugPrimitiveDepth);
 	}
 
-	//Simple element pixel shaders do not output background color for composite, 
-	//so this allows the background to be drawn to the RT at the same time as depth without adding extra draw calls
-	PopulateDepthPass(GraphBuilder,
-		*DebugView,
-		Inputs.SceneColor,
-		SceneDepth,
-		Output.Texture,
-		DebugPrimitiveDepth,
-		SceneDepthJitter,
-		NumMSAASamples,
-		true,
-		Inputs.bUseMetalMSAAHDRDecode);
+	FDebugPrimitivesDrawingContext Result {};
+	Result.DebugView = DebugView;
+	Result.DebugPrimitiveColor = DebugPrimitiveColor;
+	Result.DebugPrimitiveDepth = DebugPrimitiveDepth;
+
+	return Result;
+}
+
+static bool DrawDebugPDE(
+	FRDGBuilder& GraphBuilder,
+	FDebugPrimitivesDrawingContext& Context,
+	const FScreenPassTextureViewport& OutputViewport)
+{
+	if (!Context.DebugView->DebugSimpleElementCollector.HasAnyPrimitives())
+	{
+		return false;
+	}
 
 	//Composite the debug draw elements into the scene
 	FDebugPrimitivesPassParameters* PassParameters = GraphBuilder.AllocParameters<FDebugPrimitivesPassParameters>();
-	PassParameters->View = DebugView->GetShaderParameters();
-	PassParameters->RenderTargets[0] = Output.GetRenderTargetBinding();
+	PassParameters->View = Context.DebugView->GetShaderParameters();
+	PassParameters->RenderTargets[0] = FRenderTargetBinding(Context.DebugPrimitiveColor, ERenderTargetLoadAction::EClear);
 	PassParameters->RenderTargets[0].SetLoadAction(ERenderTargetLoadAction::ELoad);
-	PassParameters->RenderTargets.DepthStencil = FDepthStencilBinding(DebugPrimitiveDepth, ERenderTargetLoadAction::ELoad, ERenderTargetLoadAction::ELoad, FExclusiveDepthStencil::DepthWrite_StencilWrite);
+	PassParameters->RenderTargets.DepthStencil = FDepthStencilBinding(Context.DebugPrimitiveDepth, ERenderTargetLoadAction::ELoad, ERenderTargetLoadAction::ELoad, FExclusiveDepthStencil::DepthWrite_StencilWrite);
 
 	GraphBuilder.AddPass(
 		RDG_EVENT_NAME("DrawDebugPrimitives"),
 		PassParameters,
 		ERDGPassFlags::Raster,
-		[&View, PassParameters, DebugView, OutputViewport](FRDGAsyncTask, FRHICommandList& RHICmdList)
+		[PassParameters, DebugView=Context.DebugView, OutputViewport](FRDGAsyncTask, FRHICommandList& RHICmdList)
 	{
 		RHICmdList.SetViewport(OutputViewport.Rect.Min.X, OutputViewport.Rect.Min.Y, 0.0f, OutputViewport.Rect.Max.X, OutputViewport.Rect.Max.Y, 1.0f);
 
@@ -85,6 +126,151 @@ FScreenPassTexture AddDebugPrimitivePass(
 		DebugView->DebugSimpleElementCollector.DrawBatchedElements(RHICmdList, DrawRenderState, *DebugView, EBlendModeFilter::OpaqueAndMasked, SDPG_Foreground);		
 	});
 
+	return true;
+}
+
+static void ComposeDebugPrimitives(
+	FRDGBuilder& GraphBuilder,
+	const FViewInfo& View,
+	const FCompositePrimitiveInputs& Inputs,
+	FDebugPrimitivesDrawingContext& Context,
+	const FScreenPassRenderTarget& Output)
+{
+	FScreenPassRenderTarget DepthOutput = Inputs.OverrideDepthOutput;
+	int NumMSAASamples = Context.DebugPrimitiveColor->Desc.NumSamples;
+	
+	FScreenPassTexture SceneDepth = Inputs.SceneDepth;
+	FVector2f SceneDepthJitter = FVector2f(View.TemporalJitterPixels);
+
+	bool bCompositeAnyNonNullDepth = true;
+
+	FRHISamplerState* PointClampSampler = TStaticSamplerState<SF_Point, AM_Clamp, AM_Clamp, AM_Clamp>::GetRHI();
+
+	const bool bOpaqueEditorGizmo = View.Family->EngineShowFlags.OpaqueCompositeEditorPrimitives || View.Family->EngineShowFlags.Wireframe;
+
+	FCompositePostProcessPrimitivesPS::FParameters* PassParameters = GraphBuilder.AllocParameters<FCompositePostProcessPrimitivesPS::FParameters>();
+	PassParameters->RenderTargets[0] = Output.GetRenderTargetBinding();
+
+	bool bOutputIsMSAA = Output.Texture->Desc.NumSamples > 1;
+	if (DepthOutput.IsValid())
+	{
+		PassParameters->RenderTargets.DepthStencil = FDepthStencilBinding(DepthOutput.Texture, ERenderTargetLoadAction::EClear, ERenderTargetLoadAction::ENoAction, FExclusiveDepthStencil::DepthWrite);
+		verify(Output.Texture->Desc.NumSamples == DepthOutput.Texture->Desc.NumSamples);
+	}
+
+	PassParameters->View = View.ViewUniformBuffer;
+	PassParameters->Color = GetScreenPassTextureViewportParameters(FScreenPassTextureViewport(Inputs.SceneColor));
+	PassParameters->Depth = GetScreenPassTextureViewportParameters(FScreenPassTextureViewport(SceneDepth));
+
+	PassParameters->UndistortingDisplacementTexture = GSystemTextures.GetBlackDummy(GraphBuilder);
+	PassParameters->UndistortingDisplacementSampler = TStaticSamplerState<SF_Bilinear, AM_Clamp, AM_Clamp, AM_Clamp>::GetRHI();
+	if (Inputs.LensDistortionLUT.IsEnabled())
+	{
+		PassParameters->UndistortingDisplacementTexture = Inputs.LensDistortionLUT.UndistortingDisplacementTexture;
+	}
+
+	PassParameters->ColorTexture = Inputs.SceneColor.Texture;
+	PassParameters->ColorSampler = PointClampSampler;
+	PassParameters->DepthTexture = SceneDepth.Texture;
+	PassParameters->DepthSampler = PointClampSampler;
+	PassParameters->EditorPrimitivesDepth = Context.DebugPrimitiveDepth;
+	PassParameters->EditorPrimitivesColor = Context.DebugPrimitiveColor;
+		
+	PassParameters->PassSvPositionToViewportUV = FScreenTransform::SvPositionToViewportUV(Output.ViewRect);
+	PassParameters->ViewportUVToColorUV = FScreenTransform::ChangeTextureBasisFromTo(
+		FScreenPassTextureViewport(Inputs.SceneColor), FScreenTransform::ETextureBasis::ViewportUV, FScreenTransform::ETextureBasis::TextureUV);
+	PassParameters->ViewportUVToDepthUV = FScreenTransform::ChangeTextureBasisFromTo(
+		FScreenPassTextureViewport(SceneDepth), FScreenTransform::ETextureBasis::ViewportUV, FScreenTransform::ETextureBasis::TextureUV);
+		
+	PassParameters->bOpaqueEditorGizmo = bOpaqueEditorGizmo;
+	PassParameters->bCompositeAnyNonNullDepth = bCompositeAnyNonNullDepth;
+	PassParameters->DepthTextureJitter = SceneDepthJitter;
+	PassParameters->bProcessAlpha = IsPostProcessingWithAlphaChannelSupported();
+
+	for (int32 i = 0; i < int32(NumMSAASamples); i++)
+	{
+		PassParameters->SampleOffsetArray[i].X = GetMSAASampleOffsets(NumMSAASamples, i).X;
+		PassParameters->SampleOffsetArray[i].Y = GetMSAASampleOffsets(NumMSAASamples, i).Y;
+	}
+
+	FCompositePostProcessPrimitivesPS::FPermutationDomain PermutationVector;
+	PermutationVector.Set<FCompositePostProcessPrimitivesPS::FSampleCountDimension>(NumMSAASamples);
+	PermutationVector.Set<FCompositePostProcessPrimitivesPS::FMSAADontResolve>(bOutputIsMSAA);
+	PermutationVector.Set<FCompositePostProcessPrimitivesPS::FWriteDepth>(DepthOutput.IsValid());
+
+	TShaderMapRef<FCompositePostProcessPrimitivesPS> PixelShader(View.ShaderMap, PermutationVector);
+		
+	FRHIDepthStencilState* DepthStencilState = nullptr;
+	if (DepthOutput.IsValid())
+	{
+		DepthStencilState = TStaticDepthStencilState<true, CF_Always>::GetRHI();
+	}
+
+	FPixelShaderUtils::AddFullscreenPass(
+		GraphBuilder,
+		View.ShaderMap,
+		RDG_EVENT_NAME("Composite %dx%d MSAA=%d", Output.ViewRect.Width(), Output.ViewRect.Height(), NumMSAASamples),
+		PixelShader,
+		PassParameters,
+		Output.ViewRect,
+		nullptr,
+		nullptr,
+		DepthStencilState);
+}
+
+bool IsDebugPrimitivePassEnabled(const FViewInfo& View)
+{
+	static bool bIsForceDisabled = IConsoleManager::Get().FindTConsoleVariableDataInt(TEXT("r.ForceDebugViewModes"))->GetValueOnAnyThread() == 2;
+	if (bIsForceDisabled || !View.Family->EngineShowFlags.CompositeDebugPrimitives)
+	{
+		return false;
+	}
+
+	bool bHasDebugPDEPrimitives = View.DebugSimpleElementCollector.HasAnyPrimitives();
+	bool bVisualizeInstanceOcclusionCulling = View.Family->EngineShowFlags.VisualizeInstanceOcclusionQueries;
+
+	return bHasDebugPDEPrimitives || bVisualizeInstanceOcclusionCulling;
+}
+
+FScreenPassTexture AddDebugPrimitivePass(
+	FRDGBuilder& GraphBuilder,
+	const FViewInfo& View,
+	const FCompositePrimitiveInputs& Inputs)
+{
+	check(Inputs.SceneColor.IsValid());
+	check(Inputs.SceneDepth.IsValid());
+
+	RDG_EVENT_SCOPE(GraphBuilder, "CompositeDebugPrimitives");
+
+	FScreenPassRenderTarget Output = Inputs.OverrideOutput;
+	if (!Output.IsValid())
+	{
+		Output = FScreenPassRenderTarget::CreateFromInput(GraphBuilder, Inputs.SceneColor, View.GetOverwriteLoadAction(), TEXT("Debug.DrawPrimitivesColor"));
+	}
+	const FScreenPassTextureViewport OutputViewport(Output);
+
+	FDebugPrimitivesDrawingContext Context = InitializeDebugPrimitiveTextures(GraphBuilder, View, Inputs, Output);
+
+	bool bHasDrawn = DrawDebugPDE(GraphBuilder, Context, OutputViewport);
+
+	{
+		FScene* Scene = (FScene*)View.Family->Scene;
+		if (View.Family->EngineShowFlags.VisualizeInstanceOcclusionQueries && Scene && Scene->InstanceCullingOcclusionQueryRenderer)
+		{
+			Scene->InstanceCullingOcclusionQueryRenderer->RenderDebug(GraphBuilder, Scene->GPUScene, *Context.DebugView, OutputViewport.Rect, Context.DebugPrimitiveColor, Context.DebugPrimitiveDepth);
+			bHasDrawn = true;
+		}
+	}
+
+	if (bHasDrawn)
+	{
+		ComposeDebugPrimitives(GraphBuilder, View, Inputs, Context, Output);
+	}
+	else
+	{
+		AddDrawTexturePass(GraphBuilder, View, Inputs.SceneColor, Output);
+	}
+
 	return MoveTemp(Output);
 }
 
diff --git a/Engine/Source/Runtime/Renderer/Private/PostProcess/PostProcessCompositeDebugPrimitives.h b/Engine/Source/Runtime/Renderer/Private/PostProcess/PostProcessCompositeDebugPrimitives.h
--- a/Engine/Source/Runtime/Renderer/Private/PostProcess/PostProcessCompositeDebugPrimitives.h
+++ b/Engine/Source/Runtime/Renderer/Private/PostProcess/PostProcessCompositeDebugPrimitives.h
@@ -4,5 +4,6 @@
 #include "PostProcess/PostProcessCompositePrimitivesCommon.h"
 
 #if UE_ENABLE_DEBUG_DRAWING
+bool IsDebugPrimitivePassEnabled(const FViewInfo& View);
 FScreenPassTexture AddDebugPrimitivePass(FRDGBuilder& GraphBuilder, const FViewInfo& View, const FCompositePrimitiveInputs& Inputs);
 #endif
diff --git a/Engine/Source/Runtime/Renderer/Private/PostProcess/PostProcessCompositeEditorPrimitives.cpp b/Engine/Source/Runtime/Renderer/Private/PostProcess/PostProcessCompositeEditorPrimitives.cpp
--- a/Engine/Source/Runtime/Renderer/Private/PostProcess/PostProcessCompositeEditorPrimitives.cpp
+++ b/Engine/Source/Runtime/Renderer/Private/PostProcess/PostProcessCompositeEditorPrimitives.cpp
@@ -2,6 +2,12 @@
 
 #include "PostProcess/PostProcessCompositeEditorPrimitives.h"
 
+#if UE_ENABLE_DEBUG_DRAWING
+
+IMPLEMENT_GLOBAL_SHADER(FCompositePostProcessPrimitivesPS, "/Engine/Private/PostProcessCompositePrimitives.usf", "MainCompositePostProcessPrimitivesPS", SF_Pixel);
+
+#endif //UE_ENABLE_DEBUG_DRAWING
+
 #if WITH_EDITOR
 #include "EditorPrimitivesRendering.h"
 #include "MeshPassProcessor.inl"
@@ -15,68 +21,19 @@
 
 namespace
 {
-class FCompositeEditorPrimitivesPS : public FCompositePrimitiveShaderBase
+struct FEditorPrimitivesDrawingContext
 {
-public:
-	class FWriteDepth : SHADER_PERMUTATION_BOOL("WRITE_DEPTH");
-					
-	using FPermutationDomain = TShaderPermutationDomain<FWriteDepth, FSampleCountDimension, FMSAADontResolve>;
-			
-	DECLARE_GLOBAL_SHADER(FCompositeEditorPrimitivesPS);
-	SHADER_USE_PARAMETER_STRUCT(FCompositeEditorPrimitivesPS, FCompositePrimitiveShaderBase);
-
-	BEGIN_SHADER_PARAMETER_STRUCT(FParameters, )
-		SHADER_PARAMETER_STRUCT_REF(FViewUniformShaderParameters, View)
-		SHADER_PARAMETER_STRUCT(FScreenPassTextureViewportParameters, Color)
-		SHADER_PARAMETER_STRUCT(FScreenPassTextureViewportParameters, Depth)
-		SHADER_PARAMETER_ARRAY(FVector4f, SampleOffsetArray, [FCompositePrimitiveShaderBase::kMSAASampleCountMax])
-
-		SHADER_PARAMETER_RDG_TEXTURE(Texture2D, UndistortingDisplacementTexture)
-		SHADER_PARAMETER_SAMPLER(SamplerState,  UndistortingDisplacementSampler)
-		SHADER_PARAMETER_RDG_TEXTURE(Texture2D, EditorPrimitivesDepth)
-		SHADER_PARAMETER_RDG_TEXTURE(Texture2D, EditorPrimitivesColor)
-		SHADER_PARAMETER_RDG_TEXTURE(Texture2D, ColorTexture)
-		SHADER_PARAMETER_SAMPLER(SamplerState,  ColorSampler)
-		SHADER_PARAMETER_RDG_TEXTURE(Texture2D, DepthTexture)
-		SHADER_PARAMETER_SAMPLER(SamplerState,  DepthSampler)
-
-		SHADER_PARAMETER(FScreenTransform, PassSvPositionToViewportUV)
-		SHADER_PARAMETER(FScreenTransform, ViewportUVToColorUV)
-		SHADER_PARAMETER(FScreenTransform, ViewportUVToDepthUV)
-		SHADER_PARAMETER(uint32, bOpaqueEditorGizmo)
-		SHADER_PARAMETER(uint32, bCompositeAnyNonNullDepth)
-		SHADER_PARAMETER(FVector2f, DepthTextureJitter)
-		SHADER_PARAMETER(uint32, bProcessAlpha)
-		RENDER_TARGET_BINDING_SLOTS()
-	END_SHADER_PARAMETER_STRUCT()
-
-	static bool ShouldCompilePermutation(const FPermutationDomain& PermutationVector, const EShaderPlatform Platform)
-	{
-		const int32 SampleCount = PermutationVector.Get<FSampleCountDimension>();
-		// Only use permutations with valid MSAA sample counts.
-		if (!FMath::IsPowerOfTwo(SampleCount))
-		{
-		   return false;
-		}
+	// View at UnscaledViewRect size, with TAA jitter disabled, exposure set for post-tonemapper, ...
+	const FViewInfo* EditorView = nullptr;
 
-		return IsPCPlatform(Platform);
-	}
-	   
-	static bool ShouldCompilePermutation(const FGlobalShaderPermutationParameters& Parameters)
-	{
-	   const FPermutationDomain PermutationVector(Parameters.PermutationId);
-	   return ShouldCompilePermutation(PermutationVector, Parameters.Platform);
-	}
+	// EditorPrimitives* textures, at UnscaledViewRect size
+	FRDGTextureRef EditorPrimitiveColor = nullptr;
+	FRDGTextureRef EditorPrimitiveDepth = nullptr;
 
-	static void ModifyCompilationEnvironment(const FGlobalShaderPermutationParameters& Parameters, FShaderCompilerEnvironment& OutEnvironment)
-	{
-		FGlobalShader::ModifyCompilationEnvironment(Parameters, OutEnvironment);
-	}
+	bool bCompositeAnyNonNullDepth;
 };
-IMPLEMENT_GLOBAL_SHADER(FCompositeEditorPrimitivesPS, "/Engine/Private/PostProcessCompositePrimitives.usf", "MainCompositeEditorPrimitivesPS", SF_Pixel);
 
-
-void RenderEditorPrimitives(FRHICommandList& RHICmdList, const FViewInfo& View, FMeshPassProcessorRenderState& DrawRenderState, FInstanceCullingManager& InstanceCullingManager)
+static void RenderEditorPrimitives(FRHICommandList& RHICmdList, const FViewInfo& View, FMeshPassProcessorRenderState& DrawRenderState, FInstanceCullingManager& InstanceCullingManager)
 {
 	// Always depth test against other editor primitives
 	DrawRenderState.SetDepthStencilState(TStaticDepthStencilState<
@@ -125,7 +82,7 @@ void RenderEditorPrimitives(FRHICommandList& RHICmdList, const FViewInfo& View,
 	View.BatchedViewElements.Draw(RHICmdList, DrawRenderState, FeatureLevel, View, false, 1.0f);
 }
 
-void RenderForegroundTranslucentEditorPrimitives(FRHICommandList& RHICmdList, const FViewInfo& View, FMeshPassProcessorRenderState& DrawRenderState, FInstanceCullingManager& InstanceCullingManager)
+static void RenderForegroundTranslucentEditorPrimitives(FRHICommandList& RHICmdList, const FViewInfo& View, FMeshPassProcessorRenderState& DrawRenderState, FInstanceCullingManager& InstanceCullingManager)
 {
 	const auto FeatureLevel = View.GetFeatureLevel();
 	const auto ShaderPlatform = GShaderPlatformForFeatureLevel[FeatureLevel];
@@ -163,7 +120,7 @@ void RenderForegroundTranslucentEditorPrimitives(FRHICommandList& RHICmdList, co
 	View.TopBatchedViewElements.Draw(RHICmdList, DrawRenderState, FeatureLevel, View, false, 1.0f, EBlendModeFilter::Translucent);
 }
 
-void RenderForegroundEditorPrimitives(FRHICommandList& RHICmdList, const FViewInfo& View, FMeshPassProcessorRenderState& DrawRenderState, FInstanceCullingManager& InstanceCullingManager)
+static void RenderForegroundEditorPrimitives(FRHICommandList& RHICmdList, const FViewInfo& View, FMeshPassProcessorRenderState& DrawRenderState, FInstanceCullingManager& InstanceCullingManager)
 {
 	const auto FeatureLevel = View.GetFeatureLevel();
 	const auto ShaderPlatform = GShaderPlatformForFeatureLevel[FeatureLevel];
@@ -241,104 +198,146 @@ BEGIN_SHADER_PARAMETER_STRUCT(FEditorPrimitivesPassParameters, )
 	RENDER_TARGET_BINDING_SLOTS()
 END_SHADER_PARAMETER_STRUCT()
 
-FScreenPassTexture AddEditorPrimitivePass(
+static FEditorPrimitivesDrawingContext InitializeEditorPrimitivesTextures(
 	FRDGBuilder& GraphBuilder,
 	const FViewInfo& View,
-	const FCompositePrimitiveInputs& Inputs,
-	FInstanceCullingManager& InstanceCullingManager)
+	const FScreenPassTexture& SceneColor,
+	const FScreenPassTexture& SceneDepth)
 {
-	check(Inputs.SceneColor.IsValid());
-	check(Inputs.SceneDepth.IsValid());
-	check(Inputs.BasePassType != FCompositePrimitiveInputs::EBasePassType::MAX);
+	check(SceneColor.IsValid());
+	check(SceneDepth.IsValid());
 
 	const FSceneTextures& SceneTextures = View.GetSceneTextures();
 	const uint32 NumMSAASamples = SceneTextures.Config.EditorPrimitiveNumSamples;
-	const FViewInfo* EditorView = CreateCompositePrimitiveView(View, Inputs.SceneColor.ViewRect, NumMSAASamples);
+	const FViewInfo* EditorView = CreateCompositePrimitiveView(View, SceneColor.ViewRect, NumMSAASamples);	
 
 	// Load the color target if it already exists.
 	bool bProducedByPriorPass = HasBeenProduced(SceneTextures.EditorPrimitiveColor);
-	FIntPoint Extent = Inputs.SceneColor.Texture->Desc.Extent;
-	FRDGTextureRef EditorPrimitiveColor;
-	FRDGTextureRef EditorPrimitiveDepth;
+	FIntPoint Extent = SceneColor.Texture->Desc.Extent;
+	FRDGTextureRef EditorPrimitiveColor = SceneTextures.EditorPrimitiveColor;
+	FRDGTextureRef EditorPrimitiveDepth = SceneTextures.EditorPrimitiveDepth;
+
+	bool bViewIsScaled = View.ViewRect.Size() != EditorView->ViewRect.Size();
+	bool bSceneDepthIsScaled = SceneColor.ViewRect != SceneDepth.ViewRect;
+
 	if (bProducedByPriorPass)
 	{
-		EditorPrimitiveColor = SceneTextures.EditorPrimitiveColor;
-	}
-	else
-	{
-		const FRDGTextureDesc ColorDesc = FRDGTextureDesc::Create2D(
-			Extent,
-			PF_B8G8R8A8,
-			FClearValueBinding::Transparent,
-			TexCreate_ShaderResource | TexCreate_RenderTargetable,
-			1,
-			NumMSAASamples);
-
-		EditorPrimitiveColor = GraphBuilder.CreateTexture(ColorDesc, TEXT("Editor.PrimitivesColor"));
+		// Upscaling EditorPrimitiveColor is not supported.
+		ensure(!bViewIsScaled);
 	}
 
-	if (bProducedByPriorPass && Inputs.SceneColor.ViewRect == Inputs.SceneDepth.ViewRect)
+	// If the view is scaled, use fresh textures with the correct size for the editor viewport
+	if (bViewIsScaled || bSceneDepthIsScaled)
 	{
-		EditorPrimitiveDepth = SceneTextures.EditorPrimitiveDepth;
-	}
-	else
-	{
-		EditorPrimitiveDepth = CreateCompositeDepthTexture(GraphBuilder, Extent, NumMSAASamples);
-		//bProducedByPriorPass no longer true as this pass had to create a depth texture for temporal upscaling
+		FRDGTextureDesc EditorPrimitiveColorDesc = EditorPrimitiveColor->Desc;
+		FRDGTextureDesc EditorPrimitiveDepthDesc = EditorPrimitiveDepth->Desc;
+
+		EditorPrimitiveColorDesc.Extent = Extent;
+		EditorPrimitiveDepthDesc.Extent = Extent;
+		EditorPrimitiveDepthDesc.AliasableFormats.Add(PF_X24_G8); // could this go into SceneTextures?
+
+		if (bViewIsScaled)
+		{
+			EditorPrimitiveColor = GraphBuilder.CreateTexture(EditorPrimitiveColorDesc, TEXT("Editor.PrimitivesColor"));
+		}
+		EditorPrimitiveDepth = GraphBuilder.CreateTexture(EditorPrimitiveDepthDesc, TEXT("Editor.PrimitivesDepth"));
+
 		bProducedByPriorPass = false;
 	}
-	
-	// Substrate data might not be produced in certain case (e.g., path-tracer). In such a case we force generate 
-	// them with a simple clear to please validation.
-	if (Substrate::IsSubstrateEnabled() && Substrate::UsesSubstrateMaterialBuffer(View.GetShaderPlatform()) && View.SubstrateViewData.SceneData->TopLayerTexture && !HasBeenProduced(View.SubstrateViewData.SceneData->TopLayerTexture))
-	{
-		FRDGTextureClearInfo ClearInfo;
-		AddClearRenderTargetPass(GraphBuilder, View.SubstrateViewData.SceneData->TopLayerTexture, ClearInfo);
-	}
 
-	// Load the color target if it already exists.
-	const FScreenPassTextureViewport EditorPrimitivesViewport(EditorPrimitiveColor, Inputs.SceneColor.ViewRect);
+	const FScreenPassTextureViewport EditorPrimitivesViewport(EditorPrimitiveColor, SceneColor.ViewRect);
 
-	RDG_EVENT_SCOPE_STAT(GraphBuilder, EditorPrimitives, "CompositeEditorPrimitives %dx%d MSAA=%d",
-		EditorPrimitivesViewport.Rect.Width(),
-		EditorPrimitivesViewport.Rect.Height(),
-		NumMSAASamples);
+	RDG_EVENT_SCOPE_STAT(GraphBuilder, EditorPrimitives, "InitializeEditorPrimitivesTextures %dx%d MSAA=%d",
+						 EditorPrimitivesViewport.Rect.Width(),
+						 EditorPrimitivesViewport.Rect.Height(),
+						 NumMSAASamples);
 	RDG_GPU_STAT_SCOPE(GraphBuilder, EditorPrimitives);
 
-	//Inputs is const so create a over-ridable texture reference
-	FScreenPassTexture SceneDepth = Inputs.SceneDepth;
-	FVector2f SceneDepthJitter = FVector2f(View.TemporalJitterPixels);
-
 	// The editor primitive composition pass is also used when rendering VMI_WIREFRAME in order to use MSAA.
 	// So we need to check whether the editor primitives are enabled inside this function.
+	// Similarly, MeshEdges is also wireframe, but as a separate pre-pass.
 	if (View.Family->EngineShowFlags.CompositeEditorPrimitives || View.Family->EngineShowFlags.MeshEdges)
 	{
+		FVector2f SceneDepthJitter = FVector2f(View.TemporalJitterPixels);
+
 		// Populate depth if a prior pass did not already do it.
 		if (!bProducedByPriorPass)
 		{
+			FScreenPassTexture UpsampledDepth = SceneDepth;
+
 			if (IsTemporalAccumulationBasedMethod(View.AntiAliasingMethod))
 			{
 				TemporalUpscaleDepthPass(GraphBuilder,
-					*EditorView,
-					Inputs.SceneColor,
-					SceneDepth,
-					SceneDepthJitter);
+										 *EditorView,
+										 SceneColor,
+										 UpsampledDepth,
+										 SceneDepthJitter);
 			}
 
 			PopulateDepthPass(GraphBuilder,
-				*EditorView,
-				Inputs.SceneColor,
-				SceneDepth,
-				EditorPrimitiveColor,
-				EditorPrimitiveDepth,
-				SceneDepthJitter,
-				NumMSAASamples);
+							  *EditorView,
+							  SceneColor,
+							  UpsampledDepth,
+							  EditorPrimitiveColor,
+							  EditorPrimitiveDepth,
+							  SceneDepthJitter,
+							  NumMSAASamples);
 		}	
 		
 		FScreenPassRenderTarget EditorPrimitiveColorRT(EditorPrimitiveColor, EditorPrimitivesViewport.Rect, ERenderTargetLoadAction::ELoad);
 		FScreenPassRenderTarget EditorPrimitiveDepthRT(EditorPrimitiveDepth, EditorPrimitivesViewport.Rect, ERenderTargetLoadAction::ELoad);
 		ComposeMeshEdges(GraphBuilder, View, EditorPrimitiveColorRT, EditorPrimitiveDepthRT);
+	}
+
+	FEditorPrimitivesDrawingContext Result {};
+	Result.EditorView = EditorView;
+	Result.bCompositeAnyNonNullDepth = bProducedByPriorPass && !View.Family->EngineShowFlags.MeshEdges;
+	Result.EditorPrimitiveColor = EditorPrimitiveColor;
+	Result.EditorPrimitiveDepth = EditorPrimitiveDepth;
+
+	return Result;
+}
+
+FScreenPassTexture AddEditorPrimitivePass(
+	FRDGBuilder& GraphBuilder,
+	const FViewInfo& View,
+	const FCompositePrimitiveInputs& Inputs,
+	FInstanceCullingManager& InstanceCullingManager)
+{
+	check(Inputs.SceneColor.IsValid());
+	check(Inputs.SceneDepth.IsValid());
+	check(Inputs.BasePassType != FCompositePrimitiveInputs::EBasePassType::MAX);
+	
+	FEditorPrimitivesDrawingContext Context = InitializeEditorPrimitivesTextures(GraphBuilder, View, Inputs.SceneColor, Inputs.SceneDepth);
 
+	FRDGTextureRef EditorPrimitiveColor = Context.EditorPrimitiveColor;
+	FRDGTextureRef EditorPrimitiveDepth = Context.EditorPrimitiveDepth;
+
+	const FViewInfo* EditorView = Context.EditorView;
+	const uint32 NumMSAASamples = EditorPrimitiveColor->Desc.NumSamples;
+	
+	// Substrate data might not be produced in certain case (e.g., path-tracer). In such a case we force generate 
+	// them with a simple clear to please validation.
+	if (Substrate::IsSubstrateEnabled() && Substrate::UsesSubstrateMaterialBuffer(View.GetShaderPlatform()) && View.SubstrateViewData.SceneData->TopLayerTexture && !HasBeenProduced(View.SubstrateViewData.SceneData->TopLayerTexture))
+	{
+		FRDGTextureClearInfo ClearInfo;
+		AddClearRenderTargetPass(GraphBuilder, View.SubstrateViewData.SceneData->TopLayerTexture, ClearInfo);
+	}
+
+	const FScreenPassTextureViewport EditorPrimitivesViewport(EditorPrimitiveColor, Inputs.SceneColor.ViewRect);
+
+	RDG_EVENT_SCOPE_STAT(GraphBuilder, EditorPrimitives, "CompositeEditorPrimitives %dx%d MSAA=%d",
+		EditorPrimitivesViewport.Rect.Width(),
+		EditorPrimitivesViewport.Rect.Height(),
+		NumMSAASamples);
+	RDG_GPU_STAT_SCOPE(GraphBuilder, EditorPrimitives);
+
+	//Inputs is const so create a over-ridable texture reference
+	FScreenPassTexture SceneDepth = Inputs.SceneDepth;
+	FVector2f SceneDepthJitter = FVector2f(View.TemporalJitterPixels);
+
+	if (View.Family->EngineShowFlags.CompositeEditorPrimitives)
+	{
 		// Draws the editors opaque primitives
 		{
 			FEditorPrimitivesPassParameters* PassParameters = GraphBuilder.AllocParameters<FEditorPrimitivesPassParameters>();
@@ -410,7 +409,7 @@ FScreenPassTexture AddEditorPrimitivePass(
 
 		const bool bOpaqueEditorGizmo = View.Family->EngineShowFlags.OpaqueCompositeEditorPrimitives || View.Family->EngineShowFlags.Wireframe;
 
-		FCompositeEditorPrimitivesPS::FParameters* PassParameters = GraphBuilder.AllocParameters<FCompositeEditorPrimitivesPS::FParameters>();
+		FCompositePostProcessPrimitivesPS::FParameters* PassParameters = GraphBuilder.AllocParameters<FCompositePostProcessPrimitivesPS::FParameters>();
 		PassParameters->RenderTargets[0] = Output.GetRenderTargetBinding();
 
 		bool bOutputIsMSAA = Output.Texture->Desc.NumSamples > 1;
@@ -452,7 +451,7 @@ FScreenPassTexture AddEditorPrimitivePass(
 			FScreenPassTextureViewport(SceneDepth), FScreenTransform::ETextureBasis::ViewportUV, FScreenTransform::ETextureBasis::TextureUV);
 		
 		PassParameters->bOpaqueEditorGizmo = bOpaqueEditorGizmo;
-		PassParameters->bCompositeAnyNonNullDepth = bProducedByPriorPass && !View.Family->EngineShowFlags.MeshEdges;
+		PassParameters->bCompositeAnyNonNullDepth = Context.bCompositeAnyNonNullDepth;
 		PassParameters->DepthTextureJitter = SceneDepthJitter;
 		PassParameters->bProcessAlpha = IsPostProcessingWithAlphaChannelSupported();
 
@@ -462,12 +461,12 @@ FScreenPassTexture AddEditorPrimitivePass(
 			PassParameters->SampleOffsetArray[i].Y = GetMSAASampleOffsets(NumMSAASamples, i).Y;
 		}
 
-		FCompositeEditorPrimitivesPS::FPermutationDomain PermutationVector;
-		PermutationVector.Set<FCompositeEditorPrimitivesPS::FSampleCountDimension>(NumMSAASamples);
-		PermutationVector.Set<FCompositeEditorPrimitivesPS::FMSAADontResolve>(bOutputIsMSAA);
-		PermutationVector.Set<FCompositeEditorPrimitivesPS::FWriteDepth>(DepthOutput.IsValid());
+		FCompositePostProcessPrimitivesPS::FPermutationDomain PermutationVector;
+		PermutationVector.Set<FCompositePostProcessPrimitivesPS::FSampleCountDimension>(NumMSAASamples);
+		PermutationVector.Set<FCompositePostProcessPrimitivesPS::FMSAADontResolve>(bOutputIsMSAA);
+		PermutationVector.Set<FCompositePostProcessPrimitivesPS::FWriteDepth>(DepthOutput.IsValid());
 
-		TShaderMapRef<FCompositeEditorPrimitivesPS> PixelShader(View.ShaderMap, PermutationVector);
+		TShaderMapRef<FCompositePostProcessPrimitivesPS> PixelShader(View.ShaderMap, PermutationVector);
 		
 		FRHIDepthStencilState* DepthStencilState = nullptr;
 		if (DepthOutput.IsValid())
diff --git a/Engine/Source/Runtime/Renderer/Private/PostProcess/PostProcessCompositePrimitivesCommon.cpp b/Engine/Source/Runtime/Renderer/Private/PostProcess/PostProcessCompositePrimitivesCommon.cpp
--- a/Engine/Source/Runtime/Renderer/Private/PostProcess/PostProcessCompositePrimitivesCommon.cpp
+++ b/Engine/Source/Runtime/Renderer/Private/PostProcess/PostProcessCompositePrimitivesCommon.cpp
@@ -37,20 +37,6 @@ const FViewInfo* CreateCompositePrimitiveView(const FViewInfo& ParentView, FIntR
 	return View;
 }
 
-FRDGTextureRef CreateCompositeDepthTexture(FRDGBuilder& GraphBuilder, FIntPoint Extent, uint32 NumMSAASamples)
-{
-	FRDGTextureDesc DepthDesc = FRDGTextureDesc::Create2D(
-		Extent,
-		PF_DepthStencil,
-		FClearValueBinding::DepthFar,
-		TexCreate_ShaderResource | TexCreate_DepthStencilTargetable,
-		1,
-		NumMSAASamples);
-	DepthDesc.AliasableFormats.Add(PF_X24_G8);
-
-	return GraphBuilder.CreateTexture(DepthDesc, TEXT("Composite.PrimitivesDepth"));
-}
-
 class FTemporalUpsampleDepthPS : public FGlobalShader
 {
 public:
diff --git a/Engine/Source/Runtime/Renderer/Private/PostProcess/PostProcessCompositePrimitivesCommon.h b/Engine/Source/Runtime/Renderer/Private/PostProcess/PostProcessCompositePrimitivesCommon.h
--- a/Engine/Source/Runtime/Renderer/Private/PostProcess/PostProcessCompositePrimitivesCommon.h
+++ b/Engine/Source/Runtime/Renderer/Private/PostProcess/PostProcessCompositePrimitivesCommon.h
@@ -24,12 +24,17 @@ public:
 	static bool ShouldCompilePermutation(const FPermutationDomain& PermutationVector, const EShaderPlatform Platform)
 	{
 		const int32 SampleCount = PermutationVector.Get<FSampleCountDimension>();
+		const bool bMSAADontResolve = PermutationVector.Get<FMSAADontResolve>();
 
 		// Only use permutations with valid MSAA sample counts.
 		if (!FMath::IsPowerOfTwo(SampleCount))
 		{
 			return false;
 		}
+		if (!RHISupportsMSAA(Platform) && (SampleCount > 1 || bMSAADontResolve))
+		{
+			return false;
+		}
 
 		return IsPCPlatform(Platform);
 	}
@@ -78,9 +83,6 @@ struct FCompositePrimitiveInputs
 // Constructs a new view suitable for rendering debug primitives.
 const FViewInfo* CreateCompositePrimitiveView(const FViewInfo& ParentView, FIntRect ViewRect, uint32 NumMSAASamples);
 
-FRDGTextureRef CreateCompositeDepthTexture(FRDGBuilder& GraphBuilder, FIntPoint Extent, uint32 NumMSAASamples);
-
-
 void TemporalUpscaleDepthPass(
 	FRDGBuilder& GraphBuilder,
 	const FViewInfo& View,
@@ -99,4 +101,70 @@ void PopulateDepthPass(FRDGBuilder& GraphBuilder,
 	bool bForceDrawColor = false,
 	bool bUseMetalPlatformHDRDecode = false);
 
+
+class FCompositePostProcessPrimitivesPS : public FCompositePrimitiveShaderBase
+{
+public:
+	class FWriteDepth : SHADER_PERMUTATION_BOOL("WRITE_DEPTH");
+					
+	using FPermutationDomain = TShaderPermutationDomain<FWriteDepth, FSampleCountDimension, FMSAADontResolve>;
+			
+	DECLARE_GLOBAL_SHADER(FCompositePostProcessPrimitivesPS);
+	SHADER_USE_PARAMETER_STRUCT(FCompositePostProcessPrimitivesPS, FCompositePrimitiveShaderBase);
+
+	BEGIN_SHADER_PARAMETER_STRUCT(FParameters, )
+		SHADER_PARAMETER_STRUCT_REF(FViewUniformShaderParameters, View)
+		SHADER_PARAMETER_STRUCT(FScreenPassTextureViewportParameters, Color)
+		SHADER_PARAMETER_STRUCT(FScreenPassTextureViewportParameters, Depth)
+		SHADER_PARAMETER_ARRAY(FVector4f, SampleOffsetArray, [FCompositePrimitiveShaderBase::kMSAASampleCountMax])
+
+		SHADER_PARAMETER_RDG_TEXTURE(Texture2D, UndistortingDisplacementTexture)
+		SHADER_PARAMETER_SAMPLER(SamplerState,  UndistortingDisplacementSampler)
+		SHADER_PARAMETER_RDG_TEXTURE(Texture2D, EditorPrimitivesDepth)
+		SHADER_PARAMETER_RDG_TEXTURE(Texture2D, EditorPrimitivesColor)
+		SHADER_PARAMETER_RDG_TEXTURE(Texture2D, ColorTexture)
+		SHADER_PARAMETER_SAMPLER(SamplerState,  ColorSampler)
+		SHADER_PARAMETER_RDG_TEXTURE(Texture2D, DepthTexture)
+		SHADER_PARAMETER_SAMPLER(SamplerState,  DepthSampler)
+
+		SHADER_PARAMETER(FScreenTransform, PassSvPositionToViewportUV)
+		SHADER_PARAMETER(FScreenTransform, ViewportUVToColorUV)
+		SHADER_PARAMETER(FScreenTransform, ViewportUVToDepthUV)
+		SHADER_PARAMETER(uint32, bOpaqueEditorGizmo)
+		SHADER_PARAMETER(uint32, bCompositeAnyNonNullDepth)
+		SHADER_PARAMETER(FVector2f, DepthTextureJitter)
+		SHADER_PARAMETER(uint32, bProcessAlpha)
+		RENDER_TARGET_BINDING_SLOTS()
+	END_SHADER_PARAMETER_STRUCT()
+
+	static bool ShouldCompilePermutation(const FPermutationDomain& PermutationVector, const EShaderPlatform Platform)
+	{
+		const int32 SampleCount = PermutationVector.Get<FSampleCountDimension>();
+		const bool bMSAADontResolve = PermutationVector.Get<FMSAADontResolve>();
+
+		// Only use permutations with valid MSAA sample counts.
+		if (!FMath::IsPowerOfTwo(SampleCount))
+		{
+			return false;
+		}
+		if (!RHISupportsMSAA(Platform) && (SampleCount > 1 || bMSAADontResolve))
+		{
+			return false;
+		}
+
+		return true;
+	}
+
+	static bool ShouldCompilePermutation(const FGlobalShaderPermutationParameters& Parameters)
+	{
+		const FPermutationDomain PermutationVector(Parameters.PermutationId);
+		return ShouldCompilePermutation(PermutationVector, Parameters.Platform);
+	}
+
+	static void ModifyCompilationEnvironment(const FGlobalShaderPermutationParameters& Parameters, FShaderCompilerEnvironment& OutEnvironment)
+	{
+		FCompositePrimitiveShaderBase::ModifyCompilationEnvironment(Parameters, OutEnvironment);
+	}
+};
+
 #endif //#if UE_ENABLE_DEBUG_DRAWING
\ No newline at end of file
diff --git a/Engine/Source/Runtime/Renderer/Private/PostProcess/PostProcessSelectionOutline.cpp b/Engine/Source/Runtime/Renderer/Private/PostProcess/PostProcessSelectionOutline.cpp
--- a/Engine/Source/Runtime/Renderer/Private/PostProcess/PostProcessSelectionOutline.cpp
+++ b/Engine/Source/Runtime/Renderer/Private/PostProcess/PostProcessSelectionOutline.cpp
@@ -63,16 +63,27 @@ public:
 		return IsFeatureLevelSupported(Platform, ERHIFeatureLevel::SM5);
 	}
 
-	static bool ShouldCompilePermutation(const FGlobalShaderPermutationParameters& Parameters)
+	static bool ShouldCompilePermutation(const FPermutationDomain& PermutationVector, const EShaderPlatform Platform)
 	{
-		FPermutationDomain CurrentPermutation(Parameters.PermutationId);
-		FCompositePrimitiveShaderBase::FPermutationDomain EditorPrimitiveShaderDomain = CurrentPermutation.Get<FCompositePrimitiveShaderBase::FPermutationDomain>();
-		bool bIsHDR = CurrentPermutation.Get<FSelectionOutlineHDRDim>();
-		if (bIsHDR && !SupportHDR(Parameters.Platform))
+		FCompositePrimitiveShaderBase::FPermutationDomain EditorPrimitiveShaderDomain = PermutationVector.Get<FCompositePrimitiveShaderBase::FPermutationDomain>();
+		bool bIsHDR = PermutationVector.Get<FSelectionOutlineHDRDim>();
+		if (bIsHDR && !SupportHDR(Platform))
+		{
+			return false;
+		}
+
+		if (!IsPCPlatform(Platform))
 		{
 			return false;
 		}
-		return FCompositePrimitiveShaderBase::ShouldCompilePermutation(EditorPrimitiveShaderDomain, Parameters.Platform);
+
+		return FCompositePrimitiveShaderBase::ShouldCompilePermutation(EditorPrimitiveShaderDomain, Platform);
+	}
+
+	static bool ShouldCompilePermutation(const FGlobalShaderPermutationParameters& Parameters)
+	{
+		const FPermutationDomain PermutationVector(Parameters.PermutationId);
+		return ShouldCompilePermutation(PermutationVector, Parameters.Platform);
 	}
 };
 
diff --git a/Engine/Source/Runtime/Renderer/Private/PostProcess/PostProcessing.cpp b/Engine/Source/Runtime/Renderer/Private/PostProcess/PostProcessing.cpp
--- a/Engine/Source/Runtime/Renderer/Private/PostProcess/PostProcessing.cpp
+++ b/Engine/Source/Runtime/Renderer/Private/PostProcess/PostProcessing.cpp
@@ -601,7 +601,7 @@ void AddPostProcessingPasses(
 	PassSequence.SetEnabled(EPass::HMDDistortion, EngineShowFlags.StereoRendering && EngineShowFlags.HMDDistortion);
 	PassSequence.SetEnabled(EPass::HighResolutionScreenshotMask, IsHighResolutionScreenshotMaskEnabled(View));
 #if UE_ENABLE_DEBUG_DRAWING
-	PassSequence.SetEnabled(EPass::DebugPrimitive, FSceneRenderer::ShouldCompositeDebugPrimitivesInPostProcess(View));
+	PassSequence.SetEnabled(EPass::DebugPrimitive, IsDebugPrimitivePassEnabled(View));
 #endif
 	PassSequence.SetEnabled(EPass::PrimaryUpscale, (bApplyLensDistortion && !bApplyLensDistortionInTSR) || (View.PrimaryScreenPercentageMethod == EPrimaryScreenPercentageMethod::SpatialUpscale && PrimaryViewRect.Size() != View.GetSecondaryViewRectSize()));
 	PassSequence.SetEnabled(EPass::SecondaryUpscale, View.RequiresSecondaryUpscale() || View.Family->GetSecondarySpatialUpscalerInterface() != nullptr);
@@ -2533,7 +2533,7 @@ void AddMobilePostProcessingPasses(FRDGBuilder& GraphBuilder, FScene* Scene, con
 #endif
 
 #if UE_ENABLE_DEBUG_DRAWING
-	PassSequence.SetEnabled(EPass::DebugPrimitive, FSceneRenderer::ShouldCompositeDebugPrimitivesInPostProcess(View));
+	PassSequence.SetEnabled(EPass::DebugPrimitive, IsDebugPrimitivePassEnabled(View));
 #endif
 	PassSequence.SetEnabled(EPass::PrimaryUpscale, bShouldPrimaryUpscale);
 	// Match Desktop Renderer, calling View.RequiresSecondaryUpscale()
diff --git a/Engine/Source/Runtime/Renderer/Private/SceneRendering.cpp b/Engine/Source/Runtime/Renderer/Private/SceneRendering.cpp
--- a/Engine/Source/Runtime/Renderer/Private/SceneRendering.cpp
+++ b/Engine/Source/Runtime/Renderer/Private/SceneRendering.cpp
@@ -4834,19 +4834,6 @@ bool FSceneRenderer::ShouldCompositeEditorPrimitives(const FViewInfo& View)
 	return false;
 }
 
-#if UE_ENABLE_DEBUG_DRAWING
-bool FSceneRenderer::ShouldCompositeDebugPrimitivesInPostProcess(const FViewInfo& View)
-{
-	if (View.DebugSimpleElementCollector.HasAnyPrimitives())
-	{
-		//If we have primitives to draw, check we aren't forcing debug view modes off.
-		static bool bIsForceDisabled = IConsoleManager::Get().FindTConsoleVariableDataInt(TEXT("r.ForceDebugViewModes"))->GetValueOnAnyThread() == 2;
-		return !bIsForceDisabled && View.Family->EngineShowFlags.CompositeDebugPrimitives;
-	}
-	return false;
-}
-#endif
-
 void FSceneRenderer::UpdatePrimitiveIndirectLightingCacheBuffers(FRHICommandListBase& RHICmdList)
 {
 	// Use a bit array to prevent primitives from being updated more than once.
diff --git a/Engine/Source/Runtime/Renderer/Private/SceneRendering.h b/Engine/Source/Runtime/Renderer/Private/SceneRendering.h
--- a/Engine/Source/Runtime/Renderer/Private/SceneRendering.h
+++ b/Engine/Source/Runtime/Renderer/Private/SceneRendering.h
@@ -2244,9 +2244,6 @@ public:
 	* @return true if compositing is needed
 	*/
 	static bool ShouldCompositeEditorPrimitives(const FViewInfo& View);
-#if UE_ENABLE_DEBUG_DRAWING
-	static bool ShouldCompositeDebugPrimitivesInPostProcess(const FViewInfo& View);
-#endif
 
 	/** Apply the ResolutionFraction on ViewSize, taking into account renderer's requirements. */
 	static FIntPoint ApplyResolutionFraction(
