[VSM] Change local light default settings to improve performance for 60Hz targets and align with future direction.

* Changed the default for r.Shadow.Virtual.MarkCoarsePagesLocal to 2 to force-cache coarse dynamic pages for local lights which improves performance. Epic and Cinematic scalabilities remain using mode 1.
* Changed r.Shadow.Virtual.UseConservativeDistantLightThreshold to default to enabled to fix recurring issues with spot lights. This causes mild regressions, but this effect is reduced by the above change and taken together performance is still improved.
* Deprecated r.Shadow.Virtual.UseConservativeDistantLightThreshold to align with future changes to bypass the majority of the old and costly shadow setup for scenes with primarily Nanite.

diff --git a/Engine/Config/BaseScalability.ini b/Engine/Config/BaseScalability.ini
--- a/Engine/Config/BaseScalability.ini
+++ b/Engine/Config/BaseScalability.ini
@@ -230,6 +230,7 @@ r.Shadow.Virtual.SMRT.RayCountDirectional=8
 r.Shadow.Virtual.SMRT.SamplesPerRayDirectional=4
 r.Shadow.Virtual.SMRT.RayCountLocal=8
 r.Shadow.Virtual.SMRT.SamplesPerRayLocal=4
+r.Shadow.Virtual.MarkCoarsePagesLocal=1
 
 [ShadowQuality@Cine]
 r.LightFunctionQuality=1
@@ -257,6 +258,7 @@ r.Shadow.Virtual.SMRT.RayCountDirectional=16
 r.Shadow.Virtual.SMRT.SamplesPerRayDirectional=8
 r.Shadow.Virtual.SMRT.RayCountLocal=16
 r.Shadow.Virtual.SMRT.SamplesPerRayLocal=8
+r.Shadow.Virtual.MarkCoarsePagesLocal=1
 
 ;-----------------------------------------------------------------------------------------------------------------
 
diff --git a/Engine/Source/Runtime/Renderer/Private/Shadows/ShadowSceneRenderer.cpp b/Engine/Source/Runtime/Renderer/Private/Shadows/ShadowSceneRenderer.cpp
--- a/Engine/Source/Runtime/Renderer/Private/Shadows/ShadowSceneRenderer.cpp
+++ b/Engine/Source/Runtime/Renderer/Private/Shadows/ShadowSceneRenderer.cpp
@@ -44,22 +44,22 @@ static TAutoConsoleVariable<int32> CVarDistantLightMode(
 );
 
 static TAutoConsoleVariable<float> CVarDistantLightForceCacheFootprintFraction(
-	TEXT("r.Shadow.Virtual.DistantLightForceCacheFootprintFraction"),
+	TEXT("r.Shadow.Virtual.DistantLightForceCacheFootprintFraction_WillBeRemoved"),
 	0.0f,
-	TEXT("Fraction of footprint size below which start force-caching lights that are invalidated (i.e., are moving or re-added)\n")
-	TEXT("  Larger values may improve performance but may also produce more visible artifacts\n")
-	TEXT("  The base footprint is based on the page size.\n")
-	TEXT("  0.0 == Never force-cache (default), 1.0 == Always force-cache."),
+	TEXT("Will be removed."),
 	ECVF_Scalability | ECVF_RenderThreadSafe
 );
+static FAutoConsoleVariableDeprecated CVarDistantLightForceCacheFootprintFraction_Deprecated(TEXT("r.Shadow.Virtual.DistantLightForceCacheFootprintFraction"), TEXT("r.Shadow.Virtual.DistantLightForceCacheFootprintFraction_WillBeRemoved"), TEXT("5.7"));
 
 static TAutoConsoleVariable<bool> CVarUseConservativeDistantLightThreshold(
-	TEXT("r.Shadow.Virtual.UseConservativeDistantLightThreshold"),
-	false,
-	TEXT("Base the distant light cutoff on the minimum mip level instead of the shadow resolution calculated through the old path.\n")
-	TEXT("  This fixes problems around the use of an inscribed sphere."),
+	TEXT("r.Shadow.Virtual.UseConservativeDistantLightThreshold_WillBeRemoved"),
+	true,
+	TEXT("Note: Will be removed in a future release!\n")
+	TEXT("  Base the distant light cutoff on the minimum mip level instead of the shadow resolution calculated through the old path.\n")
+	TEXT("  Disabling this causes problems with long narrow spot lights, due to the use of an inscribed sphere in the original heuristic, but also typically leads to somewhat more lights being classified as \"distant\" which reduces cost."),
 	ECVF_Scalability | ECVF_RenderThreadSafe
 );
+static FAutoConsoleVariableDeprecated CVarUseConservativeDistantLightThreshold_Deprecated(TEXT("r.Shadow.Virtual.UseConservativeDistantLightThreshold"), TEXT("r.Shadow.Virtual.UseConservativeDistantLightThreshold_WillBeRemoved"), TEXT("5.7"));
 
 static TAutoConsoleVariable<float> CVarNaniteShadowsLODBias(
 	TEXT("r.Shadow.NaniteLODBias"),
diff --git a/Engine/Source/Runtime/Renderer/Private/VirtualShadowMaps/VirtualShadowMapArray.cpp b/Engine/Source/Runtime/Renderer/Private/VirtualShadowMaps/VirtualShadowMapArray.cpp
--- a/Engine/Source/Runtime/Renderer/Private/VirtualShadowMaps/VirtualShadowMapArray.cpp
+++ b/Engine/Source/Runtime/Renderer/Private/VirtualShadowMaps/VirtualShadowMapArray.cpp
@@ -326,7 +326,7 @@ TAutoConsoleVariable<int32> CVarMarkPixelPagesMipModeLocal(
 
 TAutoConsoleVariable<int32> CVarMarkCoarsePagesLocal(
 	TEXT("r.Shadow.Virtual.MarkCoarsePagesLocal"),
-	1,
+	2,
 	TEXT("Marks coarse pages in local light virtual shadow maps so that low resolution data is available everywhere.\n")
 	TEXT(" 0 - Disabled\n")
 	TEXT(" 1 - (default) Always mark the last mip level.\n")
