[VSM] Trivial cvar comment correction.

#rnx

diff --git a/Engine/Source/Runtime/Renderer/Private/VirtualShadowMaps/VirtualShadowMapArray.cpp b/Engine/Source/Runtime/Renderer/Private/VirtualShadowMaps/VirtualShadowMapArray.cpp
--- a/Engine/Source/Runtime/Renderer/Private/VirtualShadowMaps/VirtualShadowMapArray.cpp
+++ b/Engine/Source/Runtime/Renderer/Private/VirtualShadowMaps/VirtualShadowMapArray.cpp
@@ -329,8 +329,8 @@ TAutoConsoleVariable<int32> CVarMarkCoarsePagesLocal(
 	2,
 	TEXT("Marks coarse pages in local light virtual shadow maps so that low resolution data is available everywhere.\n")
 	TEXT(" 0 - Disabled\n")
-	TEXT(" 1 - (default) Always mark the last mip level.\n")
-	TEXT(" 2 - Performance Mode: Suppress dynamic invalidations due to geometry changes (e.g., moving objects or WPO or animation).\n"),
+	TEXT(" 1 - Always mark the last mip level.\n")
+	TEXT(" 2 - (default) Performance Mode: Suppress dynamic invalidations due to geometry changes (e.g., moving objects or WPO or animation).\n"),
 	ECVF_Scalability | ECVF_RenderThreadSafe
 );
 
