[HWRT] Fix bug due to FDynamicRayTracingInstancesContext using DynamicPrimitiveCollector instead of RayTracingDynamicPrimitiveCollector.

diff --git a/Engine/Source/Runtime/Renderer/Private/RayTracing/RayTracing.cpp b/Engine/Source/Runtime/Renderer/Private/RayTracing/RayTracing.cpp
--- a/Engine/Source/Runtime/Renderer/Private/RayTracing/RayTracing.cpp
+++ b/Engine/Source/Runtime/Renderer/Private/RayTracing/RayTracing.cpp
@@ -664,7 +664,7 @@ namespace RayTracing
 		{
 			FViewInfo& View = *Views[ViewIndex];
 
-			RayTracingInstanceCollector.AddViewMeshArrays(&View, &View.DynamicPrimitiveCollector);
+			RayTracingInstanceCollector.AddViewMeshArrays(&View, &View.RayTracingDynamicPrimitiveCollector);
 
 			bAnyViewRequiresTranslucentGeometry |= ViewSceneOptions[ViewIndex]->bTranslucentGeometry;
 		}
