[HWRT] Fix crash due to duplicate shader bindings.

diff --git a/Engine/Source/Runtime/Renderer/Private/RayTracing/RayTracing.cpp b/Engine/Source/Runtime/Renderer/Private/RayTracing/RayTracing.cpp
--- a/Engine/Source/Runtime/Renderer/Private/RayTracing/RayTracing.cpp
+++ b/Engine/Source/Runtime/Renderer/Private/RayTracing/RayTracing.cpp
@@ -2121,14 +2121,13 @@ namespace RayTracing
 		{
 			for (int32 TaskIndex = 0; TaskIndex < ViewTaskData.View.DynamicRayTracingShaderBindingsPerTask.Num(); TaskIndex++)
 			{
-				ViewTaskData.View.VisibleRayTracingShaderBindings.Append(*ViewTaskData.View.DynamicRayTracingShaderBindingsPerTask[TaskIndex]);
+				TaskData.VisibleShaderBindings.Append(*ViewTaskData.View.DynamicRayTracingShaderBindingsPerTask[TaskIndex]);
 			}
 
 			ViewTaskData.View.DynamicRayTracingShaderBindingsPerTask.Empty();
-
-			TaskData.VisibleShaderBindings.Append(ViewTaskData.View.VisibleRayTracingShaderBindings);
 		}
 
+
 		// Even though task dependencies are setup so all work is done by this point, we still have to wait on the
 		// pipe to clear out its internal state. Otherwise it can assert that it still has work at shutdown.
 		TaskData.AddInstancesPipe.WaitUntilEmpty();
