Enable VSM receiver masks for directional lights by default
- ~10MB memory overhead (depending on settings/scalability). Often fairly significant performance improvements in uncached cases with lots of dynamic geometry.
- Trim allocation slightly when possible
- Disable page dilation by default as it has minimal effects and doesn't work nicely with receiver mask
- Tweak some engine defaults

diff --git a/Engine/Source/Runtime/Renderer/Private/Shadows/ShadowSceneRenderer.cpp b/Engine/Source/Runtime/Renderer/Private/Shadows/ShadowSceneRenderer.cpp
--- a/Engine/Source/Runtime/Renderer/Private/Shadows/ShadowSceneRenderer.cpp
+++ b/Engine/Source/Runtime/Renderer/Private/Shadows/ShadowSceneRenderer.cpp
@@ -869,6 +869,9 @@ void FShadowSceneRenderer::BeginMarkVirtualShadowMapPages(
 
 	RDG_EVENT_SCOPE(GraphBuilder, "VirtualShadowMapMarkPages");
 
+	RDG_EVENT_SCOPE_STAT(GraphBuilder, ShadowDepths, "ShadowDepths");
+	RDG_GPU_STAT_SCOPE(GraphBuilder, ShadowDepths);
+
 	VirtualShadowMapArray.BeginMarkPages(
 		GraphBuilder,
 		SceneRenderer,
diff --git a/Engine/Source/Runtime/Renderer/Private/VirtualShadowMaps/VirtualShadowMapArray.cpp b/Engine/Source/Runtime/Renderer/Private/VirtualShadowMaps/VirtualShadowMapArray.cpp
--- a/Engine/Source/Runtime/Renderer/Private/VirtualShadowMaps/VirtualShadowMapArray.cpp
+++ b/Engine/Source/Runtime/Renderer/Private/VirtualShadowMaps/VirtualShadowMapArray.cpp
@@ -268,10 +268,10 @@ static FAutoConsoleVariableDeprecated CVarVSMShowStats_Deprecated(TEXT("r.Shadow
 
 static TAutoConsoleVariable<float> CVarPageDilationBorderSizeDirectional(
 	TEXT("r.Shadow.Virtual.PageDilationBorderSizeDirectional"),
-	0.05f,
+	0.0f,
 	TEXT("If a screen pixel falls within this fraction of a page border for directional lights, the adacent page will also be mapped.")
 	TEXT("Higher values can reduce page misses at screen edges or disocclusions, but increase total page counts."),
-	ECVF_RenderThreadSafe
+	ECVF_Scalability | ECVF_RenderThreadSafe
 );
 
 static TAutoConsoleVariable<float> CVarFirstPersonPixelRequestBias(
@@ -301,10 +301,10 @@ TAutoConsoleVariable<float> CVarMaxDOFResolutionBias(
 
 static TAutoConsoleVariable<float> CVarPageDilationBorderSizeLocal(
 	TEXT("r.Shadow.Virtual.PageDilationBorderSizeLocal"),
-	0.05f,
+	0.0f,
 	TEXT("If a screen pixel falls within this fraction of a page border for local lights, the adacent page will also be mapped.")
 	TEXT("Higher values can reduce page misses at screen edges or disocclusions, but increase total page counts."),
-	ECVF_RenderThreadSafe
+	ECVF_Scalability | ECVF_RenderThreadSafe
 );
 
 TAutoConsoleVariable<int32> CVarMarkPixelPages(
@@ -2308,9 +2308,13 @@ void FVirtualShadowMapArray::BeginMarkPages(
 			// Directional goes before local so that if local receiver masks are disabled we don't need to allocate them
 			// NOTE: One extra for the single page entries
 			int32 NumDirectionalPageReceiverMaskEntriesRequired = NumDirectionalShadowMaps + 1;
-			const uint32 NumPageReceiverMaskRows = FMath::DivideAndRoundUp((uint32)NumDirectionalPageReceiverMaskEntriesRequired, NumPageTablesEntriesPerRow);
+			uint32 NumPageReceiverMaskRows = FMath::DivideAndRoundUp((uint32)NumDirectionalPageReceiverMaskEntriesRequired, NumPageTablesEntriesPerRow);
+
+			// Special but common case with a single directional light; we don't need to allocate a full row if we only have one
+			uint32 RowEntries = (NumPageReceiverMaskRows == 1) ? NumDirectionalPageReceiverMaskEntriesRequired : NumPageTablesEntriesPerRow;
+
 			// Each page stores 2x2 masks, each representing a 4x4 tile
-			FIntPoint PageReceiverMaskTextureSize(2U * NumPageTablesEntriesPerRow * VSM_PAGE_TABLE_TEX2D_SIZE_X, 2U * NumPageReceiverMaskRows * VSM_PAGE_TABLE_TEX2D_SIZE_Y);
+			FIntPoint PageReceiverMaskTextureSize(2U * RowEntries * VSM_PAGE_TABLE_TEX2D_SIZE_X, 2U * NumPageReceiverMaskRows * VSM_PAGE_TABLE_TEX2D_SIZE_Y);
 
 			// TODO: should be uint16 - but atomics prevent this
 			PageReceiverMasksRDG = GraphBuilder.CreateTexture(
diff --git a/Engine/Source/Runtime/Renderer/Private/VirtualShadowMaps/VirtualShadowMapCacheManager.cpp b/Engine/Source/Runtime/Renderer/Private/VirtualShadowMaps/VirtualShadowMapCacheManager.cpp
--- a/Engine/Source/Runtime/Renderer/Private/VirtualShadowMaps/VirtualShadowMapCacheManager.cpp
+++ b/Engine/Source/Runtime/Renderer/Private/VirtualShadowMaps/VirtualShadowMapCacheManager.cpp
@@ -1500,7 +1500,6 @@ void FVirtualShadowMapArrayCacheManager::ExtractFrameData(
 		GraphBuilder.QueueTextureExtraction(VirtualShadowMapArray.PageFlagsRDG, &PrevBuffers.PageFlags);
 		GraphBuilder.QueueBufferExtraction(VirtualShadowMapArray.NanitePerformanceFeedbackRDG, &PrevBuffers.NanitePerformanceFeedback);
 		GraphBuilder.QueueBufferExtraction(VirtualShadowMapArray.ThrottleBufferRDG, &PrevBuffers.ThrottleBuffer);
-		GraphBuilder.QueueTextureExtraction(VirtualShadowMapArray.PageReceiverMasksRDG, &PrevBuffers.PageReceiverMasks);
 
 		if (IsCacheEnabled())
 		{
@@ -1958,13 +1957,13 @@ FVirtualShadowMapArrayCacheManager::FInvalidationPassCommon FVirtualShadowMapArr
 		UniformParameters->ProjectionData = RegExtCreateSrv(PrevBuffers.ProjectionData, TEXT("Shadow.Virtual.PrevProjectionData"));
 		UniformParameters->PageTable = GraphBuilder.RegisterExternalTexture(PrevBuffers.PageTable, TEXT("Shadow.Virtual.PrevPageTable"));
 		UniformParameters->PageFlags = GraphBuilder.RegisterExternalTexture(PrevBuffers.PageFlags, TEXT("Shadow.Virtual.PrevPageFlags"));
-		UniformParameters->PageReceiverMasks = GraphBuilder.RegisterExternalTexture(PrevBuffers.PageReceiverMasks, TEXT("Shadow.Virtual.PrevPageReceiverMasks"));
 		UniformParameters->UncachedPageRectBounds = RegExtCreateSrv(PrevBuffers.UncachedPageRectBounds, TEXT("Shadow.Virtual.PrevUncachedPageRectBounds"));
 		UniformParameters->AllocatedPageRectBounds = RegExtCreateSrv(PrevBuffers.AllocatedPageRectBounds, TEXT("Shadow.Virtual.PrevAllocatedPageRectBounds"));
 		UniformParameters->CachePrimitiveAsDynamic = GraphBuilder.CreateSRV(UploadCachePrimitiveAsDynamic(GraphBuilder));
 
 		// Unused in this path... may be a better way to handle this
 		UniformParameters->PhysicalPagePool = GSystemTextures.GetZeroUIntArrayAtomicCompatDummy(GraphBuilder);
+		UniformParameters->PageReceiverMasks = GSystemTextures.GetDefaultTexture2D(GraphBuilder, PF_R32_UINT, 0xFFFFFFFFu);
 		FRDGBufferSRVRef Uint32SRVDummy = GraphBuilder.CreateSRV(GSystemTextures.GetDefaultStructuredBuffer(GraphBuilder, sizeof(uint32)));
 		UniformParameters->PerViewData.LightGridData = Uint32SRVDummy;
 		UniformParameters->PerViewData.NumCulledLightsGrid = Uint32SRVDummy;
diff --git a/Engine/Source/Runtime/Renderer/Private/VirtualShadowMaps/VirtualShadowMapCacheManager.h b/Engine/Source/Runtime/Renderer/Private/VirtualShadowMaps/VirtualShadowMapCacheManager.h
--- a/Engine/Source/Runtime/Renderer/Private/VirtualShadowMaps/VirtualShadowMapCacheManager.h
+++ b/Engine/Source/Runtime/Renderer/Private/VirtualShadowMaps/VirtualShadowMapCacheManager.h
@@ -227,7 +227,6 @@ struct FVirtualShadowMapArrayFrameData
 	TRefCountPtr<FRDGPooledBuffer>				ProjectionData;
 	TRefCountPtr<FRDGPooledBuffer>				PhysicalPageLists;
 	TRefCountPtr<IPooledRenderTarget>			PageRequestFlags;
-	TRefCountPtr<IPooledRenderTarget>			PageReceiverMasks;
 	TRefCountPtr<FRDGPooledBuffer>				NanitePerformanceFeedback;
 	TRefCountPtr<FRDGPooledBuffer>				ThrottleBuffer;
 	uint64 GetGPUSizeBytes(bool bLogSizes) const;
diff --git a/Engine/Source/Runtime/Renderer/Private/VirtualShadowMaps/VirtualShadowMapClipmap.cpp b/Engine/Source/Runtime/Renderer/Private/VirtualShadowMaps/VirtualShadowMapClipmap.cpp
--- a/Engine/Source/Runtime/Renderer/Private/VirtualShadowMaps/VirtualShadowMapClipmap.cpp
+++ b/Engine/Source/Runtime/Renderer/Private/VirtualShadowMaps/VirtualShadowMapClipmap.cpp
@@ -28,14 +28,14 @@ static TAutoConsoleVariable<int32>  CVarForceInvalidateDirectionalVSM(
 
 static TAutoConsoleVariable<float> CVarVirtualShadowMapResolutionLodBiasDirectional(
 	TEXT( "r.Shadow.Virtual.ResolutionLodBiasDirectional" ),
-	-0.5f,
+	0.0f,
 	TEXT( "Bias applied to LOD calculations for directional lights. -1.0 doubles resolution, 1.0 halves it and so on." ),
 	ECVF_Scalability | ECVF_RenderThreadSafe
 );
 
 static TAutoConsoleVariable<float> CVarVirtualShadowMapResolutionLodBiasDirectionalMoving(
 	TEXT( "r.Shadow.Virtual.ResolutionLodBiasDirectionalMoving" ),
-	0.5f,
+	0.0f,
 	TEXT( "Bias applied to LOD calculations for directional lights that are moving. -1.0 doubles resolution, 1.0 halves it and so on.\n" )
 	TEXT( "The bias transitions smoothly back to ResolutionLodBiasDirectional as the light transitions to non-moving, see 'r.Shadow.Scene.LightActiveFrameCount'." ),
 	ECVF_Scalability | ECVF_RenderThreadSafe
@@ -144,7 +144,7 @@ static TAutoConsoleVariable<bool> CVarVSMClipmapCullDynamicTightly(
 
 static TAutoConsoleVariable<bool> CVarVSMUseReceiverMaskDirectional(
 	TEXT("r.Shadow.Virtual.UseReceiverMaskDirectional"),
-	false,
+	true,
 	TEXT("Use receiver page masks with directional lights. This enables much more effective culling especially at lower resolutions."),	
 	ECVF_RenderThreadSafe
 );
