Fixed showflag.rendering.

diff --git a/Engine/Source/Runtime/Renderer/Private/DeferredShadingRenderer.cpp b/Engine/Source/Runtime/Renderer/Private/DeferredShadingRenderer.cpp
--- a/Engine/Source/Runtime/Renderer/Private/DeferredShadingRenderer.cpp
+++ b/Engine/Source/Runtime/Renderer/Private/DeferredShadingRenderer.cpp
@@ -1711,11 +1711,6 @@ extern void RayTracingDebugDisplayOnScreenMessages(FScreenMessageWriter& Writer,
 
 void FDeferredShadingSceneRenderer::Render(FRDGBuilder& GraphBuilder, const FSceneRenderUpdateInputs* SceneUpdateInputs)
 {
-	if (!ViewFamily.EngineShowFlags.Rendering)
-	{
-		return;
-	}
-
 	{
 		FRayTracingVisualizationData& RayTracingVisualizationData = GetRayTracingVisualizationData();
 
