PR #13378: Scale the light fade range by `GLightMaxDrawDistanceScale` to avoid lights becoming too dim at lower scales
- Added immediately deprecated r.LightMaxDrawDistanceScale.UseLegacyFadeBehavior_WillBeRemoved to temporarily restore previous behavior

diff --git a/Engine/Source/Runtime/Renderer/Private/LightRendering.cpp b/Engine/Source/Runtime/Renderer/Private/LightRendering.cpp
--- a/Engine/Source/Runtime/Renderer/Private/LightRendering.cpp
+++ b/Engine/Source/Runtime/Renderer/Private/LightRendering.cpp
@@ -768,6 +768,26 @@ ELightOcclusionType GetLightOcclusionType(const FLightSceneInfoCompact& LightInf
 	return bUseRaytracing ? ELightOcclusionType::Raytraced : ELightOcclusionType::Shadowmap;
 }
 
+static TAutoConsoleVariable<bool> CVarLightMaxDrawDistanceScaleLegacyFade(
+	TEXT("r.LightMaxDrawDistanceScale.UseLegacyFadeBehavior_WillBeRemoved"),
+	false,
+	TEXT("Note: Will be removed in a future release!\n")
+	TEXT("  Use legacy fade behavior that does not take r.LightMaxDrawDistanceScale into account."),
+	FConsoleVariableDelegate::CreateLambda(
+		[](IConsoleVariable* CVar)
+		{
+			bool bValue = CVar->GetInt() != 0;
+			if (bValue)
+			{
+				UE_LOG(LogRenderer, Warning, 
+					   TEXT("r.LightMaxDrawDistanceScale.UseLegacyFadeBehavior_WillBeRemoved can produce undesirable light fading behavior.")
+					   TEXT("This variable is deprecated and will be removed in a future version."));
+			}
+		}
+	),
+	ECVF_Scalability | ECVF_RenderThreadSafe
+);
+
 float GetLightFadeFactor(const FSceneView& View, const FLightSceneProxy* Proxy)
 {
 	// Distance fade
@@ -781,8 +801,12 @@ float GetLightFadeFactor(const FSceneView& View, const FLightSceneProxy* Proxy)
 	extern float GLightMaxDrawDistanceScale;
 	float ProxyMaxDist = Proxy->GetMaxDrawDistance();
 	float ScaledMaxDist = ProxyMaxDist * GLightMaxDrawDistanceScale;
-	// NOTE: Feels like we should scale fade range by GLightMaxDrawDistanceScale as well, but would change legacy behavior
-	float Range = Proxy->GetFadeRange();
+	float Range = FMath::Max(0.001f, Proxy->GetFadeRange() * GLightMaxDrawDistanceScale);
+	// Legacy behavior under a (temporary!) chicken bit
+	if (CVarLightMaxDrawDistanceScaleLegacyFade.GetValueOnRenderThread())
+	{
+		Range = Proxy->GetFadeRange();
+	}
 	float DistanceFade = ProxyMaxDist > 0.0f ? (ScaledMaxDist - FMath::Sqrt(DistanceSquared)) / Range : 1.0f;
 	DistanceFade = FMath::Clamp(DistanceFade, 0.0f, 1.0f);
 	return SizeFade * DistanceFade;
