Fix for crash in mobile renderer when Substrate is enabled. Niagara creates caches an RDG uniform buffer on the scene data and it wasn't getting cleared, causing it to return a stale pointer.

diff --git a/Engine/Source/Runtime/Renderer/Private/MobileShadingRenderer.cpp b/Engine/Source/Runtime/Renderer/Private/MobileShadingRenderer.cpp
--- a/Engine/Source/Runtime/Renderer/Private/MobileShadingRenderer.cpp
+++ b/Engine/Source/Runtime/Renderer/Private/MobileShadingRenderer.cpp
@@ -1216,6 +1216,8 @@ void FMobileSceneRenderer::Render(FRDGBuilder& GraphBuilder, const FSceneRenderU
 	GetSceneExtensionsRenderers().PreRender(GraphBuilder);
 	GEngine->GetPreRenderDelegateEx().Broadcast(GraphBuilder);
 
+	::Substrate::PreInitViews(*Scene);
+
 	FSceneTextures::InitializeViewFamily(GraphBuilder, ViewFamily, FamilySize);
 	FSceneTextures& SceneTextures = GetActiveSceneTextures();
 
@@ -1620,6 +1622,8 @@ void FMobileSceneRenderer::Render(FRDGBuilder& GraphBuilder, const FSceneRenderU
 
 	QueueSceneTextureExtractions(GraphBuilder, SceneTextures);
 
+	::Substrate::PostRender(*Scene);
+
 	if (Scene->InstanceCullingOcclusionQueryRenderer)
 	{
 		Scene->InstanceCullingOcclusionQueryRenderer->EndFrame(GraphBuilder);
