Added specialized clear dispatcher for clipmaps to slightly improve clearing performance for receiver masks for directional lights.

#rnx

diff --git a/Engine/Source/Runtime/Renderer/Private/VirtualShadowMaps/VirtualShadowMapArray.cpp b/Engine/Source/Runtime/Renderer/Private/VirtualShadowMaps/VirtualShadowMapArray.cpp
--- a/Engine/Source/Runtime/Renderer/Private/VirtualShadowMaps/VirtualShadowMapArray.cpp
+++ b/Engine/Source/Runtime/Renderer/Private/VirtualShadowMaps/VirtualShadowMapArray.cpp
@@ -2084,6 +2084,13 @@ bool FVirtualShadowMapArray::IsCsvLogEnabled() const
 #endif
 }
 
+class FVirtualShadowMapSetupContext
+{
+public:
+	FPerPageShaderDispatcher PerPageShaderDispatcher;
+	FPerPageShaderDispatcher PerPageShaderDispatcherDirectionalOnly;
+};
+
 void FVirtualShadowMapArray::BeginMarkPages(
 	FRDGBuilder& GraphBuilder,	
 	const FSceneRenderer& SceneRenderer,
@@ -2093,6 +2100,7 @@ void FVirtualShadowMapArray::BeginMarkPages(
 	bool bAnyLocalLightsWithVSMs)
 {
 	check(IsEnabled());
+	check(!VirtualShadowMapSetupContext);
 
 	if (GetNumShadowMaps() == 0)
 	{
@@ -2100,6 +2108,8 @@ void FVirtualShadowMapArray::BeginMarkPages(
 		return;
 	}
 
+	VirtualShadowMapSetupContext = GraphBuilder.AllocObject<FVirtualShadowMapSetupContext>();
+
 	const FMinimalSceneTextures& SceneTextures = SceneRenderer.GetActiveSceneTextures();
 	const TConstArrayView<FViewInfo>& Views = SceneRenderer.Views;
 	const TConstArrayView<FVisibleLightInfo>& VisibleLightInfos = SceneRenderer.VisibleLightInfos;
@@ -2144,6 +2154,7 @@ void FVirtualShadowMapArray::BeginMarkPages(
 
 	{
 		FPerPageShaderDispatcher::FBuilder Builder;
+		FPerPageShaderDispatcher::FBuilder BuilderDirectionalOnly;
 
 		// Create large enough to hold all the unused elements too (wastes GPU memory but allows direct indexing via the ID)
 		uint32 DataSize = sizeof(FVirtualShadowMapProjectionShaderData) * uint32(GetNumShadowMapSlots());
@@ -2167,12 +2178,18 @@ void FVirtualShadowMapArray::BeginMarkPages(
 				int32 EntryVirtualShadowMapId = VirtualShadowMapId + Index;
 				Uploader.Add(EntryVirtualShadowMapId, &Entry.ProjectionData);
 				Builder.Add(EntryVirtualShadowMapId, Entry.ProjectionData.MinMipLevel);
+
+				if (IsDirectional(VirtualShadowMapId))
+				{
+					BuilderDirectionalOnly.Add(EntryVirtualShadowMapId, Entry.ProjectionData.MinMipLevel);
+				}
 			}
 		}
 		Uploader.ResourceUploadTo(GraphBuilder, ProjectionDataRDG);
 
-		PerPageShaderDispatcher = MakeShareable(new FPerPageShaderDispatcher);
-		PerPageShaderDispatcher->Init(GraphBuilder, Builder);
+		
+		VirtualShadowMapSetupContext->PerPageShaderDispatcher.Init(GraphBuilder, Builder);
+		VirtualShadowMapSetupContext->PerPageShaderDispatcherDirectionalOnly.Init(GraphBuilder, BuilderDirectionalOnly);
 	}
 
 	// Stats
@@ -2233,7 +2250,7 @@ void FVirtualShadowMapArray::BeginMarkPages(
 				bAllocateMipLevels ? (VSM_LOG2_PAGE_SIZE + SampleStride / 2u): 1u), 
 			DebugName);
 	};
-	auto AddClearPageTableUAVPass = [&](FRDGTextureRef Dest, uint32 ClearValue, uint32 SampleStride = 1u)
+	auto AddClearPageTableUAVPass = [&](FRDGTextureRef Dest, uint32 ClearValue, uint32 SampleStride = 1u, bool bDirectionalOnly = false)
 	{
 		FRDGTextureUAV* DestUAV = GraphBuilder.CreateUAV(Dest);
 		FClearPageTableCS::FParameters* PassParameters = GraphBuilder.AllocParameters< FClearPageTableCS::FParameters >();
@@ -2255,7 +2272,14 @@ void FVirtualShadowMapArray::BeginMarkPages(
 		PermutationVector.Set<FClearPageTableCS::FNumMipLevelsDim>(Dest->Desc.NumMips);
 		auto ComputeShader = GetGlobalShaderMap(Scene.GetFeatureLevel())->GetShader<FClearPageTableCS>(PermutationVector);
 
-		PerPageShaderDispatcher->AddPass(GraphBuilder, RDG_EVENT_NAME("ClearPageTable(%s Size=%ux%u texels)", DestUAV->GetParent()->Name, DestUAV->GetParent()->Desc.GetSize().X, DestUAV->GetParent()->Desc.GetSize().Y), ComputeShader, PassParameters);
+		if (bDirectionalOnly)
+		{
+			VirtualShadowMapSetupContext->PerPageShaderDispatcherDirectionalOnly.AddPass(GraphBuilder, RDG_EVENT_NAME("ClearPageTableDir(%s Size=%ux%u texels)", DestUAV->GetParent()->Name, DestUAV->GetParent()->Desc.GetSize().X, DestUAV->GetParent()->Desc.GetSize().Y), ComputeShader, PassParameters);
+		}
+		else
+		{
+			VirtualShadowMapSetupContext->PerPageShaderDispatcher.AddPass(GraphBuilder, RDG_EVENT_NAME("ClearPageTable(%s Size=%ux%u texels)", DestUAV->GetParent()->Name, DestUAV->GetParent()->Desc.GetSize().X, DestUAV->GetParent()->Desc.GetSize().Y), ComputeShader, PassParameters);
+		}
 	};
 	auto GetOrCreatePageTableDummy = [&]()
 	{
@@ -2327,7 +2351,7 @@ void FVirtualShadowMapArray::BeginMarkPages(
 					// One extra mip level for the 2x2 tiles
 					(VSM_LOG2_PAGE_SIZE + 1)),
 				TEXT("Shadow.Virtual.PageReceiverMasks"));
-			AddClearUAVPass(GraphBuilder, GraphBuilder.CreateUAV(PageReceiverMasksRDG), FUintVector4(0, 0, 0, 0));
+			AddClearPageTableUAVPass(PageReceiverMasksRDG, 0u, 2u, true);
 		}
 
 		const FIntPoint TextureSize = PageReceiverMasksRDG->Desc.Extent;
@@ -2791,6 +2815,11 @@ void FVirtualShadowMapArray::BuildPageAllocations(
 		// Nothing to do
 		return;
 	}
+	// BeginMarkPages was not run, this indicates a systemic error
+	if (!ensure(VirtualShadowMapSetupContext))
+	{
+		return;
+	}
 
 	RDG_EVENT_SCOPE(GraphBuilder, "FVirtualShadowMapArray::BuildPageAllocation");
 	SCOPED_NAMED_EVENT(FVirtualShadowMapArray_BuildPageAllocation, FColor::Emerald);
@@ -2886,7 +2915,7 @@ void FVirtualShadowMapArray::BuildPageAllocations(
 		SetStatsArgsAndPermutation<FAllocateNewPageMappingsCS>(ShouldGenerateStats(), StatsBufferUAV, PassParameters, PermutationVector);
 		auto ComputeShader = Views[0].ShaderMap->GetShader<FAllocateNewPageMappingsCS>(PermutationVector);
 
-		PerPageShaderDispatcher->AddPass(
+		VirtualShadowMapSetupContext->PerPageShaderDispatcher.AddPass(
 			GraphBuilder,
 			RDG_EVENT_NAME("AllocateNewPageMappings"),
 			ComputeShader,
@@ -2931,7 +2960,7 @@ void FVirtualShadowMapArray::BuildPageAllocations(
 		PassParameters->OutPageTable	= GraphBuilder.CreateUAV(PageTableRDG);
 
 		auto ComputeShader = Views[0].ShaderMap->GetShader<FPropagateMappedMipsCS>();
-		PerPageShaderDispatcher->AddPass(
+		VirtualShadowMapSetupContext->PerPageShaderDispatcher.AddPass(
 			GraphBuilder,
 			RDG_EVENT_NAME("PropagateMappedMips"),
 			ComputeShader,
diff --git a/Engine/Source/Runtime/Renderer/Private/VirtualShadowMaps/VirtualShadowMapArray.h b/Engine/Source/Runtime/Renderer/Private/VirtualShadowMaps/VirtualShadowMapArray.h
--- a/Engine/Source/Runtime/Renderer/Private/VirtualShadowMaps/VirtualShadowMapArray.h
+++ b/Engine/Source/Runtime/Renderer/Private/VirtualShadowMaps/VirtualShadowMapArray.h
@@ -27,7 +27,7 @@ struct FFrontLayerTranslucencyData;
 class FSceneInstanceCullingQuery;
 class FInstanceCullingManager;
 struct FPersistentViewId;
-class FPerPageShaderDispatcher;
+class FVirtualShadowMapSetupContext;
 
 namespace Froxel
 {
@@ -311,11 +311,17 @@ public:
 	
 	void UpdateNextData(int32 PrevVirtualShadowMapId, const FNextVirtualShadowMapData& Data);
 	
-	static bool IsSinglePage(int VirtualShadowMapId)
+	static bool IsSinglePage(int32 VirtualShadowMapId)
 	{
 		return (VirtualShadowMapId < VSM_MAX_SINGLE_PAGE_SHADOW_MAPS);
 	}
 
+	bool IsDirectional(int32 VirtualShadowMapId)
+	{
+		int32 FullId = VirtualShadowMapId - int32(VSM_MAX_SINGLE_PAGE_SHADOW_MAPS);
+		return FullId >= 0 && FullId < NumDirectionalShadowMaps;
+	}
+
 	int32 GetNumShadowMapSlots() const
 	{
 		return NumShadowMapSlots;
@@ -553,7 +559,7 @@ private:
 	// This is primarily an indirection that allows us to reallocate/repack VirtualShadowMapIds each frame
 	TArray<FNextVirtualShadowMapData, SceneRenderingAllocator> NextData;
 
-	TSharedPtr<FPerPageShaderDispatcher> PerPageShaderDispatcher;
+	FVirtualShadowMapSetupContext* VirtualShadowMapSetupContext = nullptr;
 
 	// Allocation/page table order is:
 	// - Single page local light shadow maps (fixed number of slots)
