Fixed: PSO pre-caching for the case when mobile multi-pass rendering is active
#DisabledValidations NonUnityAndStaticAnalysis

diff --git a/Engine/Source/Runtime/Renderer/Private/MeshPassProcessor.cpp b/Engine/Source/Runtime/Renderer/Private/MeshPassProcessor.cpp
--- a/Engine/Source/Runtime/Renderer/Private/MeshPassProcessor.cpp
+++ b/Engine/Source/Runtime/Renderer/Private/MeshPassProcessor.cpp
@@ -1539,7 +1539,7 @@ ESubpassHint GetSubpassHint(const FStaticShaderPlatform Platform, bool bIsUsingG
 {
 	ESubpassHint SubpassHint = ESubpassHint::None;
 	
-	if (IsMobilePlatform(Platform))
+	if (IsMobilePlatform(Platform) && MobileAllowFramebufferFetch(Platform))
 	{
 		if (bIsUsingGBuffers)
 		{
diff --git a/Engine/Source/Runtime/Renderer/Private/MobileBasePass.cpp b/Engine/Source/Runtime/Renderer/Private/MobileBasePass.cpp
--- a/Engine/Source/Runtime/Renderer/Private/MobileBasePass.cpp
+++ b/Engine/Source/Runtime/Renderer/Private/MobileBasePass.cpp
@@ -972,7 +972,7 @@ void FMobileBasePassMeshProcessor::CollectPSOInitializersForLMPolicy(
 	// subpass info set during the submission of the draws in mobile deferred renderer.
 	uint8 SubpassIndex = 0;
 	ESubpassHint SubpassHint = GetSubpassHint(GMaxRHIShaderPlatform, bDeferredShading, RenderTargetsInfo.MultiViewCount > 1, RenderTargetsInfo.NumSamples);
-	if (bTranslucentBasePass)
+	if (bTranslucentBasePass && MobileAllowFramebufferFetch(GMaxRHIShaderPlatform))
 	{
 		if (MeshPassType == EMeshPass::TranslucencyAfterDOF)
 		{
diff --git a/Engine/Source/Runtime/Renderer/Private/MobileShadingRenderer.cpp b/Engine/Source/Runtime/Renderer/Private/MobileShadingRenderer.cpp
--- a/Engine/Source/Runtime/Renderer/Private/MobileShadingRenderer.cpp
+++ b/Engine/Source/Runtime/Renderer/Private/MobileShadingRenderer.cpp
@@ -947,8 +947,11 @@ void FMobileSceneRenderer::RenderCustomRenderPassBasePass(FRDGBuilder& GraphBuil
 	{
 		FColorTargets ColorTargets = GetColorTargets_Deferred(SceneTextures);
 		BasePassRenderTargets = InitRenderTargetBindings_Deferred(SceneTextures, ColorTargets);
-		BasePassRenderTargets.SubpassHint = ESubpassHint::DeferredShadingSubpass;
-		NumAdditionalSubpasses = 2;
+		if (MobileAllowFramebufferFetch(ShaderPlatform))
+		{
+			BasePassRenderTargets.SubpassHint = ESubpassHint::DeferredShadingSubpass;
+			NumAdditionalSubpasses = 2;
+		}
 	}
 	else
 	{
