Mobile renderer: render distance field shadows with async compute where supported.
If this is enabled, we render shadow depths after the prepass in order to allow DF shadows to overlap with it. We don't do this in the general case, because it's slower on some platforms.
Move volumetic fog rendering after the prepass and shadows so the order is consistent regardless of whether async DF shadows is enabled

#DisabledValidations Preflight, CheckCopyright, checkBalancedMacros

diff --git a/Engine/Source/Runtime/Renderer/Private/MobileShadingRenderer.cpp b/Engine/Source/Runtime/Renderer/Private/MobileShadingRenderer.cpp
--- a/Engine/Source/Runtime/Renderer/Private/MobileShadingRenderer.cpp
+++ b/Engine/Source/Runtime/Renderer/Private/MobileShadingRenderer.cpp
@@ -1065,6 +1065,11 @@ void FMobileSceneRenderer::Render(FRDGBuilder& GraphBuilder, const FSceneRenderU
 
 	FRDGExternalAccessQueue ExternalAccessQueue;
 
+	// If we're rendering with async compute then render shadows after the prepass
+	// We don't do this in the general case because it is slower on some platforms
+	extern int32 GDFShadowAsyncCompute;
+	const bool bRenderShadowsAfterPrepass = !!GDFShadowAsyncCompute && bIsFullDepthPrepassEnabled;
+
 	if (SceneUpdateInputs)
 	{
 		static auto CVarDistanceFieldShadowQuality = IConsoleManager::Get().FindConsoleVariable(TEXT("r.DFShadowQuality"));
@@ -1328,13 +1333,10 @@ void FMobileSceneRenderer::Render(FRDGBuilder& GraphBuilder, const FSceneRenderU
 				// RunHairStrandsBookmark(GraphBuilder, EHairStrandsBookmark::ProcessCardsAndMeshesInterpolation_ShadowView, HairStrandsBookmarkParameters);
 			}
 		}
-
-		RenderShadowDepthMaps(GraphBuilder, nullptr, InstanceCullingManager, ExternalAccessQueue);
-		GraphBuilder.AddDispatchHint();
-
-		if (ShouldRenderVolumetricFog())
+		if (!bRenderShadowsAfterPrepass)
 		{
-			ComputeVolumetricFog(GraphBuilder, SceneTextures);
+			RenderShadowDepthMaps(GraphBuilder, nullptr, InstanceCullingManager, ExternalAccessQueue);
+			GraphBuilder.AddDispatchHint();
 		}
 		ExternalAccessQueue.Submit(GraphBuilder);
 
@@ -1430,7 +1432,6 @@ void FMobileSceneRenderer::Render(FRDGBuilder& GraphBuilder, const FSceneRenderU
 		}
 	}
 
-	FDBufferTextures DBufferTextures{};
 	if (bIsFullDepthPrepassEnabled)
 	{
 		RenderFullDepthPrepass(GraphBuilder, Views, SceneTextures, InstanceCullingManager);
@@ -1442,7 +1443,32 @@ void FMobileSceneRenderer::Render(FRDGBuilder& GraphBuilder, const FSceneRenderU
 
 		SceneTextures.MobileSetupMode = EMobileSceneTextureSetupMode::SceneDepth;
 		SceneTextures.MobileUniformBuffer = CreateMobileSceneTextureUniformBuffer(GraphBuilder, &SceneTextures, SceneTextures.MobileSetupMode);
+	}
+
+	if (bRendererOutputFinalSceneColor)
+	{
+		if (bRenderShadowsAfterPrepass)
+		{
+			// Kick async DF shadows to run concurrently with shadow depths
+			// Note: this is a no-op if async DF shadows is disabled
+			BeginAsyncDistanceFieldShadowProjections(GraphBuilder, SceneTextures, InitViewTaskDatas.DynamicShadows);
 
+			// Render shadows after the prepass so we can overap async compute work with shadows
+			RenderShadowDepthMaps(GraphBuilder, nullptr, InstanceCullingManager, ExternalAccessQueue);
+			GraphBuilder.AddDispatchHint();
+			ExternalAccessQueue.Submit(GraphBuilder);
+		}
+
+		if (ShouldRenderVolumetricFog())
+		{
+			ComputeVolumetricFog(GraphBuilder, SceneTextures);
+		}
+	}
+
+
+	FDBufferTextures DBufferTextures{};
+	if (bIsFullDepthPrepassEnabled)
+	{
 		// When renderer is in ERendererOutput::DepthPrepassOnly mode, bShouldRenderHZB is set to false in InitViews()
 		if (bShouldRenderHZB)
 		{
