Fixed showflag.rendering to properly clear and cleanup renderers.

diff --git a/Engine/Source/Runtime/Renderer/Private/MobileShadingRenderer.cpp b/Engine/Source/Runtime/Renderer/Private/MobileShadingRenderer.cpp
--- a/Engine/Source/Runtime/Renderer/Private/MobileShadingRenderer.cpp
+++ b/Engine/Source/Runtime/Renderer/Private/MobileShadingRenderer.cpp
@@ -1047,11 +1047,6 @@ void FMobileSceneRenderer::RenderCustomRenderPassBasePass(FRDGBuilder& GraphBuil
 
 void FMobileSceneRenderer::Render(FRDGBuilder& GraphBuilder, const FSceneRenderUpdateInputs* SceneUpdateInputs)
 {
-	if (!ViewFamily.EngineShowFlags.Rendering)
-	{
-		return;
-	}
-
 	const ERendererOutput RendererOutput = GetRendererOutput();
 	const bool bRendererOutputFinalSceneColor = (RendererOutput != ERendererOutput::DepthPrepassOnly);
 
diff --git a/Engine/Source/Runtime/Renderer/Private/SceneRenderBuilder.cpp b/Engine/Source/Runtime/Renderer/Private/SceneRenderBuilder.cpp
--- a/Engine/Source/Runtime/Renderer/Private/SceneRenderBuilder.cpp
+++ b/Engine/Source/Runtime/Renderer/Private/SceneRenderBuilder.cpp
@@ -575,25 +575,25 @@ public:
 			CommonShowFlags |= Renderer->ViewFamily.EngineShowFlags.PathTracing ? ESceneRenderCommonShowFlags::PathTracing : ESceneRenderCommonShowFlags::None;
 		}
 
-		Renderers.Emplace(Renderer);
+		FGroupNode* TailGroup = GroupNodes.GetTail();
 
-		if (Renderer->ViewFamily.EngineShowFlags.Rendering)
-		{
-			FGroupNode* TailGroup = GroupNodes.GetTail();
+		FRenderNode* RenderNode = Allocator.Create<FRenderNode>(
+			Renderer,
+			Forward<FString&&>(Name),
+			Forward<FSceneRenderFunction&&>(Function),
+			TailGroup);
 
-			FRenderNode* RenderNode = Allocator.Create<FRenderNode>(
-				Renderer,
-				Forward<FString&&>(Name),
-				Forward<FSceneRenderFunction&&>(Function),
-				TailGroup);
+		if (TailGroup)
+		{
+			TailGroup->RenderNodes.AddTail(RenderNode);
+		}
+		RenderNodes.AddTail(RenderNode);
 
-			if (TailGroup)
-			{
-				TailGroup->RenderNodes.AddTail(RenderNode);
-			}
-			RenderNodes.AddTail(RenderNode);
+		Ops.Emplace(FOp::Render(RenderNode));
+		Renderers.Emplace(Renderer);
 
-			Ops.Emplace(FOp::Render(RenderNode));
+		if (Renderer->ViewFamily.EngineShowFlags.Rendering)
+		{
 			ActiveRenderers.Emplace(Renderer);
 			ActiveViewFamilies.Emplace(&Renderer->ViewFamily);
 			ActiveViews.Reserve(ActiveViews.Num() + Renderer->Views.Num());
@@ -848,8 +848,8 @@ void FSceneRenderProcessor::Execute()
 					RenderState.FullPath += RenderNode.Name;
 				}
 
-				const bool bFirstRenderer = RenderNode.Renderer == ActiveRenderers[0];
-				const bool bLastRenderer  = RenderNode.Renderer == ActiveRenderers.Last();
+				const bool bFirstRenderer = !ActiveRenderers.IsEmpty() && RenderNode.Renderer == ActiveRenderers[0];
+				const bool bLastRenderer  = !ActiveRenderers.IsEmpty() && RenderNode.Renderer == ActiveRenderers.Last();
 
 				TOptional<FSceneRenderUpdateInputs> SceneUpdateInputs;
 
@@ -884,7 +884,12 @@ void FSceneRenderProcessor::Execute()
 					VISUALIZE_TEXTURE_BEGIN_VIEW(Scene->GetFeatureLevel(), Renderer.Views[0].GetViewKey(), FunctionInputs.FullPath, Renderer.Views[0].bIsSceneCapture);
 				}
 	
-				const bool bRenderCalled = RenderNode.Function(GraphBuilder, FunctionInputs);
+				bool bRenderCalled = false;
+
+				if (Renderer.ViewFamily.EngineShowFlags.Rendering)
+				{
+					bRenderCalled = RenderNode.Function(GraphBuilder, FunctionInputs);
+				}
 
 				if (SceneUpdateInputs)
 				{
