[Material Editor] Add two experimental custom output node `Motion Vector World Offset (Per-Pixel)` and `Temporal Responsiveness` to allow user modify per pixel motion vector and set how responsive the temporal history is. Visible Materials with any pin connected will cause the corresponding mesh to render in the velocity pass. Please see node tooltips for actual support progress.
1. Temporal Responsiveness. Describes how temporal history will be rejected for different velocity mismatch level. a) Default :0. b) Medium (0,0.5]. c) Full (0.5, 1.0]. Now, it will be used by TSR to change rejection heuristics.
2. Motion Vector World Offset (Per-Pixel). Works similar to the `Previous Frame` input of `Previous Frame Switch` but in the pixel shader. Regions with invalid velocity will be set to the output of the vertex shader. This function currently supports non-nanite meshes only in non-basepass velocity write and is implemented in two passes. Velocity pass writes out an offset texture. In a second pass, we use the offset texture and the velocity texture to derive proper per-pixel motion vector.

Turn on by setting r.Velocity.TemporalResponsiveness.Supported and r.Velocity.PixelShaderMotionVectorWorldOffset.Supported to one.

#tests test scene with the new added node

diff --git a/Engine/Shaders/Private/BasePassPixelShader.usf b/Engine/Shaders/Private/BasePassPixelShader.usf
--- a/Engine/Shaders/Private/BasePassPixelShader.usf
+++ b/Engine/Shaders/Private/BasePassPixelShader.usf
@@ -1202,8 +1202,8 @@ void FPixelShaderInOut_MainPS(
 	#else
 		float3 Velocity = Calculate3DVelocity(MaterialParameters.ScreenPosition, BasePassInterpolants.VelocityPrevScreenPosition);
 	#endif
-
-		float4 EncodedVelocity = EncodeVelocityToTexture(Velocity, (GetPrimitiveData(MaterialParameters).Flags & PRIMITIVE_SCENE_DATA_FLAG_HAS_PIXEL_ANIMATION) != 0);
+		float TemporalResponsiveness = GetMaterialTemporalResponsiveness(MaterialParameters);	
+		float4 EncodedVelocity = EncodeVelocityToTexture(Velocity, (GetPrimitiveData(MaterialParameters).Flags & PRIMITIVE_SCENE_DATA_FLAG_HAS_PIXEL_ANIMATION) != 0, TemporalResponsiveness);
 
 	#if USES_GBUFFER
 		GBuffer.Velocity = EncodedVelocity;
diff --git a/Engine/Shaders/Private/Common.ush b/Engine/Shaders/Private/Common.ush
--- a/Engine/Shaders/Private/Common.ush
+++ b/Engine/Shaders/Private/Common.ush
@@ -2042,10 +2042,21 @@ struct FScreenVertexOutput
 /** Whether HasPixelAnimation is encoded in the velocity texture. Matches VelocityEncodeHasPixelAnimation() */
 #define VELOCITY_ENCODE_HAS_PIXEL_ANIMATION VELOCITY_ENCODE_DEPTH
 
+/** Whether we use more bits to simulate temporal responsiveness **/
+#define VELOCITY_ENCODE_TEMPORAL_RESPONSIVENESS (VELOCITY_SUPPORT_TEMPORAL_RESPONSIVENESS & VELOCITY_ENCODE_HAS_PIXEL_ANIMATION)
+
+#if VELOCITY_ENCODE_TEMPORAL_RESPONSIVENESS
+	#define TEMPORAL_RESPONSIVENESS_MASK	0x3
+	#define VELOCITY_Z_LOW_MASK			0xFFFC
+#else
+	#define TEMPORAL_RESPONSIVENESS_MASK	0x1
+	#define VELOCITY_Z_LOW_MASK			0xFFFE
+#endif
+	
 // for velocity rendering, motionblur and temporal AA
 // velocity needs to support -2..2 screen space range for x and y
 // texture is 16bit 0..1 range per channel
-ENCODED_VELOCITY_TYPE EncodeVelocityToTexture(float3 V, bool bHasPixelAnimation)
+ENCODED_VELOCITY_TYPE EncodeVelocityToTexture(float3 V, uint TemporalResponsivenessMask)
 {
 	#if VELOCITY_ENCODE_GAMMA
 		V.xy = sign(V.xy) * sqrt(abs(V.xy)) * (2.0 / sqrt(2.0));
@@ -2060,7 +2071,7 @@ ENCODED_VELOCITY_TYPE EncodeVelocityToTexture(float3 V, bool bHasPixelAnimation)
 		uint Vz = asuint(V.z);
 
 		EncodedV.z = saturate(float((Vz >> 16) & 0xFFFF) * rcp(65535.0f) + (0.1 / 65535.0f));
-		EncodedV.w = saturate(float(((Vz >>  0) & 0xFFFE) | uint(bHasPixelAnimation)) * rcp(65535.0f) + (0.1 / 65535.0f));
+		EncodedV.w = saturate(float(((Vz >>  0) & VELOCITY_Z_LOW_MASK) | (TemporalResponsivenessMask & TEMPORAL_RESPONSIVENESS_MASK)) * rcp(65535.0f) + (0.1 / 65535.0f));
 	#else
 		EncodedV.zw = 0.0;
 	#endif
@@ -2072,9 +2083,23 @@ ENCODED_VELOCITY_TYPE EncodeVelocityToTexture(float3 V, bool bHasPixelAnimation)
 #endif
 }
 
+ENCODED_VELOCITY_TYPE EncodeVelocityToTexture(float3 V,  bool bHasPixelAnimation, float TemporalResponsiveness)
+{
+	uint TemporalResponsivenessMask = uint(bHasPixelAnimation);
+#if VELOCITY_ENCODE_TEMPORAL_RESPONSIVENESS
+	TemporalResponsivenessMask |= select(TemporalResponsiveness > 0.5f, 3U, uint(TemporalResponsiveness>0));
+#endif
+	return EncodeVelocityToTexture(V, TemporalResponsivenessMask);
+}
+
+ENCODED_VELOCITY_TYPE EncodeVelocityToTexture(float3 V, bool bHasPixelAnimation)
+{
+	return EncodeVelocityToTexture(V, /* TemporalResponsivenessMask = */uint(bHasPixelAnimation));
+}
+
 ENCODED_VELOCITY_TYPE EncodeVelocityToTexture(float3 V)
 {
-	return EncodeVelocityToTexture(V, /* bHasPixelAnimation = */ false);
+	return EncodeVelocityToTexture(V, /* TemporalResponsivenessMask = */ 0);
 }
 
 // see EncodeVelocityToTexture()
@@ -2091,7 +2116,7 @@ float3 DecodeVelocityFromTexture(ENCODED_VELOCITY_TYPE InEncodedV)
 	V.xy = EncodedV.xy * InvDiv - 32767.0f / 65535.0f * InvDiv;
 
 	#if VELOCITY_ENCODE_DEPTH
-		V.z = asfloat((uint(round(EncodedV.z * 65535.0f)) << 16) | (uint(round(EncodedV.w * 65535.0f)) & 0xFFFE));
+		V.z = asfloat((uint(round(EncodedV.z * 65535.0f)) << 16) | (uint(round(EncodedV.w * 65535.0f)) & VELOCITY_Z_LOW_MASK));
 	#else
 		V.z = 0.0;
 	#endif
@@ -2111,6 +2136,13 @@ bool DecodeHasPixelAnimationFromVelocityTexture(ENCODED_VELOCITY_TYPE EncodedV)
 }
 #endif
 
+#if VELOCITY_ENCODE_TEMPORAL_RESPONSIVENESS
+bool DecodeDisableAnimationAAFromVelocityTexture(ENCODED_VELOCITY_TYPE EncodedV)
+{
+	return (uint(round(EncodedV.w * 65535.0f)) & 0x2) != 0x0;
+}
+#endif
+
 // Used for the Global Illumination in the GIReplace material expression
 bool GetGIReplaceState()
 {
diff --git a/Engine/Shaders/Private/MaterialTemplate.ush b/Engine/Shaders/Private/MaterialTemplate.ush
--- a/Engine/Shaders/Private/MaterialTemplate.ush
+++ b/Engine/Shaders/Private/MaterialTemplate.ush
@@ -4015,6 +4015,24 @@ float3 GetMaterialPreviousWorldPositionOffset(FMaterialVertexParameters Paramete
 	return float3(0, 0, 0);
 }
 
+float GetMaterialTemporalResponsiveness(FMaterialPixelParameters Parameters)
+{
+	#if HAVE_GetTemporalResponsivenessOutput0 > 0
+		return saturate(GetTemporalResponsivenessOutput0(Parameters));
+	#else
+		return 0;
+	#endif
+}
+
+float3 GetMaterialMotionVectorWorldOffset(FMaterialPixelParameters Parameters)
+{
+	#if HAVE_GetMotionVectorWorldOffsetOutput0 > 0
+		return GetMotionVectorWorldOffsetOutput0(Parameters);
+	#else
+		return float3(0, 0, 0);
+	#endif
+}
+
 // .rgb:SubsurfaceColor, .a:SSProfileId in 0..1 range
 MaterialFloat4 GetMaterialSubsurfaceDataRaw(FPixelMaterialInputs PixelMaterialInputs)
 {
diff --git a/Engine/Shaders/Private/MotionBlur/VisualizeMotionVectors.usf b/Engine/Shaders/Private/MotionBlur/VisualizeMotionVectors.usf
--- a/Engine/Shaders/Private/MotionBlur/VisualizeMotionVectors.usf
+++ b/Engine/Shaders/Private/MotionBlur/VisualizeMotionVectors.usf
@@ -108,8 +108,14 @@ void MainPS(
 		
 		float4 EncodedVelocity = VelocityTexture.SampleLevel(VelocitySampler, VelocityTextureUV, 0.0);
 		bool bHasPixelAnimation = DecodeHasPixelAnimationFromVelocityTexture(EncodedVelocity);
-		
+		#if VELOCITY_ENCODE_TEMPORAL_RESPONSIVENESS
+			bool bDisableAnimationAA = DecodeDisableAnimationAAFromVelocityTexture(EncodedVelocity);
+		#else
+			bool bDisableAnimationAA = false;
+		#endif
+
 		OutColor = bHasPixelAnimation ? 1.0 : 0.0;
+		OutColor.rgb = lerp(OutColor.rgb,float3(1,0,0), bDisableAnimationAA ? 0.5 : 0); 
 	}
 	#endif
 	else
diff --git a/Engine/Shaders/Private/TemporalSuperResolution/TSRDecimateHistory.usf b/Engine/Shaders/Private/TemporalSuperResolution/TSRDecimateHistory.usf
--- a/Engine/Shaders/Private/TemporalSuperResolution/TSRDecimateHistory.usf
+++ b/Engine/Shaders/Private/TemporalSuperResolution/TSRDecimateHistory.usf
@@ -269,6 +269,9 @@ void MainCS(
 	float PrevDeviceZ;
 	float DeviceZError;
 	bool bHasPixelAnimation;
+#if VELOCITY_ENCODE_TEMPORAL_RESPONSIVENESS
+	bool bDisableAnimationAA = false;
+#endif
 	bool bHasReprojectionOffset;
 	tsr_half ReprojectionEdge;
 	ISOLATE
@@ -288,7 +291,12 @@ void MainCS(
 			uint EncodedDeviceZError = DepthErrorTexture[InputPixelPos];
 			#if CONFIG_RESURRECTION_REPROJECTION
 				float2 DeviceZAndPrevDeviceZ = ClosestDepthTexture[InputPixelPos];
-				DeviceZ = DeviceZAndPrevDeviceZ.g;
+				#if VELOCITY_ENCODE_TEMPORAL_RESPONSIVENESS
+					bDisableAnimationAA = DeviceZAndPrevDeviceZ.g < 0;
+					DeviceZ = abs(DeviceZAndPrevDeviceZ.g);
+				#else
+					DeviceZ = DeviceZAndPrevDeviceZ.g;
+				#endif
 			#else
 				float DeviceZAndPrevDeviceZ = ClosestDepthTexture[InputPixelPos];
 				DeviceZ = 0.0;
@@ -462,6 +470,9 @@ void MainCS(
 		select(bIsParallaxDisocclusion,  tsr_ushort(0x02), tsr_ushort(0x0)) |
 		select(bHasPixelAnimation,       tsr_ushort(0x04), tsr_ushort(0x0)) |
 		select(bReprojectionHollFill,    tsr_ushort(0x08), tsr_ushort(0x0)) |
+#if VELOCITY_ENCODE_TEMPORAL_RESPONSIVENESS 
+		select(bDisableAnimationAA,      tsr_ushort(0x20), tsr_ushort(0x0)) |
+#endif
 		select(bIsResurrectionOffScreen, tsr_ushort(0x10), tsr_ushort(0x0));
 
 	// Output full res
diff --git a/Engine/Shaders/Private/TemporalSuperResolution/TSRDepthVelocityAnalysis.ush b/Engine/Shaders/Private/TemporalSuperResolution/TSRDepthVelocityAnalysis.ush
--- a/Engine/Shaders/Private/TemporalSuperResolution/TSRDepthVelocityAnalysis.ush
+++ b/Engine/Shaders/Private/TemporalSuperResolution/TSRDepthVelocityAnalysis.ush
@@ -112,7 +112,8 @@ void FetchAndComputeScreenVelocity(
 	float ClosestDeviceZ,
 	tsr_short2 ReprojectionOffset,
 	out float3 ScreenVelocity,
-	out bool bHasPixelAnimation)
+	out bool bHasPixelAnimation,
+	out bool bDisableAnimationAA)
 {
 	float4 EncodedVelocity = SceneVelocityTexture[InputPixelPos + ReprojectionOffset];
 
@@ -121,6 +122,12 @@ void FetchAndComputeScreenVelocity(
 	#else
 		bHasPixelAnimation = false;
 	#endif
+
+	#if VELOCITY_ENCODE_TEMPORAL_RESPONSIVENESS
+		bDisableAnimationAA = DecodeDisableAnimationAAFromVelocityTexture(EncodedVelocity);
+	#else
+		bDisableAnimationAA = false;
+	#endif
 	
 	#if CONFIG_USE_SCENE_VELOCITY
 	BRANCH
diff --git a/Engine/Shaders/Private/TemporalSuperResolution/TSRDilateVelocity.usf b/Engine/Shaders/Private/TemporalSuperResolution/TSRDilateVelocity.usf
--- a/Engine/Shaders/Private/TemporalSuperResolution/TSRDilateVelocity.usf
+++ b/Engine/Shaders/Private/TemporalSuperResolution/TSRDilateVelocity.usf
@@ -169,6 +169,7 @@ void MainCS(
 	uint EncodedReprojectionJacobian;
 	uint EncodedDeviceZError;
 	bool bHasPixelAnimation;
+	bool bDisableAnimationAA;
 	bool bHasReprojectionOffset;
 	tsr_half ReprojectionEdge;
 	BRANCH
@@ -199,6 +200,9 @@ void MainCS(
 		#if VELOCITY_ENCODE_HAS_PIXEL_ANIMATION
 			bHasPixelAnimation = DecodeHasPixelAnimationFromVelocityTexture(AccessNeighborhoodCenter(VelocityNeighborhood));
 		#endif
+		#if VELOCITY_ENCODE_TEMPORAL_RESPONSIVENESS
+			bDisableAnimationAA = DecodeDisableAnimationAAFromVelocityTexture(AccessNeighborhoodCenter(VelocityNeighborhood));
+		#endif
 	
 		tsr_half2x2 ReprojectionJacobian;
 		float ParallaxDepthError;
@@ -230,14 +234,16 @@ void MainCS(
 			ScreenVelocity = AccessNeighborhoodCenter(PixelVelocityNeighborhood) * InputPixelVelocityToScreenVelocity;
 	
 			// Compute the dilated screen velocity.
-			bool bDildatedHasPixelAnimation;
+			bool bDilatedHasPixelAnimation;
+			bool bDilatedDisableAnimationAA;
 			FetchAndComputeScreenVelocity(
 				SceneVelocityTexture,
 				InputPixelPos,
 				ClosestDeviceZ,
 				ReprojectionOffset,
 				/* out */ DilatedScreenVelocity,
-				/* out */ bDildatedHasPixelAnimation);
+				/* out */ bDilatedHasPixelAnimation,
+				/* out */ bDilatedDisableAnimationAA);
 		}
 		
 		ReprojectionEdge = ComputeReprojectionEdge(ScreenVelocity, DilatedScreenVelocity.xy);
@@ -293,17 +299,20 @@ void MainCS(
 			DeviceZ,
 			/* ReprojectionOffset = */ tsr_short2(0, 0),
 			/* out */ ScreenVelocity,
-			/* out */ bHasPixelAnimation);
+			/* out */ bHasPixelAnimation,
+			/* out */ bDisableAnimationAA);
 
 		// Compute the dilated screen vecocity.
 		bool bDilatedHasPixelAnimation;
+		bool bDilatedDisableAnimationAA;
 		FetchAndComputeScreenVelocity(
 			SceneVelocityTexture,
 			InputPixelPos,
 			ClosestDeviceZ,
 			ReprojectionOffset,
 			/* out */ DilatedScreenVelocity,
-			/* out */ bDilatedHasPixelAnimation);
+			/* out */ bDilatedHasPixelAnimation,
+			/* out */ bDilatedDisableAnimationAA);
 		
 		ReprojectionEdge = ComputeReprojectionEdge(ScreenVelocity.xy, DilatedScreenVelocity.xy);
 		
@@ -369,7 +378,13 @@ void MainCS(
 		bool bValidOutputPixel = all(InputPixelPos < InputInfo_ViewportMax);
 		uint2 OutputPixelPos = bValidOutputPixel ? InputPixelPos : uint(~0).xx;
 
-		ClosestDepthOutput[OutputPixelPos] = float2(select(bHasPixelAnimation, -PrevClosestDeviceZ, PrevClosestDeviceZ), ClosestDeviceZ);
+		ClosestDepthOutput[OutputPixelPos] = float2(select(bHasPixelAnimation, -PrevClosestDeviceZ, PrevClosestDeviceZ),
+		#if VELOCITY_ENCODE_TEMPORAL_RESPONSIVENESS 
+													select(bDisableAnimationAA, -ClosestDeviceZ, ClosestDeviceZ)
+		#else
+													ClosestDeviceZ
+		#endif
+													);
 
 		R8Output[uint3(OutputPixelPos, 0)] = uint(round(ReprojectionEdge * 127.0)) | select(bHasReprojectionOffset, 0x80u, 0x00u);
 		R8Output[uint3(OutputPixelPos, 1)] = EncodedDeviceZError;
diff --git a/Engine/Shaders/Private/TemporalSuperResolution/TSRRejectShading.usf b/Engine/Shaders/Private/TemporalSuperResolution/TSRRejectShading.usf
--- a/Engine/Shaders/Private/TemporalSuperResolution/TSRRejectShading.usf
+++ b/Engine/Shaders/Private/TemporalSuperResolution/TSRRejectShading.usf
@@ -156,9 +156,9 @@ void FetchDecimateMask(
 		tsr_ushort BitMask = tsr_ushort(DecimateMask.r * tsr_half(255.0));
 
 		VelocityEdge.SetElement(ElementIndex, DecimateMask.g);
-		bIsDisoccluded.SetElement(ElementIndex,     (BitMask & tsr_ushort(0x3)) != tsr_ushort(0));
+		bIsDisoccluded.SetElement(ElementIndex,     (BitMask & tsr_ushort(0x23)) != tsr_ushort(0));
 		bHasPixelAnimation.SetElement(ElementIndex, (BitMask & tsr_ushort(0x4)) != tsr_ushort(0));
-		bIsResurrectionOffScreen.SetElement(ElementIndex, (BitMask & tsr_ushort(0x10)) != tsr_ushort(0));
+		bIsResurrectionOffScreen.SetElement(ElementIndex, (BitMask & tsr_ushort(0x30)) != tsr_ushort(0));
 	}
 } // FetchDecimateMask()
 
diff --git a/Engine/Shaders/Private/TemporalSuperResolution/TSRVisualize.usf b/Engine/Shaders/Private/TemporalSuperResolution/TSRVisualize.usf
--- a/Engine/Shaders/Private/TemporalSuperResolution/TSRVisualize.usf
+++ b/Engine/Shaders/Private/TemporalSuperResolution/TSRVisualize.usf
@@ -157,6 +157,7 @@ struct FDecimateMaskSample
 	bool bIsParallaxDisocclusion;
 	bool bIsDisoccluded;
 	bool bHasPixelAnimation;
+	bool bDisableAnimationAA;
 	bool bReprojectionHollFill;
 };
 
@@ -167,11 +168,12 @@ FDecimateMaskSample SampleDecimateMask(FContext Ctx)
 	tsr_half2 DecimateMask = DecimateMaskTexture[Ctx.InputPixelPos];
 	
 	tsr_ushort BitMask = tsr_ushort(DecimateMask.r * tsr_half(255.0));
-	Decimate.bIsOffScreen            = (BitMask & tsr_ushort(0x1)) != tsr_ushort(0);
-	Decimate.bIsParallaxDisocclusion = (BitMask & tsr_ushort(0x2)) != tsr_ushort(0);
-	Decimate.bIsDisoccluded          = (BitMask & tsr_ushort(0x3)) != tsr_ushort(0);
-	Decimate.bHasPixelAnimation      = (BitMask & tsr_ushort(0x4)) != tsr_ushort(0);
-	Decimate.bReprojectionHollFill   = (BitMask & tsr_ushort(0x8)) != tsr_ushort(0);
+	Decimate.bIsOffScreen            = (BitMask & tsr_ushort(0x01)) != tsr_ushort(0);
+	Decimate.bIsParallaxDisocclusion = (BitMask & tsr_ushort(0x02)) != tsr_ushort(0);
+	Decimate.bIsDisoccluded          = (BitMask & tsr_ushort(0x03)) != tsr_ushort(0);
+	Decimate.bHasPixelAnimation      = (BitMask & tsr_ushort(0x04)) != tsr_ushort(0);
+	Decimate.bReprojectionHollFill   = (BitMask & tsr_ushort(0x08)) != tsr_ushort(0);
+	Decimate.bDisableAnimationAA	 = (BitMask & tsr_ushort(0x20)) != tsr_ushort(0);
 	return Decimate;
 }
 
diff --git a/Engine/Shaders/Private/VelocityCommon.ush b/Engine/Shaders/Private/VelocityCommon.ush
--- a/Engine/Shaders/Private/VelocityCommon.ush
+++ b/Engine/Shaders/Private/VelocityCommon.ush
@@ -6,10 +6,10 @@
 
 #pragma once
 
-float3 Calculate3DVelocity(float4 PackedVelocityA, float4 PackedVelocityC)
+float3 Calculate3DVelocityBase(float4 PackedVelocityA, float4 PackedVelocityC, float2 JitterA, float2 JitterC)
 {
-	float2 ScreenPos = PackedVelocityA.xy / PackedVelocityA.w - ResolvedView.TemporalAAJitter.xy;
-	float2 PrevScreenPos = PackedVelocityC.xy / PackedVelocityC.w - ResolvedView.TemporalAAJitter.zw;
+	float2 ScreenPos = PackedVelocityA.xy / PackedVelocityA.w - JitterA;
+	float2 PrevScreenPos = PackedVelocityC.xy / PackedVelocityC.w - JitterC;
 
 	float DeviceZ = PackedVelocityA.z / PackedVelocityA.w;
 	float PrevDeviceZ = PackedVelocityC.z / PackedVelocityC.w;
@@ -20,3 +20,14 @@ float3 Calculate3DVelocity(float4 PackedVelocityA, float4 PackedVelocityC)
 	// Make sure not to touch 0,0 which is clear color
 	return Velocity;
 }
+
+float3 Calculate3DVelocity(float4 PackedVelocityA, float4 PackedVelocityC)
+{
+	return Calculate3DVelocityBase(PackedVelocityA, PackedVelocityC, ResolvedView.TemporalAAJitter.xy, ResolvedView.TemporalAAJitter.zw);
+}
+
+// Calculated the distance between packed velocity in the current frame
+float3 Calculate3DVelocityOffset(float4 PackedVelocityA, float4 PackedVelocityC)
+{
+	return Calculate3DVelocityBase(PackedVelocityA, PackedVelocityC, ResolvedView.TemporalAAJitter.xy, ResolvedView.TemporalAAJitter.xy);
+}
\ No newline at end of file
diff --git a/Engine/Shaders/Private/VelocityShader.usf b/Engine/Shaders/Private/VelocityShader.usf
--- a/Engine/Shaders/Private/VelocityShader.usf
+++ b/Engine/Shaders/Private/VelocityShader.usf
@@ -108,6 +108,8 @@ void MainVertexShader(
 #define VELOCITY_THIN_TRANSLUCENT_MODE 0
 #endif
 
+#define OUPUT_MOTION_VECTOR_WORLD_OFFSET (!TRANSLUCENCY_VELOCITY_FROM_DEPTH && !VELOCITY_CLIPPED_DEPTH_PASS && VELOCITY_SUPPORT_PS_MOTION_VECTOR_WORLD_OFFSET)
+
 void MainPixelShader(
 	in INPUT_POSITION_QUALIFIERS float4 SvPosition : SV_Position,
 	FVelocityInterpsVSToPS Input,
@@ -116,6 +118,9 @@ void MainPixelShader(
 	OPTIONAL_OutDepthConservative
 
 	,out ENCODED_VELOCITY_TYPE OutColor : SV_Target0
+#if OUPUT_MOTION_VECTOR_WORLD_OFFSET
+	,out ENCODED_VELOCITY_TYPE OutMotionVectorWorldOffset : SV_Target1
+#endif
 #if MATERIALBLENDING_MASKED_USING_COVERAGE
 	, out uint OutCoverage : SV_Coverage
 #endif
@@ -206,10 +211,27 @@ void MainPixelShader(
 #endif //MATERIALBLENDING_MASKED_USING_COVERAGE
 
 #if !TRANSLUCENCY_VELOCITY_FROM_DEPTH
-	// 3d velocity, includes camera an object motion
-	float3 Velocity = Calculate3DVelocity(ScreenPosition, Input.PackedVelocityC);
 
-	OutColor = EncodeVelocityToTexture(Velocity, (GetPrimitiveData(MaterialParameters).Flags & PRIMITIVE_SCENE_DATA_FLAG_HAS_PIXEL_ANIMATION) != 0);
+	float3 Velocity;
+	float TemporalResponsiveness = GetMaterialTemporalResponsiveness(MaterialParameters);
+	bool bHasPixelAnimationFlag = (GetPrimitiveData(MaterialParameters).Flags & PRIMITIVE_SCENE_DATA_FLAG_HAS_PIXEL_ANIMATION) != 0;
+ 
+	#if (!VELOCITY_CLIPPED_DEPTH_PASS && VELOCITY_SUPPORT_PS_MOTION_VECTOR_WORLD_OFFSET) // Opaque pass that uses motion vector offset in PS.
+		
+		float3 OffsetTranslatedWorldPosition = GetMaterialMotionVectorWorldOffset(MaterialParameters) + GetTranslatedWorldPosition(MaterialParameters);
+		float4 ClipPosition = mul(float4(OffsetTranslatedWorldPosition, 1.0f), ResolvedView.TranslatedWorldToClip);
+
+		{
+			// Write the offset at the same time so that we can recover the actual velocity to previous frame.
+			OutMotionVectorWorldOffset = EncodeVelocityToTexture(Calculate3DVelocityOffset(ScreenPosition, ClipPosition), bHasPixelAnimationFlag, TemporalResponsiveness);
+
+			Velocity = Calculate3DVelocity(ScreenPosition, Input.PackedVelocityC);
+		}
+	#else	
+		Velocity = Calculate3DVelocity(ScreenPosition, Input.PackedVelocityC);
+	#endif // !VELOCITY_CLIPPED_DEPTH_PASS
+
+	OutColor = EncodeVelocityToTexture(Velocity, bHasPixelAnimationFlag, TemporalResponsiveness);
 #else
 	// We only want velocity generated from camera movement and depth so write zero which is the special clear color value.
 	OutColor = 0;
diff --git a/Engine/Shaders/Private/VelocityUpdate.usf b/Engine/Shaders/Private/VelocityUpdate.usf
new file mode 100644
--- /dev/null
+++ b/Engine/Shaders/Private/VelocityUpdate.usf
@@ -0,0 +1,92 @@
+// Copyright Epic Games, Inc. All Rights Reserved.
+
+/*=============================================================================
+	VelocityUpdate.usf: Make any update to velocity vectors.
+=============================================================================*/
+
+#include "Common.ush"
+#include "VelocityCommon.ush"
+#include "SceneTexturesCommon.ush"
+
+Texture2D<float4> DepthTexture;
+Texture2D<ENCODED_VELOCITY_TYPE> VelocityTexture;
+RWTexture2D<ENCODED_VELOCITY_TYPE> RWMotionVectorWorldOffset;
+
+bool PointInView(float2 Position)
+{
+	return (all(Position >= ResolvedView.ViewRectMin.xy) &&  
+			all(Position < ResolvedView.ViewRectMin.xy + ResolvedView.ViewSizeAndInvSize.xy));
+}
+
+float3 DecodeVelocityTexture(Texture2D<ENCODED_VELOCITY_TYPE> Velocity, float2 Position)
+{
+	int2 IntPosition = int2(Position);
+	//TODO: sub-pixel hole filing.
+	return DecodeVelocityFromTexture(Velocity[IntPosition + int2(0, 0)]);
+}
+
+float CalcDeviceZ(float2 Position)
+{
+	int2 IntPosition = int2(Position);
+	#if 1
+		return DepthTexture.Load(int3(IntPosition, 0)).r;
+	#else
+		float2 Weight = frac(Position);
+		float P00 = DepthTexture.Load(int3(IntPosition + int2(0, 0), 0)).r;
+		float P01 = DepthTexture.Load(int3(IntPosition + int2(0, 1), 0)).r;
+		float P10 = DepthTexture.Load(int3(IntPosition + int2(1, 0), 0)).r;
+		float P11 = DepthTexture.Load(int3(IntPosition + int2(1, 1), 0)).r;
+		return lerp(lerp(P00, P10, Weight.x),lerp(P01,P11, Weight.x),Weight.y);
+	#endif
+}
+
+[numthreads(THREADGROUP_SIZE, THREADGROUP_SIZE, 1)]
+void MainCS(in uint3 ThreadID : SV_DispatchThreadID)
+{
+	ResolvedView = ResolveView();
+	int2 Position = ThreadID.xy;
+	if (any(Position >= ResolvedView.ViewSizeAndInvSize.xy))
+	{
+		return;
+	}
+	int2 TexelPosition = Position + ResolvedView.ViewRectMin.xy;
+ 	
+	ENCODED_VELOCITY_TYPE EncodedVectorOffset = RWMotionVectorWorldOffset[TexelPosition];
+	uint TemporalResponsivenessMask = 0;
+	#if (VELOCITY_ENCODE_HAS_PIXEL_ANIMATION | VELOCITY_ENCODE_TEMPORAL_RESPONSIVENESS)	
+		TemporalResponsivenessMask |= (uint(round(EncodedVectorOffset.w * 65535.0f)) & TEMPORAL_RESPONSIVENESS_MASK);
+	#endif
+	float3 DecodedVectorOffset = DecodeVelocityFromTexture(EncodedVectorOffset);
+	bool bHasVelocity = EncodedVectorOffset.x > 0;
+	float2 OffsetTexelPosition = TexelPosition - (ScreenPosToViewportUV(DecodedVectorOffset.xy) - 0.5f/*remove bias*/) * ResolvedView.ViewSizeAndInvSize.xy;// non-integer velocity;
+	
+	ENCODED_VELOCITY_TYPE EncodedMergedVelocity;
+	if (bHasVelocity && PointInView(OffsetTexelPosition) && any(OffsetTexelPosition != float2(TexelPosition)))
+	{
+		float3 DecodedTargetVelocity = DecodeVelocityTexture(VelocityTexture, OffsetTexelPosition);
+		float2 TargetTexelPosition = OffsetTexelPosition - (ScreenPosToViewportUV(DecodedTargetVelocity.xy) - 0.5f)*ResolvedView.ViewSizeAndInvSize.xy;
+	
+		// Construct velocity offset from Position to Target Texel Position.
+		if (PointInView(TargetTexelPosition))
+		{
+			float SourceDeviceZ = DepthTexture.Load(int3(TexelPosition, 0)).r;
+			float OffsetDeviceZ = CalcDeviceZ(OffsetTexelPosition);
+			float3 Velocity = float3(DecodedVectorOffset.xy, SourceDeviceZ - OffsetDeviceZ) + DecodedTargetVelocity;
+
+			//TODO: check if the TargetTexelPosition is within bound for rigid body transform to reject history that is too far/close.
+			//  |QtQt+1| \in (|PtPt+1| +/- (s+1)/s * |Pt+1Qt+1|, where s is the scaling.
+			EncodedMergedVelocity = EncodeVelocityToTexture(Velocity,TemporalResponsivenessMask);
+		}
+		else
+		{
+			EncodedMergedVelocity = VelocityTexture[TexelPosition];
+		}
+	}
+	else
+	{
+		EncodedMergedVelocity = VelocityTexture[TexelPosition];
+	}
+
+	RWMotionVectorWorldOffset[TexelPosition] = EncodedMergedVelocity;
+}
+
diff --git a/Engine/Source/Runtime/Engine/Private/Materials/HLSLMaterialTranslator.cpp b/Engine/Source/Runtime/Engine/Private/Materials/HLSLMaterialTranslator.cpp
--- a/Engine/Source/Runtime/Engine/Private/Materials/HLSLMaterialTranslator.cpp
+++ b/Engine/Source/Runtime/Engine/Private/Materials/HLSLMaterialTranslator.cpp
@@ -2049,6 +2049,13 @@ void FHLSLMaterialTranslator::TranslateMaterial()
 	MaterialCompilationOutput.bUsesDisplacement = bUsesDisplacement;
 	MaterialCompilationOutput.bUsesCustomizedUVs = Material->GetNumCustomizedUVs() > 0;
 
+	static const FGuid TemporalResponsivenessGuid = FMaterialAttributeDefinitionMap::GetCustomAttributeID(TEXT("TemporalResponsiveness"));
+	static const FGuid MotionVectorWorldOffsetGuid = FMaterialAttributeDefinitionMap::GetCustomAttributeID(TEXT("MotionVectorWorldOffset"));
+	const bool bUsesTemporalResponsivenessOutput =  (CustomOutputsBitmask & FMaterialAttributeDefinitionMap::GetBitmask(TemporalResponsivenessGuid)) != 0;
+	const bool bUsesMotionVectorWorldOffsetOutput = (CustomOutputsBitmask & FMaterialAttributeDefinitionMap::GetBitmask(MotionVectorWorldOffsetGuid)) != 0;
+	MaterialCompilationOutput.bUsesTemporalResponsiveness = bUsesTemporalResponsivenessOutput;
+	MaterialCompilationOutput.bUsesMotionVectorWorldOffset = bUsesMotionVectorWorldOffsetOutput;
+
 	// Fully rough if we have a roughness code chunk and it's constant and evaluates to 1.
 	bIsFullyRough = Chunk[MP_Roughness] != INDEX_NONE && IsMaterialPropertyUsed(MP_Roughness, Chunk[MP_Roughness], FLinearColor(1, 0, 0, 0), 1) == false;
 
diff --git a/Engine/Source/Runtime/Engine/Private/Materials/MaterialAttributeDefinitionMap.cpp b/Engine/Source/Runtime/Engine/Private/Materials/MaterialAttributeDefinitionMap.cpp
--- a/Engine/Source/Runtime/Engine/Private/Materials/MaterialAttributeDefinitionMap.cpp
+++ b/Engine/Source/Runtime/Engine/Private/Materials/MaterialAttributeDefinitionMap.cpp
@@ -16,6 +16,8 @@
 #include "Materials/MaterialExpressionThinTranslucentMaterialOutput.h"
 #include "Materials/MaterialExpressionFirstPersonOutput.h"
 #include "Materials/MaterialExpressionMaterialCache.h"
+#include "Materials/MaterialExpressionTemporalResponsivenessOutput.h"
+#include "Materials/MaterialExpressionMotionVectorWorldOffsetOutput.h"
 #include "Serialization/ShaderKeyGenerator.h"
 #include "SubstrateDefinitions.h"
 #include "RenderUtils.h"
@@ -385,6 +387,8 @@ void FMaterialAttributeDefinitionMap::InitializeAttributeMap()
 	AddCustomAttribute(FGuid(0xAAA1D0E0, 0xE312458E, 0xF3556BAD, 0xA01277CD), UMaterialExpressionThinTranslucentMaterialOutput::StaticClass(), "ThinTranslucentSurfaceCoverage", "GetThinTranslucentMaterialOutput", 1, MCT_Float, FVector4(1, 0, 0, 0), SF_Pixel);
 	AddCustomAttribute(FGuid(0xDB380E51, 0xBDED47B6, 0xB185A8BE, 0x8EACBFA4), UMaterialExpressionFirstPersonOutput::StaticClass(), "FirstPersonInterpolationAlpha", "GetFirstPersonOutput", 0, MCT_Float, FVector4(1, 0, 0, 0), SF_Vertex, true /*bNeedsPreviousFrameEvaluation*/);
 	AddCustomAttribute(FGuid(0x6B1237CA, 0x1B8C43C2, 0xA9C8A020, 0x430F1258), UMaterialExpressionMaterialCache::StaticClass(), "MaterialCache", "MaterialCache", 0, MCT_Float4, FVector4(0, 0, 0, 0), SF_Pixel);
+	AddCustomAttribute(FGuid(0xB58CD3FB, 0x8CB4411D, 0x83C4C11E, 0xF89953AD), UMaterialExpressionTemporalResponsivenessOutput::StaticClass(), "TemporalResponsiveness", "GetTemporalResponsivenessOutput", 0, MCT_Float, FVector4(0, 0, 0, 0), SF_Pixel);
+	AddCustomAttribute(FGuid(0x48A2CEAB, 0x8E63473F, 0x821345BC, 0x6FA907F7), UMaterialExpressionMotionVectorWorldOffsetOutput::StaticClass(), "MotionVectorWorldOffset", "GetMotionVectorWorldOffsetOutput", 0, MCT_Float3, FVector4(0, 0, 0, 0), SF_Pixel);
 }
 
 void FMaterialAttributeDefinitionMap::Add(const FGuid& AttributeID, const FString& AttributeName, EMaterialProperty Property,
diff --git a/Engine/Source/Runtime/Engine/Private/Materials/MaterialCachedData.cpp b/Engine/Source/Runtime/Engine/Private/Materials/MaterialCachedData.cpp
--- a/Engine/Source/Runtime/Engine/Private/Materials/MaterialCachedData.cpp
+++ b/Engine/Source/Runtime/Engine/Private/Materials/MaterialCachedData.cpp
@@ -16,6 +16,8 @@
 #include "Materials/MaterialExpressionSceneColor.h"
 #include "Materials/MaterialExpressionRuntimeVirtualTextureOutput.h"
 #include "Materials/MaterialExpressionFirstPersonOutput.h"
+#include "Materials/MaterialExpressionTemporalResponsivenessOutput.h"
+#include "Materials/MaterialExpressionMotionVectorWorldOffsetOutput.h"
 #include "Materials/MaterialExpressionLandscapeGrassOutput.h"
 #include "Materials/MaterialExpressionSetMaterialAttributes.h"
 #include "Materials/MaterialExpressionMakeMaterialAttributes.h"
@@ -65,6 +67,8 @@ FMaterialCachedExpressionData::FMaterialCachedExpressionData()
 	, bHasMaterialLayers(false)
 	, bHasRuntimeVirtualTextureOutput(false)
 	, bHasFirstPersonOutput(false)
+	, bUsesTemporalResponsiveness(false)
+	, bUsesMotionVectorWorldOffset(false)
 	, bSamplesMaterialCache(false)
 	, bHasMaterialCacheOutput(false)
 	, bMaterialCacheHasNonUVDerivedExpression(false)
@@ -467,7 +471,9 @@ void FMaterialCachedExpressionData::UpdateForExpressions(const FMaterialCachedEx
 {
 	check(EditorOnlyData);
 	static const FGuid FirstPersonInterpolationAlphaGuid = FMaterialAttributeDefinitionMap::GetCustomAttributeID(TEXT("FirstPersonInterpolationAlpha"));
-
+	static const FGuid TemporalResponsivenessGuid = FMaterialAttributeDefinitionMap::GetCustomAttributeID(TEXT("TemporalResponsiveness"));
+	static const FGuid MotionVectorWorldOffsetGuid = FMaterialAttributeDefinitionMap::GetCustomAttributeID(TEXT("MotionVectorWorldOffset"));
+	
 	for (UMaterialExpression* Expression : Expressions)
 	{
 		if (!Expression)
@@ -618,6 +624,20 @@ void FMaterialCachedExpressionData::UpdateForExpressions(const FMaterialCachedEx
 		{
 			bHasFirstPersonOutput = true;
 		}
+		else if (UMaterialExpressionTemporalResponsivenessOutput* TemporalResponsivenessOutputNode = Cast<UMaterialExpressionTemporalResponsivenessOutput>(Expression))
+		{
+			if (TemporalResponsivenessOutputNode->Input.IsConnected())
+			{
+				bUsesTemporalResponsiveness = true;
+			}
+
+		}else if (UMaterialExpressionMotionVectorWorldOffsetOutput* MotionVectorWorldOffsetOutputNode = Cast<UMaterialExpressionMotionVectorWorldOffsetOutput>(Expression))
+		{
+			if (MotionVectorWorldOffsetOutputNode->Input.IsConnected())
+			{
+				bUsesMotionVectorWorldOffset = true;
+			}
+		}
 		else if (UMaterialExpressionMaterialCache* ExpressionCache = Cast<UMaterialExpressionMaterialCache>(Expression))
 		{
 			MaterialCacheTags.Add(ExpressionCache->Tag);
@@ -708,6 +728,14 @@ void FMaterialCachedExpressionData::UpdateForExpressions(const FMaterialCachedEx
 					{
 						bHasFirstPersonOutput = true;
 					}
+					else if (Guid == TemporalResponsivenessGuid)
+					{
+						bUsesTemporalResponsiveness = true;
+					}
+					else if (Guid == MotionVectorWorldOffsetGuid)
+					{
+						bUsesMotionVectorWorldOffset = true;
+					}
 				}
 			}
 		}
diff --git a/Engine/Source/Runtime/Engine/Private/Materials/MaterialExpressions.cpp b/Engine/Source/Runtime/Engine/Private/Materials/MaterialExpressions.cpp
--- a/Engine/Source/Runtime/Engine/Private/Materials/MaterialExpressions.cpp
+++ b/Engine/Source/Runtime/Engine/Private/Materials/MaterialExpressions.cpp
@@ -310,6 +310,8 @@
 #include "Materials/MaterialExpressionRecordTextureStreamingInfo.h"
 #include "Materials/MaterialExpressionOperator.h"
 #include "Materials/MaterialExpressionAggregate.h"
+#include "Materials/MaterialExpressionTemporalResponsivenessOutput.h"
+#include "Materials/MaterialExpressionMotionVectorWorldOffsetOutput.h"
 #include "EditorSupportDelegates.h"
 #if WITH_EDITOR
 #include "MaterialGraph/MaterialGraphNode_Comment.h"
@@ -22961,4 +22963,122 @@ EMaterialValueType UMaterialExpressionBlend::GetOutputValueType(int32 OutputInde
 
 #endif // WITH_EDITOR
 
+///////////////////////////////////////////////////////////////////////////////
+// UMaterialExpressionTemporalResponsivenessOutput
+///////////////////////////////////////////////////////////////////////////////
+
+UMaterialExpressionTemporalResponsivenessOutput::UMaterialExpressionTemporalResponsivenessOutput(const FObjectInitializer& ObjectInitializer)
+	: Super(ObjectInitializer)
+{
+#if WITH_EDITOR
+	Outputs.Reset();
+#endif
+}
+
+#if WITH_EDITOR
+
+int32 UMaterialExpressionTemporalResponsivenessOutput::Compile(class FMaterialCompiler* Compiler, int32 OutputIndex)
+{
+	int32 CodeInput = INDEX_NONE;
+
+	if (OutputIndex == 0)
+	{
+		CodeInput = Input.IsConnected() ? Input.Compile(Compiler) : INDEX_NONE;
+	}
+
+	return Compiler->CustomOutput(this, OutputIndex, CodeInput);
+}
+
+void UMaterialExpressionTemporalResponsivenessOutput::GetCaption(TArray<FString>& OutCaptions) const
+{
+	OutCaptions.Add(FString(TEXT("Temporal Responsiveness")));
+}
+
+EMaterialValueType UMaterialExpressionTemporalResponsivenessOutput::GetInputValueType(int32 InputIndex)
+{
+	if (InputIndex == 0)
+	{
+		return MCT_Float1;
+	}
+
+	check(false);
+	return MCT_Float1;
+}
+
+#endif // WITH_EDITOR
+
+int32 UMaterialExpressionTemporalResponsivenessOutput::GetNumOutputs() const
+{
+	return 1;
+}
+
+FString UMaterialExpressionTemporalResponsivenessOutput::GetFunctionName() const
+{
+	return TEXT("GetTemporalResponsivenessOutput");
+}
+
+FString UMaterialExpressionTemporalResponsivenessOutput::GetDisplayName() const
+{
+	return TEXT("TemporalResponsiveness");
+}
+
+///////////////////////////////////////////////////////////////////////////////
+// UMaterialExpressionMotionVectorWorldOffsetOutput
+///////////////////////////////////////////////////////////////////////////////
+
+UMaterialExpressionMotionVectorWorldOffsetOutput::UMaterialExpressionMotionVectorWorldOffsetOutput(const FObjectInitializer& ObjectInitializer)
+	: Super(ObjectInitializer)
+{
+#if WITH_EDITOR
+	Outputs.Reset();
+#endif
+}
+
+#if WITH_EDITOR
+
+int32 UMaterialExpressionMotionVectorWorldOffsetOutput::Compile(class FMaterialCompiler* Compiler, int32 OutputIndex)
+{
+	int32 CodeInput = INDEX_NONE;
+
+	if (OutputIndex == 0)
+	{
+		CodeInput = Input.IsConnected() ? Input.Compile(Compiler) : INDEX_NONE;
+	}
+
+	return Compiler->CustomOutput(this, OutputIndex, CodeInput);
+}
+
+void UMaterialExpressionMotionVectorWorldOffsetOutput::GetCaption(TArray<FString>& OutCaptions) const
+{
+	OutCaptions.Add(FString(TEXT("Motion Vector World Offset (Per-Pixel)")));
+}
+
+EMaterialValueType UMaterialExpressionMotionVectorWorldOffsetOutput::GetInputValueType(int32 InputIndex)
+{
+	if (InputIndex == 0)
+	{
+		return MCT_Float3;
+	}
+
+	check(false);
+	return MCT_Float1;
+}
+
+#endif // WITH_EDITOR
+
+int32 UMaterialExpressionMotionVectorWorldOffsetOutput::GetNumOutputs() const
+{
+	return 1;
+}
+
+FString UMaterialExpressionMotionVectorWorldOffsetOutput::GetFunctionName() const
+{
+	return TEXT("GetMotionVectorWorldOffsetOutput");
+}
+
+FString UMaterialExpressionMotionVectorWorldOffsetOutput::GetDisplayName() const
+{
+	return TEXT("MotionVectorWorldOffset");
+}
+
 #undef LOCTEXT_NAMESPACE
diff --git a/Engine/Source/Runtime/Engine/Private/Materials/MaterialInterface.cpp b/Engine/Source/Runtime/Engine/Private/Materials/MaterialInterface.cpp
--- a/Engine/Source/Runtime/Engine/Private/Materials/MaterialInterface.cpp
+++ b/Engine/Source/Runtime/Engine/Private/Materials/MaterialInterface.cpp
@@ -456,6 +456,41 @@ bool UMaterialInterface::IsUsingWorldPositionOffset_Concurrent(EShaderPlatform I
 	return false;
 }
 
+bool UMaterialInterface::IsUsingMotionVectorWorldOffset_Concurrent(EShaderPlatform InShaderPlatform) const
+{
+	if (const UMaterial* Material = GetMaterial_Concurrent())
+	{
+		const FMaterialResource* MaterialResource = GetMaterialResource(InShaderPlatform);
+		if (!MaterialResource || !MaterialResource->GetMaterial())
+		{
+			MaterialResource = Material->GetMaterialResource(InShaderPlatform);
+		}
+		if (MaterialResource)
+		{
+			return MaterialResource->MaterialUsesMotionVectorWorldOffset_GameThread();
+		}
+	}
+	return false;
+}
+
+bool UMaterialInterface::IsUsingTemporalResponsiveness_Concurrent(EShaderPlatform InShaderPlatform) const
+{
+	if (const UMaterial* Material = GetMaterial_Concurrent())
+	{
+		const FMaterialResource* MaterialResource = GetMaterialResource(InShaderPlatform);
+		if (!MaterialResource || !MaterialResource->GetMaterial())
+		{
+			MaterialResource = Material->GetMaterialResource(InShaderPlatform);
+		}
+		if (MaterialResource)
+		{
+			return MaterialResource->MaterialUsesTemporalResponsiveness_GameThread();
+		}
+	}
+	return false;
+}
+
+
 // Deprecated 5.7
 FMaterialRelevance UMaterialInterface::GetRelevance_Internal(const UMaterial* Material, ERHIFeatureLevel::Type InFeatureLevel) const
 {
diff --git a/Engine/Source/Runtime/Engine/Private/Materials/MaterialShared.cpp b/Engine/Source/Runtime/Engine/Private/Materials/MaterialShared.cpp
--- a/Engine/Source/Runtime/Engine/Private/Materials/MaterialShared.cpp
+++ b/Engine/Source/Runtime/Engine/Private/Materials/MaterialShared.cpp
@@ -1099,6 +1099,28 @@ bool FMaterial::MaterialUsesPixelDepthOffset_RenderThread() const
 	return RenderingThreadShaderMap ? RenderingThreadShaderMap->UsesPixelDepthOffset() : false;
 }
 
+bool FMaterial::MaterialUsesTemporalResponsiveness_GameThread() const
+{
+	return GameThreadShaderMap.GetReference() ? GameThreadShaderMap->UsesTemporalResponsiveness() : false;
+}
+
+bool FMaterial::MaterialUsesTemporalResponsiveness_RenderThread() const
+{
+	check(IsInParallelRenderingThread());
+	return RenderingThreadShaderMap ? RenderingThreadShaderMap->UsesTemporalResponsiveness() : false;
+}
+
+bool FMaterial::MaterialUsesMotionVectorWorldOffset_GameThread() const
+{
+	return GameThreadShaderMap.GetReference() ? GameThreadShaderMap->UsesMotionVectorWorldOffset() : false;
+}
+
+bool FMaterial::MaterialUsesMotionVectorWorldOffset_RenderThread() const
+{
+	check(IsInParallelRenderingThread());
+	return RenderingThreadShaderMap ? RenderingThreadShaderMap->UsesMotionVectorWorldOffset() : false;
+}
+
 bool FMaterial::MaterialUsesDistanceCullFade_GameThread() const
 {
 	return GameThreadShaderMap ? GameThreadShaderMap->UsesDistanceCullFade() : false;
@@ -1632,6 +1654,8 @@ bool FMaterialResource::ShouldEnableResponsiveAA() const { return Material->bEna
 bool FMaterialResource::ShouldDoSSR() const { return Material->bScreenSpaceReflections; }
 bool FMaterialResource::ShouldDoContactShadows() const { return Material->bContactShadows; }
 bool FMaterialResource::HasPixelAnimation() const { return (MaterialInstance ? MaterialInstance->HasPixelAnimation() : Material->HasPixelAnimation()) && GetMaterialDomain() == MD_Surface && IsOpaqueOrMaskedBlendMode(GetBlendMode()); }
+bool FMaterialResource::UsesTemporalResponsiveness() const { return GetCachedExpressionData().bUsesTemporalResponsiveness;}
+bool FMaterialResource::UsesMotionVectorWorldOffset() const { return GetCachedExpressionData().bUsesMotionVectorWorldOffset;}
 bool FMaterialResource::IsWireframe() const { return Material->Wireframe; }
 bool FMaterialResource::IsUIMaterial() const { return Material->MaterialDomain == MD_UI; }
 bool FMaterialResource::IsPostProcessMaterial() const { return Material->MaterialDomain == MD_PostProcess; }
@@ -5808,6 +5832,8 @@ FMaterialShaderParametersBase::FMaterialShaderParametersBase(const FMaterial* In
 	bIsUsedWithStaticMesh = InMaterial->IsUsedWithStaticMesh();
 	bIsTranslucentSurfaceLighting = InMaterial->GetTranslucencyLightingMode() == ETranslucencyLightingMode::TLM_SurfacePerPixelLighting; // get lighting mode for mobile translucent receive CSM 
 	bHasFirstPersonInterpolation = InMaterial->HasFirstPersonOutput();
+	bUsesTemporalResponsiveness = InMaterial->UsesTemporalResponsiveness();
+	bUsesMotionVectorWorldOffset = InMaterial->UsesMotionVectorWorldOffset();
 }
 
 FMaterialShaderParameters::FMaterialShaderParameters(const FMaterial* InMaterial)
diff --git a/Engine/Source/Runtime/Engine/Private/PrimitiveSceneProxy.cpp b/Engine/Source/Runtime/Engine/Private/PrimitiveSceneProxy.cpp
--- a/Engine/Source/Runtime/Engine/Private/PrimitiveSceneProxy.cpp
+++ b/Engine/Source/Runtime/Engine/Private/PrimitiveSceneProxy.cpp
@@ -366,7 +366,9 @@ FPrimitiveSceneProxy::FPrimitiveSceneProxy(const FPrimitiveSceneProxyDesc& InPro
 ,	bHasWorldPositionOffsetVelocity(false)
 ,	bAnyMaterialHasWorldPositionOffset(false)
 ,	bAnyMaterialAlwaysEvaluatesWorldPositionOffset(false)
-,	bAnyMaterialHasPixelAnimation(false)	
+,	bAnyMaterialHasPixelAnimation(false)
+,   bAnyMaterialUsesTemporalResponsiveness(false)
+,   bAnyMaterialUsesMotionVectorWorldOffset(false)
 ,	bSupportsDistanceFieldRepresentation(false)
 ,	bSupportsHeightfieldRepresentation(false)
 ,	bLumenHeightfield(InProxyDesc.bLumenHeightfield)
@@ -541,6 +543,8 @@ FPrimitiveSceneProxy::FPrimitiveSceneProxy(const FPrimitiveSceneProxyDesc& InPro
 
 	bAnyMaterialHasWorldPositionOffset = false;
 	bAnyMaterialHasPixelAnimation = false;
+	bAnyMaterialUsesTemporalResponsiveness = false;
+	bAnyMaterialUsesMotionVectorWorldOffset = false;
 	{
 		// Find if we have any WPO materials.
 		EShaderPlatform ShaderPlatform = GetScene().GetShaderPlatform();
@@ -571,6 +575,24 @@ FPrimitiveSceneProxy::FPrimitiveSceneProxy(const FPrimitiveSceneProxyDesc& InPro
 					}
 				}
 			}
+
+			for (const UMaterialInterface* MaterialInterface : UsedMaterials)
+			{
+				if (MaterialInterface && MaterialInterface->IsUsingTemporalResponsiveness_Concurrent(ShaderPlatform))
+				{
+					bAnyMaterialUsesTemporalResponsiveness = true;
+					break;
+				}
+			}
+		}
+
+		for (const UMaterialInterface* MaterialInterface : UsedMaterials)
+		{
+			if (MaterialInterface && MaterialInterface->IsUsingMotionVectorWorldOffset_Concurrent(ShaderPlatform))
+			{
+				bAnyMaterialUsesMotionVectorWorldOffset = true;
+				break;
+			}
 		}
 	}
 
diff --git a/Engine/Source/Runtime/Engine/Private/ShaderCompiler/ShaderCompiler.cpp b/Engine/Source/Runtime/Engine/Private/ShaderCompiler/ShaderCompiler.cpp
--- a/Engine/Source/Runtime/Engine/Private/ShaderCompiler/ShaderCompiler.cpp
+++ b/Engine/Source/Runtime/Engine/Private/ShaderCompiler/ShaderCompiler.cpp
@@ -3793,13 +3793,21 @@ void GlobalBeginCompileShader(
 	}
 
 	{
+		static IConsoleVariable* CVarTemporalResponsivenessSupported = IConsoleManager::Get().FindConsoleVariable(TEXT("r.Velocity.TemporalResponsiveness.Supported"));
+		static IConsoleVariable* CVarPixelShaderMotionVectorWorldOffsetSupported = IConsoleManager::Get().FindConsoleVariable(TEXT("r.Velocity.PixelShaderMotionVectorWorldOffset.Supported"));
+
 		if (VelocityEncodeDepth((EShaderPlatform)Target.Platform))
 		{
 			SET_SHADER_DEFINE(Input.Environment, VELOCITY_ENCODE_DEPTH, 1);
+			SET_SHADER_DEFINE(Input.Environment, VELOCITY_SUPPORT_TEMPORAL_RESPONSIVENESS, CVarTemporalResponsivenessSupported ? CVarTemporalResponsivenessSupported->GetInt() != 0 : 0);
+			SET_SHADER_DEFINE(Input.Environment, VELOCITY_SUPPORT_PS_MOTION_VECTOR_WORLD_OFFSET, 
+				CVarPixelShaderMotionVectorWorldOffsetSupported ? CVarPixelShaderMotionVectorWorldOffsetSupported->GetInt() != 0 : 0);
 		}
 		else
 		{
 			SET_SHADER_DEFINE(Input.Environment, VELOCITY_ENCODE_DEPTH, 0);
+			SET_SHADER_DEFINE(Input.Environment, VELOCITY_SUPPORT_TEMPORAL_RESPONSIVENESS, 0);
+			SET_SHADER_DEFINE(Input.Environment, VELOCITY_SUPPORT_PS_MOTION_VECTOR_WORLD_OFFSET, 0);
 		}
 	}
 
diff --git a/Engine/Source/Runtime/Engine/Public/MaterialCachedData.h b/Engine/Source/Runtime/Engine/Public/MaterialCachedData.h
--- a/Engine/Source/Runtime/Engine/Public/MaterialCachedData.h
+++ b/Engine/Source/Runtime/Engine/Public/MaterialCachedData.h
@@ -384,6 +384,12 @@ struct FMaterialCachedExpressionData
 	UPROPERTY()
 	uint32 bHasFirstPersonOutput : 1;
 
+	UPROPERTY()
+	uint32 bUsesTemporalResponsiveness : 1;
+
+	UPROPERTY()
+	uint32 bUsesMotionVectorWorldOffset : 1;
+
 	UPROPERTY()
 	uint32 bSamplesMaterialCache : 1;
 
diff --git a/Engine/Source/Runtime/Engine/Public/MaterialShared.h b/Engine/Source/Runtime/Engine/Public/MaterialShared.h
--- a/Engine/Source/Runtime/Engine/Public/MaterialShared.h
+++ b/Engine/Source/Runtime/Engine/Public/MaterialShared.h
@@ -820,6 +820,8 @@ public:
 		bUsesEyeAdaptation(false),
 		bModifiesMeshPosition(false),
 		bUsesWorldPositionOffset(false),
+		bUsesTemporalResponsiveness(false),
+		bUsesMotionVectorWorldOffset(false),
 		bUsesGlobalDistanceField(false),
 		bUsesPixelDepthOffset(false),
 		bUsesDistanceCullFade(false),
@@ -963,6 +965,12 @@ public:
 	/** Whether the material uses world position offset. */
 	LAYOUT_BITFIELD(uint8, bUsesWorldPositionOffset, 1);
 
+	/** true if the material uses temporal responsiveness. */
+	LAYOUT_BITFIELD(uint8, bUsesTemporalResponsiveness, 1);
+
+	/** true if the material uses motion vector offset. */
+	LAYOUT_BITFIELD(uint8, bUsesMotionVectorWorldOffset, 1);
+
 	/** true if material uses the global distance field */
 	LAYOUT_BITFIELD(uint8, bUsesGlobalDistanceField, 1);
 
@@ -1060,6 +1068,8 @@ struct FMaterialShaderParametersBase
 			uint16 bIsTranslucentSurfaceLighting : 1;
 			uint16 bHasFirstPersonInterpolation : 1;
 			uint16 bIsSubstrateMaterial : 1;
+			uint16 bUsesTemporalResponsiveness : 1;
+			uint16 bUsesMotionVectorWorldOffset : 1;
 		};
 	};
 	EBlendMode BlendMode;
@@ -1698,6 +1708,8 @@ public:
 	bool NeedsSceneTextures() const { return GetContent()->MaterialCompilationOutput.bNeedsSceneTextures; }
 	bool UsesGlobalDistanceField() const { return GetContent()->MaterialCompilationOutput.bUsesGlobalDistanceField; }
 	bool UsesWorldPositionOffset() const { return GetContent()->MaterialCompilationOutput.bUsesWorldPositionOffset; }
+	bool UsesTemporalResponsiveness() const { return GetContent()->MaterialCompilationOutput.bUsesTemporalResponsiveness; }
+	bool UsesMotionVectorWorldOffset() const { return GetContent()->MaterialCompilationOutput.bUsesMotionVectorWorldOffset; }
 	bool UsesDisplacement() const { return GetContent()->MaterialCompilationOutput.bUsesDisplacement; }
 	bool NeedsGBuffer() const { return GetContent()->MaterialCompilationOutput.NeedsGBuffer(); }
 	bool UsesEyeAdaptation() const { return GetContent()->MaterialCompilationOutput.bUsesEyeAdaptation; }
@@ -2293,6 +2305,8 @@ public:
 	virtual bool ShouldDoSSR() const { return false; }
 	virtual bool ShouldDoContactShadows() const { return false; }
 	virtual	bool HasPixelAnimation() const { return false; }
+	virtual	bool UsesTemporalResponsiveness() const { return false; }
+	virtual	bool UsesMotionVectorWorldOffset() const { return false; }
 	virtual bool IsLightFunction() const = 0;
 	virtual bool IsUsedWithEditorCompositing() const { return false; }
 	virtual bool IsDeferredDecal() const = 0;
@@ -2573,6 +2587,14 @@ public:
 	ENGINE_API bool MaterialUsesPixelDepthOffset_RenderThread() const;
 	ENGINE_API bool MaterialUsesPixelDepthOffset_GameThread() const;
 
+	/** Does the material use temporal responsiveness. */
+	ENGINE_API bool MaterialUsesTemporalResponsiveness_RenderThread() const;
+	ENGINE_API bool MaterialUsesTemporalResponsiveness_GameThread() const;
+
+	/** Does the material use a motion vector offset. */
+	ENGINE_API bool MaterialUsesMotionVectorWorldOffset_RenderThread() const;
+	ENGINE_API bool MaterialUsesMotionVectorWorldOffset_GameThread() const;
+
 	/** Does the material modify the mesh position. */
 	ENGINE_API bool MaterialModifiesMeshPosition_RenderThread() const;
 	ENGINE_API bool MaterialModifiesMeshPosition_GameThread() const;
@@ -3096,6 +3118,8 @@ public:
 	ENGINE_API virtual bool ShouldDoSSR() const override;
 	ENGINE_API virtual bool ShouldDoContactShadows() const override;
 	ENGINE_API virtual bool HasPixelAnimation() const override;
+	ENGINE_API virtual bool UsesTemporalResponsiveness() const override;
+	ENGINE_API virtual bool UsesMotionVectorWorldOffset() const override;
 	ENGINE_API virtual bool IsLightFunction() const override;
 	ENGINE_API virtual bool IsUsedWithEditorCompositing() const override;
 	ENGINE_API virtual bool IsDeferredDecal() const override;
diff --git a/Engine/Source/Runtime/Engine/Public/Materials/MaterialExpressionMotionVectorWorldOffsetOutput.h b/Engine/Source/Runtime/Engine/Public/Materials/MaterialExpressionMotionVectorWorldOffsetOutput.h
new file mode 100644
--- /dev/null
+++ b/Engine/Source/Runtime/Engine/Public/Materials/MaterialExpressionMotionVectorWorldOffsetOutput.h
@@ -0,0 +1,39 @@
+// Copyright Epic Games, Inc. All Rights Reserved.
+
+#pragma once
+
+#include "CoreMinimal.h"
+#include "UObject/ObjectMacros.h"
+#include "Materials/MaterialExpressionCustomOutput.h"
+#include "MaterialExpressionMotionVectorWorldOffsetOutput.generated.h"
+
+/** Material output expression for generating a per-pixel motion vector offset that will be added to the existing motion vector. */
+UCLASS(collapsecategories, hidecategories=Object, Experimental, meta = (DisplayName = "Motion Vector World Offset (Per-Pixel)"), MinimalAPI)
+class UMaterialExpressionMotionVectorWorldOffsetOutput : public UMaterialExpressionCustomOutput
+{
+	GENERATED_UCLASS_BODY()
+
+	/**Motion vector offset in world space in the previous frame.
+	* 
+	*The input is added as a per-pixel offset to the motion vector generated by the velocity pass.
+	*
+	*Non-Nanite (Non-basepass velocity supported).
+	*/
+	UPROPERTY()
+	FExpressionInput Input;
+
+public:
+#if WITH_EDITOR
+	//~ Begin UMaterialExpression Interface
+	virtual int32 Compile(class FMaterialCompiler* Compiler, int32 OutputIndex) override;
+	virtual void GetCaption(TArray<FString>& OutCaptions) const override;
+	virtual EMaterialValueType GetInputValueType(int32 InputIndex) override;
+	//~ End UMaterialExpression Interface
+#endif
+
+	//~ Begin UMaterialExpressionCustomOutput Interface
+	virtual int32 GetNumOutputs() const override;
+	virtual FString GetFunctionName() const override;
+	virtual FString GetDisplayName() const override;
+	//~ End UMaterialExpressionCustomOutput Interface
+};
diff --git a/Engine/Source/Runtime/Engine/Public/Materials/MaterialExpressionTemporalResponsivenessOutput.h b/Engine/Source/Runtime/Engine/Public/Materials/MaterialExpressionTemporalResponsivenessOutput.h
new file mode 100644
--- /dev/null
+++ b/Engine/Source/Runtime/Engine/Public/Materials/MaterialExpressionTemporalResponsivenessOutput.h
@@ -0,0 +1,42 @@
+// Copyright Epic Games, Inc. All Rights Reserved.
+
+#pragma once
+
+#include "CoreMinimal.h"
+#include "UObject/ObjectMacros.h"
+#include "Materials/MaterialExpressionCustomOutput.h"
+#include "MaterialExpressionTemporalResponsivenessOutput.generated.h"
+
+UCLASS(collapsecategories, hidecategories=Object, Experimental, meta = (DisplayName = "Temporal Responsiveness"), MinimalAPI)
+class UMaterialExpressionTemporalResponsivenessOutput : public UMaterialExpressionCustomOutput
+{
+	GENERATED_UCLASS_BODY()
+
+	/**Indicate the temporal responsiveness requested.
+	*
+	*Default:0. Normal temporal accumulation.
+	*
+	*Medium:(0,0.5]. Assumes a medium-level mismatch of motion vectors for pixel animation, and reject temporal history more aggressively, especially for small features. 
+	*
+	*Full:(0.5,1]. Motion vector is fully unreliable. Reject all temporal history. 
+	*
+	*Non-Nanite (Full support), Nanite (Support base pass only)
+	*/
+	UPROPERTY()
+	FExpressionInput Input;
+
+public:
+#if WITH_EDITOR
+	//~ Begin UMaterialExpression Interface
+	virtual int32 Compile(class FMaterialCompiler* Compiler, int32 OutputIndex) override;
+	virtual void GetCaption(TArray<FString>& OutCaptions) const override;
+	virtual EMaterialValueType GetInputValueType(int32 InputIndex) override;
+	//~ End UMaterialExpression Interface
+#endif
+
+	//~ Begin UMaterialExpressionCustomOutput Interface
+	virtual int32 GetNumOutputs() const override;
+	virtual FString GetFunctionName() const override;
+	virtual FString GetDisplayName() const override;
+	//~ End UMaterialExpressionCustomOutput Interface
+};
diff --git a/Engine/Source/Runtime/Engine/Public/Materials/MaterialInterface.h b/Engine/Source/Runtime/Engine/Public/Materials/MaterialInterface.h
--- a/Engine/Source/Runtime/Engine/Public/Materials/MaterialInterface.h
+++ b/Engine/Source/Runtime/Engine/Public/Materials/MaterialInterface.h
@@ -784,6 +784,9 @@ public:
 	ENGINE_API bool IsUsingWorldPositionOffset_Concurrent(ERHIFeatureLevel::Type InFeatureLevel) const;
 	ENGINE_API bool IsUsingWorldPositionOffset_Concurrent(EShaderPlatform InShaderPlatform) const;
 
+	ENGINE_API bool IsUsingTemporalResponsiveness_Concurrent(EShaderPlatform InShaderPlatform) const;
+	ENGINE_API bool IsUsingMotionVectorWorldOffset_Concurrent(EShaderPlatform InShaderPlatform) const;
+
 #if !(UE_BUILD_SHIPPING || UE_BUILD_TEST)
 	/**
 	 * Output to the log which materials and textures are used by this material.
diff --git a/Engine/Source/Runtime/Engine/Public/PrimitiveSceneProxy.h b/Engine/Source/Runtime/Engine/Public/PrimitiveSceneProxy.h
--- a/Engine/Source/Runtime/Engine/Public/PrimitiveSceneProxy.h
+++ b/Engine/Source/Runtime/Engine/Public/PrimitiveSceneProxy.h
@@ -892,6 +892,8 @@ public:
 	inline bool AnyMaterialHasWorldPositionOffset() const { return bAnyMaterialHasWorldPositionOffset; }
 	inline bool AnyMaterialAlwaysEvaluatesWorldPositionOffset() const { return bAnyMaterialAlwaysEvaluatesWorldPositionOffset; }
 	inline bool AnyMaterialHasPixelAnimation() const { return bAnyMaterialHasPixelAnimation; }
+	inline bool AnyMaterialUsesTemporalResponsiveness() const { return bAnyMaterialUsesTemporalResponsiveness; }
+	inline bool AnyMaterialUsesMotionVectorWorldOffset() const { return bAnyMaterialUsesMotionVectorWorldOffset; }
 	inline float GetMaxWorldPositionOffsetExtent() const
 	{
 		if ((EvaluateWorldPositionOffset() && AnyMaterialHasWorldPositionOffset())
@@ -917,9 +919,9 @@ public:
 	/** Returns true if this proxy can change transform so that we should cache previous transform for calculating velocity. */
 	inline bool HasDynamicTransform() const { return IsMovable() || bIsBeingMovedByEditor; }
 	/** Returns true if this proxy can write velocity. This is used for setting velocity relevance. */
-	inline bool DrawsVelocity() const { return HasDynamicTransform() || bAlwaysHasVelocity || AnyMaterialHasPixelAnimation() || bHasWorldPositionOffsetVelocity; }
+	inline bool DrawsVelocity() const { return HasDynamicTransform() || bAlwaysHasVelocity || AnyMaterialHasPixelAnimation() || bHasWorldPositionOffsetVelocity || AnyMaterialUsesTemporalResponsiveness() || AnyMaterialUsesMotionVectorWorldOffset(); }
 	/** Returns true if this proxy should write velocity even when the transform isn't changing. Usually this is combined with a check for the transform changing. */
-	inline bool AlwaysHasVelocity() const {	return bAlwaysHasVelocity || AnyMaterialHasPixelAnimation() || (bHasWorldPositionOffsetVelocity && EvaluateWorldPositionOffset()); }
+	inline bool AlwaysHasVelocity() const {	return bAlwaysHasVelocity || AnyMaterialHasPixelAnimation() || (bHasWorldPositionOffsetVelocity && EvaluateWorldPositionOffset()) || AnyMaterialUsesTemporalResponsiveness() || AnyMaterialUsesMotionVectorWorldOffset(); }
 
 #if WITH_EDITOR
 	inline int32 GetNumUncachedStaticLightingInteractions() { return NumUncachedStaticLightingInteractions; }
@@ -1536,6 +1538,12 @@ protected:
 	/** Whether the primitive has any materials that must ALWAYS evaluate World Position Offset. */
 	uint8 bAnyMaterialHasPixelAnimation : 1;
 
+	/** Whether the primitive has any materials that evaluate temporal Responsiveness. */
+	uint8 bAnyMaterialUsesTemporalResponsiveness : 1;
+
+	/** Whether the primitive has any materials that evaluate motion vector offset to modify velocity. */
+	uint8 bAnyMaterialUsesMotionVectorWorldOffset : 1;
+
 	/** Whether the primitive should always be considered to have velocities, even if it hasn't moved. */
 	uint8 bAlwaysHasVelocity : 1;
 
diff --git a/Engine/Source/Runtime/Renderer/Private/SceneRendering.cpp b/Engine/Source/Runtime/Renderer/Private/SceneRendering.cpp
--- a/Engine/Source/Runtime/Renderer/Private/SceneRendering.cpp
+++ b/Engine/Source/Runtime/Renderer/Private/SceneRendering.cpp
@@ -979,6 +979,7 @@ void FViewInfo::Init()
 	bCustomDepthStencilValid = false;
 	bUsesCustomDepth = false;
 	bUsesCustomStencil = false;
+	bUsesMotionVectorWorldOffset = false;
 
 	// Sky dome, or any emissive, materials can result in high luminance values, e.g. the sun disk. 
 	// This Min here is to we make sure pre-exposed luminance remains within the boundaries of fp10 and not cause NaN on some platforms.
diff --git a/Engine/Source/Runtime/Renderer/Private/SceneRendering.h b/Engine/Source/Runtime/Renderer/Private/SceneRendering.h
--- a/Engine/Source/Runtime/Renderer/Private/SceneRendering.h
+++ b/Engine/Source/Runtime/Renderer/Private/SceneRendering.h
@@ -1371,6 +1371,9 @@ public:
 	uint32 bUsesCustomDepth : 1;
 	uint32 bUsesCustomStencil : 1;
 
+	/** Whether the view has any materials that uses motion vector offset to apply per pixel modification of velocities*/
+	uint32 bUsesMotionVectorWorldOffset : 1;
+
 	/** Whether fog should only be computed on rendered opaque pixels or not. */
 	uint32 bFogOnlyOnRenderedOpaque : 1;
 
diff --git a/Engine/Source/Runtime/Renderer/Private/SceneVisibility.cpp b/Engine/Source/Runtime/Renderer/Private/SceneVisibility.cpp
--- a/Engine/Source/Runtime/Renderer/Private/SceneVisibility.cpp
+++ b/Engine/Source/Runtime/Renderer/Private/SceneVisibility.cpp
@@ -1331,6 +1331,7 @@ void FRelevancePacket::Finalize()
 	WriteView.CustomDepthStencilValues.Append(CustomDepthStencilValues);
 	NaniteCustomDepthInstances.AppendTo(WriteView.NaniteCustomDepthInstances);
 	WriteView.bUsesCustomDepth |= bUsesCustomDepth;
+	WriteView.bUsesMotionVectorWorldOffset |= bUsesMotionVectorWorldOffset;
 	WriteView.bUsesCustomStencil |= bUsesCustomStencil;
 	WriteView.SubstrateViewData.MaxClosurePerPixel = FMath::Max(WriteView.SubstrateViewData.MaxClosurePerPixel, 8u - FMath::CountLeadingZeros8(SubstrateClosureCountMask));
 	WriteView.SubstrateViewData.MaxBytesPerPixel = FMath::Max(WriteView.SubstrateViewData.MaxBytesPerPixel, SubstrateUintPerPixel * 4u);
@@ -1633,6 +1634,11 @@ void FRelevancePacket::ComputeRelevance(FDynamicPrimitiveIndexList& DynamicPrimi
 										{
 											DrawCommandPacket.AddCommandsForMesh(PrimitiveIndex, PrimitiveSceneInfo, StaticMeshRelevance, StaticMesh, CullingPayloadFlags, Scene, bCanCache, EMeshPass::Velocity);
 											bIsMeshInVelocityPass = true;
+
+											if (PrimitiveSceneProxy->AnyMaterialUsesMotionVectorWorldOffset())
+											{
+												bUsesMotionVectorWorldOffset = true;
+											}
 										}
 
 										if (ViewRelevance.bOutputsTranslucentVelocity &&
@@ -2375,6 +2381,7 @@ static void ComputeDynamicMeshRelevance(
 	const FPrimitiveBounds& Bounds)
 {
 	const int32 NumElements = MeshBatch.Mesh->Elements.Num();
+	const bool bIsTranslucentClippedDepthEnabled = NeedTSRThinGeometryDetectionPass(View);
 
 	if (ViewRelevance.bDrawRelevance && (ViewRelevance.bRenderInMainPass || ViewRelevance.bRenderCustomDepth || ViewRelevance.bRenderInDepthPass))
 	{
@@ -2448,6 +2455,11 @@ static void ComputeDynamicMeshRelevance(
 			{
 				PassMask.Set(EMeshPass::Velocity);
 				View.NumVisibleDynamicMeshElements[EMeshPass::Velocity] += NumElements;
+
+				if (MeshBatch.PrimitiveSceneProxy->AnyMaterialUsesMotionVectorWorldOffset())
+				{
+					View.bUsesMotionVectorWorldOffset = true;
+				}
 			}
 
 			if (ViewRelevance.bOutputsTranslucentVelocity)
@@ -2455,7 +2467,7 @@ static void ComputeDynamicMeshRelevance(
 				PassMask.Set(EMeshPass::TranslucentVelocity);
 				View.NumVisibleDynamicMeshElements[EMeshPass::TranslucentVelocity] += NumElements;
 
-				if (ViewRelevance.bNormalTranslucency)
+				if (ViewRelevance.bNormalTranslucency && bIsTranslucentClippedDepthEnabled)
 				{
 					PassMask.Set(EMeshPass::TranslucentVelocityClippedDepth);
 					View.NumVisibleDynamicMeshElements[EMeshPass::TranslucentVelocityClippedDepth] += NumElements;
diff --git a/Engine/Source/Runtime/Renderer/Private/SceneVisibilityPrivate.h b/Engine/Source/Runtime/Renderer/Private/SceneVisibilityPrivate.h
--- a/Engine/Source/Runtime/Renderer/Private/SceneVisibilityPrivate.h
+++ b/Engine/Source/Runtime/Renderer/Private/SceneVisibilityPrivate.h
@@ -785,6 +785,7 @@ private:
 	bool bTranslucentSurfaceLighting = false;
 	bool bUsesCustomDepth = false;
 	bool bUsesCustomStencil = false;
+	bool bUsesMotionVectorWorldOffset = false;
 	bool bSceneHasSkyMaterial = false;
 	bool bHasSingleLayerWaterMaterial = false;
 	bool bUsesSecondStageDepthPass = false;
diff --git a/Engine/Source/Runtime/Renderer/Private/VelocityRendering.cpp b/Engine/Source/Runtime/Renderer/Private/VelocityRendering.cpp
--- a/Engine/Source/Runtime/Renderer/Private/VelocityRendering.cpp
+++ b/Engine/Source/Runtime/Renderer/Private/VelocityRendering.cpp
@@ -69,6 +69,18 @@ static TAutoConsoleVariable<int32> CVarVelocityOutputTranslucentClippedDepthEnab
 	TEXT("1: Provide functions e.g., mark before DoF translucency pixel animation for opacity below clip value.\n"),
 	ECVF_RenderThreadSafe);
 
+static TAutoConsoleVariable<int32> CVarVelocityTemporalResponsivenessSupported(
+	TEXT("r.Velocity.TemporalResponsiveness.Supported"),
+	0,
+	TEXT("Whether temporal Responsiveness is supported. use one more bit from the velocity texture.\n"),
+	ECVF_ReadOnly);
+
+static TAutoConsoleVariable<int32> CVarVelocityPixelShaderMotionVectorWorldOffsetSupported(
+	TEXT("r.Velocity.PixelShaderMotionVectorWorldOffset.Supported"),
+	0,
+	TEXT("Whether motion vector offset is supported in PS pass. Allow user to modify the motion vector per pixel.\n"),
+	ECVF_ReadOnly);
+
 DECLARE_GPU_DRAWCALL_STAT_NAMED(RenderVelocities, TEXT("Render Velocities"));
 
 /** Validate that deprecated CVars are no longer set. */
@@ -97,6 +109,25 @@ inline void ValidateVelocityCVars()
 #endif
 }
 
+bool NeedVelocityDepth(EShaderPlatform ShaderPlatform)
+{
+	// Lumen needs velocity depth
+	return (DoesProjectSupportDistanceFields() && FDataDrivenShaderPlatformInfo::GetSupportsLumenGI(ShaderPlatform)) 
+		|| FDataDrivenShaderPlatformInfo::GetSupportsRayTracing(ShaderPlatform);
+}
+
+bool SupportsTemporalResponsiveness(EShaderPlatform ShaderPlatform)
+{
+	static FShaderPlatformCachedIniValue<int32> PerPlatformCVar(TEXT("r.Velocity.TemporalResponsiveness.Supported"));
+	return NeedVelocityDepth(ShaderPlatform) && (PerPlatformCVar.Get(ShaderPlatform) != 0);
+}
+
+bool SupportsPixelShaderMotionVectorWorldOffset(EShaderPlatform ShaderPlatform)
+{
+	static FShaderPlatformCachedIniValue<int32> PerPlatformCVar(TEXT("r.Velocity.PixelShaderMotionVectorWorldOffset.Supported"));
+	return NeedVelocityDepth(ShaderPlatform) && (PerPlatformCVar.Get(ShaderPlatform) != 0);
+}
+
 class FVelocityVS : public FMeshMaterialShader
 {
 public:
@@ -116,6 +147,11 @@ public:
 		// Compile for materials which modify meshes.
 		const bool bMayModifyMeshes = Parameters.MaterialParameters.bMaterialMayModifyMeshPosition;
 
+		// Compile for materials that modifies motion vector offset or uses temporal responsiveness to indicate motion vector mismatch.
+		const bool bModifiesMotionVectorStatus = 
+			(Parameters.MaterialParameters.bUsesMotionVectorWorldOffset && SupportsPixelShaderMotionVectorWorldOffset(Parameters.Platform)) || 
+			(Parameters.MaterialParameters.bUsesTemporalResponsiveness && SupportsTemporalResponsiveness(Parameters.Platform));
+			
 		const bool bHasPlatformSupport = PlatformSupportsVelocityRendering(Parameters.Platform);
 
 		/**
@@ -125,7 +161,7 @@ public:
 		 */
 		const bool bIsSeparateVelocityPassRequired = 
 			!FVelocityRendering::BasePassCanOutputVelocity(Parameters.Platform) &&
-			(bIsMasked || bIsOpaqueAndTwoSided || bMayModifyMeshes);
+			(bIsMasked || bIsOpaqueAndTwoSided || bMayModifyMeshes || bModifiesMotionVectorStatus);
 
 		// The material may explicitly request that it be rendered into the translucent velocity pass.
 		const bool bIsSeparateVelocityPassRequiredByMaterial = Parameters.MaterialParameters.bIsTranslucencyWritingVelocity;
@@ -238,6 +274,45 @@ IMPLEMENT_MATERIAL_SHADER_TYPE(template<>, TVelocityPS<EVelocityPassMode::Veloci
 IMPLEMENT_SHADERPIPELINE_TYPE_VSPS(StandardVelocityPipeline, FVelocityVS, TVelocityPS<EVelocityPassMode::Velocity_Standard>, true);
 IMPLEMENT_SHADERPIPELINE_TYPE_VSPS(VelocityClippedDepthPipeline, FVelocityVS, TVelocityPS<EVelocityPassMode::Velocity_ClippedDepth>, true);
 
+class FMotionVectorWorldOffsetVelocityResolveCS : public FGlobalShader
+{
+public:
+	DECLARE_GLOBAL_SHADER(FMotionVectorWorldOffsetVelocityResolveCS);
+	SHADER_USE_PARAMETER_STRUCT(FMotionVectorWorldOffsetVelocityResolveCS, FGlobalShader);
+
+	BEGIN_SHADER_PARAMETER_STRUCT(FParameters,)
+		SHADER_PARAMETER_STRUCT_INCLUDE(FViewShaderParameters, View)
+		SHADER_PARAMETER_RDG_TEXTURE_SRV(Texture2D, DepthTexture)
+		SHADER_PARAMETER_RDG_TEXTURE_SRV(Texture2D, VelocityTexture)
+		SHADER_PARAMETER_RDG_TEXTURE_UAV(RWTexture2D, RWMotionVectorWorldOffset)
+	END_GLOBAL_SHADER_PARAMETER_STRUCT()
+
+	static int GetGroupSize()
+	{
+		return FComputeShaderUtils::kGolden2DGroupSize;
+	}
+
+	static EShaderPermutationPrecacheRequest ShouldPrecachePermutation(const FGlobalShaderPermutationParameters& Parameters)
+	{
+		//TODO: don't cache if per-pixel velocity offset is not used.
+		return EShaderPermutationPrecacheRequest::Precached;
+	}
+
+	static bool ShouldCompilePermutation(const FGlobalShaderPermutationParameters& Parameters)
+	{
+		return IsFeatureLevelSupported(Parameters.Platform, ERHIFeatureLevel::SM5);
+	}
+
+	static void ModifyCompilationEnvironment(const FGlobalShaderPermutationParameters& Parameters, FShaderCompilerEnvironment& OutEnvironment)
+	{
+		FGlobalShader::ModifyCompilationEnvironment(Parameters, OutEnvironment);
+		OutEnvironment.SetDefine(TEXT("THREADGROUP_SIZE"),GetGroupSize());
+		OutEnvironment.CompilerFlags.Add(CFLAG_AllowTypedUAVLoads);
+	}
+};
+
+IMPLEMENT_GLOBAL_SHADER(FMotionVectorWorldOffsetVelocityResolveCS, "/Engine/Private/VelocityUpdate.usf", "MainCS", SF_Compute);
+
 EMeshPass::Type GetMeshPassFromVelocityPass(EVelocityPass VelocityPass)
 {
 	switch (VelocityPass)
@@ -333,13 +408,6 @@ bool FMobileSceneRenderer::ShouldRenderVelocities() const
 IMPLEMENT_STATIC_UNIFORM_BUFFER_SLOT(VelocityClippedDepth);
 IMPLEMENT_STATIC_UNIFORM_BUFFER_STRUCT(FVelocityClippedDepthUniformParameters, "VelocityClippedDepth", VelocityClippedDepth);
 
-bool NeedVelocityDepth(EShaderPlatform ShaderPlatform)
-{
-	// Lumen needs velocity depth
-	return (DoesProjectSupportDistanceFields() && FDataDrivenShaderPlatformInfo::GetSupportsLumenGI(ShaderPlatform)) 
-		|| FDataDrivenShaderPlatformInfo::GetSupportsRayTracing(ShaderPlatform);
-}
-
 static FRDGTextureUAVRef CreateDummyVelocityUAV(FRDGBuilder& GraphBuilder, EPixelFormat PixelFormat)
 {
 	FRDGTextureDesc Desc = FRDGTextureDesc::Create2D(
@@ -381,6 +449,29 @@ BEGIN_SHADER_PARAMETER_STRUCT(FVelocityPassParameters, )
 	RENDER_TARGET_BINDING_SLOTS()
 END_SHADER_PARAMETER_STRUCT()
 
+void GetMotionVectorOutputFlag(
+	TArrayView<FViewInfo> InViews,
+	EMeshPass::Type MeshPass,
+	bool bForceVelocity,
+	bool& bUsesAnyMotionVectorWorldOffsetMaterials)
+{
+	bUsesAnyMotionVectorWorldOffsetMaterials = false;
+
+	for (int32 ViewIndex = 0; ViewIndex < InViews.Num(); ViewIndex++)
+	{
+		FViewInfo& View = InViews[ViewIndex];
+		if (View.ShouldRenderView())
+		{
+			const bool bHasAnyDraw = HasAnyDraw(View.ParallelMeshDrawCommandPasses[MeshPass]);
+			if (!bHasAnyDraw && !bForceVelocity)
+			{
+				continue;
+			}
+			bUsesAnyMotionVectorWorldOffsetMaterials |= View.bUsesMotionVectorWorldOffset;
+		}
+	}
+}
+
 void FSceneRenderer::RenderVelocities(
 	FRDGBuilder& GraphBuilder,
 	TArrayView<FViewInfo> InViews,
@@ -408,12 +499,26 @@ void FSceneRenderer::RenderVelocities(
 	RDG_GPU_STAT_SCOPE(GraphBuilder, RenderVelocities);
 
 	const EMeshPass::Type MeshPass = GetMeshPassFromVelocityPass(VelocityPass);
+	const bool bIsOpaquePass = VelocityPass == EVelocityPass::Opaque;
 
-	FExclusiveDepthStencil ExclusiveDepthStencil = (VelocityPass == EVelocityPass::Opaque && !(Scene->EarlyZPassMode == DDM_AllOpaqueNoVelocity))
+	FExclusiveDepthStencil ExclusiveDepthStencil = (bIsOpaquePass && !(Scene->EarlyZPassMode == DDM_AllOpaqueNoVelocity))
 														? FExclusiveDepthStencil::DepthRead_StencilWrite
 														: FExclusiveDepthStencil::DepthWrite_StencilWrite;
 	ExclusiveDepthStencil = bIsTranslucentClippedDepthPass ? FExclusiveDepthStencil::DepthRead_StencilNop : ExclusiveDepthStencil;
 
+	//Only call the per pixel velocity resolve when we have at least one material using it.
+	bool bHasAnyPixelShaderMotionVectorWorldOffsetMaterials = false;
+	GetMotionVectorOutputFlag(InViews, MeshPass, bForceVelocity, bHasAnyPixelShaderMotionVectorWorldOffsetMaterials);
+	const bool bSupportPixelShaderMotionVectorWorldOffset = SupportsPixelShaderMotionVectorWorldOffset(ShaderPlatform) && bIsOpaquePass && bHasAnyPixelShaderMotionVectorWorldOffsetMaterials;
+
+	//Only opaque pass supports per pixel override.
+	FRDGTextureRef MotionVectorWorldOffsetTexture = nullptr;
+	if (bSupportPixelShaderMotionVectorWorldOffset)
+	{
+		MotionVectorWorldOffsetTexture = GraphBuilder.CreateTexture(SceneTextures.Velocity->Desc,TEXT("MotionVectorWorldOffsetTexture"));
+		AddClearRenderTargetPass(GraphBuilder, MotionVectorWorldOffsetTexture);
+	}
+
 	for (int32 ViewIndex = 0; ViewIndex < InViews.Num(); ViewIndex++)
 	{
 		FViewInfo& View = InViews[ViewIndex];
@@ -461,7 +566,33 @@ void FSceneRenderer::RenderVelocities(
 
 			if (bBindRenderTarget)
 			{
-				PassParameters->RenderTargets[0] = FRenderTargetBinding(SceneTextures.Velocity, (bNeedsClearMask & View.GPUMask.GetNative()) ? ERenderTargetLoadAction::EClear : ERenderTargetLoadAction::ELoad);
+				ERenderTargetLoadAction LoadAction = (bNeedsClearMask & View.GPUMask.GetNative()) ? ERenderTargetLoadAction::EClear : ERenderTargetLoadAction::ELoad;
+				if (MotionVectorWorldOffsetTexture)
+				{
+					// Switch Velocity and Offset texture to avoid an additional copy
+					// 
+					// Write Velocity into the Offset texture and Offset into the Velocity texture so that
+					// When we resolve (e.g., RWOffset[Position]+= Velocity[ResolvedPosition]), the resolved velocity 
+					// is stored in the Velocity texture (RWOffset) instead of Offset texture to avoid an additional copy
+					// from Offset texture to Velocity texture.
+					// From
+					//      V			=	v
+					//      Offset		=	o
+					//		Offset[p]	+= V[rp]
+					//		V			= Offset
+					// To
+					//		Offset		= v
+					//		V			= o
+					//		V[p]		+= Offset[rp]
+
+					PassParameters->RenderTargets[0] = FRenderTargetBinding(MotionVectorWorldOffsetTexture,LoadAction);
+					PassParameters->RenderTargets[1] = FRenderTargetBinding(SceneTextures.Velocity, LoadAction);
+				}
+				else
+				{
+					PassParameters->RenderTargets[0] = FRenderTargetBinding(SceneTextures.Velocity, LoadAction);
+				}
+				
 				bNeedsClearMask &= ~View.GPUMask.GetNative();
 			}
 
@@ -495,6 +626,48 @@ void FSceneRenderer::RenderVelocities(
 		}
 	}
 
+	if (bSupportPixelShaderMotionVectorWorldOffset)
+	{
+		for (int32 ViewIndex = 0; ViewIndex < InViews.Num(); ViewIndex++)
+		{
+			FViewInfo& View = InViews[ViewIndex];
+
+			if (View.ShouldRenderView())
+			{
+				const bool bHasAnyDraw = HasAnyDraw(View.ParallelMeshDrawCommandPasses[MeshPass]);
+				if ((!bHasAnyDraw && !bForceVelocity) || !View.bUsesMotionVectorWorldOffset)
+				{
+					continue;
+				}
+
+				RDG_GPU_MASK_SCOPE(GraphBuilder, View.GPUMask);
+
+				// Resolve
+				{
+					typedef FMotionVectorWorldOffsetVelocityResolveCS SHADER;
+					SHADER::FParameters* PassParameters = GraphBuilder.AllocParameters<SHADER::FParameters>();
+					PassParameters->View = View.GetShaderParameters();
+					PassParameters->DepthTexture = GraphBuilder.CreateSRV(SceneTextures.Depth.Resolve);
+
+					// Switch back to avoid an additional copy.
+					PassParameters->VelocityTexture = GraphBuilder.CreateSRV(MotionVectorWorldOffsetTexture);
+					PassParameters->RWMotionVectorWorldOffset = GraphBuilder.CreateUAV(SceneTextures.Velocity);
+
+					TShaderMapRef<SHADER> ComputeShader(View.ShaderMap);
+					FIntVector GroupCount = FComputeShaderUtils::GetGroupCount(View.ViewRect.Size(), SHADER::GetGroupSize());
+					FComputeShaderUtils::AddPass(
+						GraphBuilder,
+						RDG_EVENT_NAME("MotionVectorWorldOffsetVelocityResolve %dx%d",
+							View.ViewRect.Width(),
+							View.ViewRect.Height()),
+						ComputeShader,
+						PassParameters,
+						GroupCount);
+				}
+			}
+		}
+	}
+
 #if !(UE_BUILD_SHIPPING)
 	const bool bForwardShadingEnabled = IsForwardShadingEnabled(ShaderPlatform);
 	if (!bForwardShadingEnabled)
@@ -690,8 +863,8 @@ bool FOpaqueVelocityMeshProcessor::TryAddMeshBatch(
 		const ERasterizerCullMode MeshCullMode = ComputeMeshCullMode(*Material, OverrideSettings);
 		const bool bVFTypeSupportsNullPixelShader = MeshBatch.VertexFactory->SupportsNullPixelShader();
 		const bool bModifiesMeshPosition = DoMaterialAndPrimitiveModifyMeshPosition(*Material, PrimitiveSceneProxy);
-
-		const bool bSwapWithDefaultMaterial = UseDefaultMaterial(Material, bVFTypeSupportsNullPixelShader, bModifiesMeshPosition);
+		const bool bModifiesMotionVectorStatus = Material->MaterialUsesMotionVectorWorldOffset_GameThread() || Material->MaterialUsesTemporalResponsiveness_GameThread();
+		const bool bSwapWithDefaultMaterial = UseDefaultMaterial(Material, bVFTypeSupportsNullPixelShader, bModifiesMeshPosition) && !bModifiesMotionVectorStatus;
 		if (bSwapWithDefaultMaterial)
 		{
 			MaterialRenderProxy = UMaterial::GetDefaultMaterial(MD_Surface)->GetRenderProxy();
@@ -774,7 +947,9 @@ void FOpaqueVelocityMeshProcessor::CollectPSOInitializers(const FSceneTexturesCo
 		if (PreCacheParams.bRenderInMainPass && bIsNotTranslucent && ShouldIncludeMaterialInDefaultOpaquePass(Material))
 		{
 			const bool bVFTypeSupportsNullPixelShader = VertexFactoryData.VertexFactoryType->SupportsNullPixelShader();
-			const bool bUseDefaultMaterial = UseDefaultMaterial(&Material, bVFTypeSupportsNullPixelShader, Material.MaterialModifiesMeshPosition_GameThread());
+			const bool bModifiesMotionVectorStatus = Material.MaterialUsesMotionVectorWorldOffset_GameThread() || Material.MaterialUsesTemporalResponsiveness_GameThread();
+			const bool bUseDefaultMaterial = UseDefaultMaterial(&Material, bVFTypeSupportsNullPixelShader, Material.MaterialModifiesMeshPosition_GameThread())
+				&& !bModifiesMotionVectorStatus;
 			if (!bUseDefaultMaterial)
 			{
 				bCollectPSOs = true;
