Enabled Reserved Resources for GPU Scene to reduce hitches on resize and peak GPU memory use by eliminating the need to double buffer during the copy.

diff --git a/Engine/Source/Runtime/Renderer/Private/GPUScene.cpp b/Engine/Source/Runtime/Renderer/Private/GPUScene.cpp
--- a/Engine/Source/Runtime/Renderer/Private/GPUScene.cpp
+++ b/Engine/Source/Runtime/Renderer/Private/GPUScene.cpp
@@ -139,7 +139,7 @@ static TAutoConsoleVariable<int32> CVarGPUSceneLightsAsyncSetup(
 
 static TAutoConsoleVariable<int32> CVarGPUSceneInstanceDataTileSizeLog2(
 	TEXT("r.GPUScene.InstanceDataTileSizeLog2"),
-	-1,
+	12,
 	TEXT("Log2 of tile size used for the instance data in number of instances. \n")
 	TEXT("  <0 (default) indicates no tiling and disables reserved resources. \n")
 	TEXT("  >0 will use the given tile-size and attempt to use reserved resources, but falls back to a regular buffer if not available."),
@@ -148,7 +148,7 @@ static TAutoConsoleVariable<int32> CVarGPUSceneInstanceDataTileSizeLog2(
 
 static TAutoConsoleVariable<bool> CVarGPUSceneUseReservedResources(
 	TEXT("r.GPUScene.UseReservedResources"),
-	false,
+	true,
 	TEXT("Turning this on makes all GPU-Scene buffers default to using reserved resource allocations."),
 	ECVF_ReadOnly);
 
