[RayTracing]
* Remove call to FlushAllocationsToClear on SBT before setting the current bindings on the SBT because the clear call could flush pending clear ops which are still used this frame. Only flush clear after add & remove of primitive scene infos on RT timeline. Pending clears can be added by async stream in requests in between the building of the dirty bindings and the flush to RHI from RDG passes causing cleared data in the SBT and triggering validation errors and potentially missing entries in the SBT.
* Also keep track of dirty records and clear dirty record when reallocated again before the clear was flushed (otherwise valid entries can be cleared)

diff --git a/Engine/Source/Runtime/Renderer/Private/DeferredShadingRenderer.cpp b/Engine/Source/Runtime/Renderer/Private/DeferredShadingRenderer.cpp
--- a/Engine/Source/Runtime/Renderer/Private/DeferredShadingRenderer.cpp
+++ b/Engine/Source/Runtime/Renderer/Private/DeferredShadingRenderer.cpp
@@ -1117,15 +1117,6 @@ void FDeferredShadingSceneRenderer::SetupRayTracingRenderingData(FRDGBuilder& Gr
 		DirtyPersistentRayTracingShaderBindings = Scene->RayTracingSBT.GetDirtyBindings(VisibleRayTracingShaderBindings, bForceAllDirty);
 	}
 
-	GraphBuilder.AddPass(RDG_EVENT_NAME("FlushRayTracingSBTAllocationsToClear"), ERDGPassFlags::NeverCull,
-		[Scene = Scene](FRDGAsyncTask, FRHICommandList& RHICmdList)
-		{
-			TRACE_CPUPROFILER_EVENT_SCOPE(FlushRayTracingSBTAllocationsToClear);
-
-			Scene->RayTracingSBT.FlushAllocationsToClear(RHICmdList);
-		}
-	);
-
 	FRDGBufferRef LumenHardwareRayTracingHitDataBuffer = nullptr;
 	if (bAnyInlineHardwareRayTracingPassEnabled)
 	{
diff --git a/Engine/Source/Runtime/Renderer/Private/RayTracing/RayTracingShaderBindingTable.cpp b/Engine/Source/Runtime/Renderer/Private/RayTracing/RayTracingShaderBindingTable.cpp
--- a/Engine/Source/Runtime/Renderer/Private/RayTracing/RayTracingShaderBindingTable.cpp
+++ b/Engine/Source/Runtime/Renderer/Private/RayTracing/RayTracingShaderBindingTable.cpp
@@ -188,11 +188,13 @@ FRayTracingShaderBindingDataOneFrameArray FRayTracingShaderBindingTable::GetDirt
 			}
 			else
 			{
-				FBitReference RecordIndex = ValidPersistentRecords[VisibleBinding.SBTRecordIndex];
-				if (!RecordIndex || VisibleBinding.BindingType == ERayTracingLocalShaderBindingType::Transient)
+				check(VisibleBinding.BindingType!= ERayTracingLocalShaderBindingType::Persistent || SetPersistentRecords[VisibleBinding.SBTRecordIndex]);
+
+				FBitReference DirtyRecordIndex = DirtyPersistentRecords[VisibleBinding.SBTRecordIndex];
+				if (DirtyRecordIndex || VisibleBinding.BindingType == ERayTracingLocalShaderBindingType::Transient)
 				{
 					DirtyBindings.Add(VisibleBinding);
-					RecordIndex = true;
+					DirtyRecordIndex = false;
 				}
 				else if (GValidatePersistentBindings)
 				{
@@ -253,7 +255,7 @@ FRayTracingPersistentShaderBindingTableID FRayTracingShaderBindingTable::Allocat
 	PersistentSBTs[ResultID].ShaderBindingMode = ShaderBindingMode;
 	AllocatePersistentShaderBindingTable(RHICmdList, ResultID, ShaderBindingMode);
 
-	ValidPersistentRecords.Init(false, PersistentSBTInitializer.NumGeometrySegments * PersistentSBTInitializer.NumShaderSlotsPerGeometrySegment);
+	DirtyPersistentRecords.Init(true, PersistentSBTInitializer.NumGeometrySegments * PersistentSBTInitializer.NumShaderSlotsPerGeometrySegment);
 
 	return ResultID;
 }
@@ -330,7 +332,7 @@ void FRayTracingShaderBindingTable::CheckPersistentRHI(FRHICommandListBase& RHIC
 				AllocatePersistentShaderBindingTable(RHICmdList, Index, PersistentSBTs[Index].ShaderBindingMode);
 			}
 		}
-		ValidPersistentRecords.Init(false, NewSBTInitializer.NumGeometrySegments * NewSBTInitializer.NumShaderSlotsPerGeometrySegment);
+		DirtyPersistentRecords.Init(true, NewSBTInitializer.NumGeometrySegments * NewSBTInitializer.NumShaderSlotsPerGeometrySegment);
 	}
 }
 
@@ -368,13 +370,35 @@ uint32 FRayTracingShaderBindingTable::GetMaxAllocatedStaticSegmentCount() const
 void FRayTracingShaderBindingTable::MarkDirty(FRayTracingSBTAllocation* SBTAllocation)
 {
 	int32 MaxRecordIndex = SBTAllocation->BaseRecordIndex + SBTAllocation->NumRecords;
-	if (ValidPersistentRecords.Num() < MaxRecordIndex)
+	if (DirtyPersistentRecords.Num() < MaxRecordIndex)
 	{
-		ValidPersistentRecords.SetNum(FMath::RoundUpToPowerOfTwo(MaxRecordIndex), false);
+		DirtyPersistentRecords.SetNum(FMath::RoundUpToPowerOfTwo(MaxRecordIndex), true);
 	}
 	for (uint32 Index = 0; Index < SBTAllocation->NumRecords; ++Index)
 	{
-		ValidPersistentRecords[SBTAllocation->BaseRecordIndex + Index] = false;
+		DirtyPersistentRecords[SBTAllocation->BaseRecordIndex + Index] = true;
+	}
+}
+
+void FRayTracingShaderBindingTable::MarkSet(FRayTracingSBTAllocation* SBTAllocation, bool bValue)
+{
+	int32 MaxRecordIndex = SBTAllocation->BaseRecordIndex + SBTAllocation->NumRecords;
+	if (SetPersistentRecords.Num() < MaxRecordIndex)
+	{
+		SetPersistentRecords.SetNum(FMath::RoundUpToPowerOfTwo(MaxRecordIndex), false);
+		PersistentRecordsToClear.SetNum(FMath::RoundUpToPowerOfTwo(MaxRecordIndex), false);
+	}
+	for (uint32 Index = 0; Index < SBTAllocation->NumRecords; ++Index)
+	{
+		uint32 RecordIndex = SBTAllocation->BaseRecordIndex + Index;
+		check(SetPersistentRecords[RecordIndex] != bValue);
+		SetPersistentRecords[RecordIndex] = bValue;
+
+		// Mark record to clear and validate that it's not marked for clear when we are trying to clear it out (double clear?)
+		check(bValue || !PersistentRecordsToClear[RecordIndex]);
+		PersistentRecordsToClear[RecordIndex] = !bValue;
+
+		//UE_LOG(LogRenderer, Log, TEXT("Update set SBT record data for RecordIndex %d with new value: %d"), RecordIndex, bValue ? 1 : 0);
 	}
 }
 
@@ -396,6 +420,7 @@ FRayTracingSBTAllocation* FRayTracingShaderBindingTable::AllocateStaticRangeInte
 	Allocation->InitStatic(AllocatedLayers, BaseIndex, RecordsPerLayer, RecordCount, Geometry, Flags);
 
 	MarkDirty(Allocation);
+	MarkSet(Allocation, true);
 
 	AllocatedStaticSegmentCount += SegmentCount * LayersCount;
 	
@@ -477,7 +502,8 @@ void FRayTracingShaderBindingTable::FreeStaticRange(const FRayTracingSBTAllocati
 			StaticRangeAllocator.Free(RefAllocation->Allocation->BaseRecordIndex, RefAllocation->Allocation->NumRecords);
 			AllocatedStaticSegmentCount -= (RefAllocation->Allocation->NumRecords / NumShaderSlotsPerGeometrySegment);
 		
-			PersistentAllocationsToClear.Add(RefAllocation->Allocation);
+			MarkSet(RefAllocation->Allocation, false);
+
 			TrackedAllocationMap.Remove(Key);
 		}
 	}
@@ -519,7 +545,7 @@ FRayTracingSBTAllocation* FRayTracingShaderBindingTable::AllocateDynamicRange(ER
 
 bool FRayTracingShaderBindingTable::IsDirty(uint32 RecordIndex) const
 {
-	return !ValidPersistentRecords[RecordIndex];
+	return DirtyPersistentRecords[RecordIndex];
 }
 
 bool FRayTracingShaderBindingTable::IsPersistent() const
@@ -534,34 +560,9 @@ void FRayTracingShaderBindingTable::FlushAllocationsToClear(FRHICommandList& RHI
 	// Don't clear outside of current allocated range - could have allocated and cleared ranges before calling CheckPersistentRHI
 	// and then the RHI SBT doesn't have those new ranges allocated yet
 	uint32 MaxNumValidRecords = PersistentSBTInitializer.NumGeometrySegments * PersistentSBTInitializer.NumShaderSlotsPerGeometrySegment;
-	check(MaxNumValidRecords <= (uint32)ValidPersistentRecords.Num());
-
-	// Build bit array if records to clear (make sure we don't clear the same record twice)
-	TBitArray<> PersistentRecordsToClear;
-	PersistentRecordsToClear.Init(false, MaxNumValidRecords);
+	check(MaxNumValidRecords <= (uint32)DirtyPersistentRecords.Num());
 	
-	uint32 TotalRecordsToClear = 0;
-	for (const FRayTracingSBTAllocation* Allocation : PersistentAllocationsToClear)
-	{
-		for (uint32 RangeIndex = 0; RangeIndex < Allocation->NumRecords; ++RangeIndex)
-		{
-			uint32 RecordIndex = Allocation->BaseRecordIndex + RangeIndex;
-			if (RecordIndex < MaxNumValidRecords)
-			{
-				FBitReference RecordFlag = PersistentRecordsToClear[RecordIndex];
-				if (!RecordFlag)
-				{
-					TotalRecordsToClear++;
-					RecordFlag = true;
-				}
-			}
-		}
-
-		// TODO: pool static allocations as well?
-		delete Allocation;
-	}
-	PersistentAllocationsToClear.Reset(PersistentAllocationsToClear.Num());
-
+	uint32 TotalRecordsToClear = PersistentRecordsToClear.CountSetBits();
 	if (TotalRecordsToClear == 0)
 	{
 		return;
@@ -595,6 +596,9 @@ void FRayTracingShaderBindingTable::FlushAllocationsToClear(FRHICommandList& RHI
 				bCopyDataToInlineStorage);
 		}
 	}
+
+	// Reset the records to clear
+	PersistentRecordsToClear.Init(false, PersistentRecordsToClear.Num());
 }
 
 #endif // RHI_RAYTRACING
\ No newline at end of file
diff --git a/Engine/Source/Runtime/Renderer/Private/RayTracing/RayTracingShaderBindingTable.h b/Engine/Source/Runtime/Renderer/Private/RayTracing/RayTracingShaderBindingTable.h
--- a/Engine/Source/Runtime/Renderer/Private/RayTracing/RayTracingShaderBindingTable.h
+++ b/Engine/Source/Runtime/Renderer/Private/RayTracing/RayTracingShaderBindingTable.h
@@ -232,6 +232,7 @@ private:
 	 * Mark all records used by the SBTAllocation as dirty again
 	 */
 	void MarkDirty(FRayTracingSBTAllocation* SBTAllocation);
+	void MarkSet(FRayTracingSBTAllocation* SBTAllocation, bool bValue);
 
 	struct FAllocationKey
 	{
@@ -279,8 +280,9 @@ private:
 
 	FRayTracingShaderBindingTableInitializer PersistentSBTInitializer;	 //< Shared initializer used for all persistent SBTs - so they can all be versioned together
 	TArray<FPersistentSBTData> PersistentSBTs;							 //< All currently allocated persistent SBTs (FRayTracingPersistentShaderBindingTableID contains index into this array so can be sparse)
-	TBitArray<> ValidPersistentRecords;									 //< BitArray containing which bits are valid in the cached persistent SBTs (used to build the dirty set)
-	TArray<FRayTracingSBTAllocation*> PersistentAllocationsToClear;		 //< Persistent SBT allocations about which the cached persistent SBTs need to be notified (clear internal cached record data)
+	TBitArray<> SetPersistentRecords;									 //< BitArray containing which bits are valid in the cached persistent SBTs (used to build the dirty set)
+	TBitArray<> DirtyPersistentRecords;									 //< BitArray containing which bits are valid in the cached persistent SBTs (used to build the dirty set)
+	TBitArray<> PersistentRecordsToClear;								 //< BitArray containing which bits need to be cleared in the persisten SBTs
 	
 	FCriticalSection StaticAllocationCS;								 //< Critical section used to access all static allocation data
 	bool bStaticAllocationsLocked = false;								 //< Static allocations are not allowed when this bool is set (used for validation)
