BugHawk fixes for leak, extra validation and documentation

diff --git a/Engine/Source/Runtime/Renderer/Private/RayTracing/RayTracingShaderBindingTable.cpp b/Engine/Source/Runtime/Renderer/Private/RayTracing/RayTracingShaderBindingTable.cpp
--- a/Engine/Source/Runtime/Renderer/Private/RayTracing/RayTracingShaderBindingTable.cpp
+++ b/Engine/Source/Runtime/Renderer/Private/RayTracing/RayTracingShaderBindingTable.cpp
@@ -504,6 +504,9 @@ void FRayTracingShaderBindingTable::FreeStaticRange(const FRayTracingSBTAllocati
 		
 			MarkSet(RefAllocation->Allocation, false);
 
+			// TODO: pool static allocations as well?
+			delete RefAllocation->Allocation;
+
 			TrackedAllocationMap.Remove(Key);
 		}
 	}
@@ -577,10 +580,14 @@ void FRayTracingShaderBindingTable::FlushAllocationsToClear(FRHICommandList& RHI
 	uint32 LocalBindingIndex = 0;
 	for (TConstSetBitIterator<> It(PersistentRecordsToClear); It; ++It)
 	{
-		FRayTracingLocalShaderBindings& RecordBindingToClear = LocalBindings[LocalBindingIndex++];
-		RecordBindingToClear = FRayTracingLocalShaderBindings();
-		RecordBindingToClear.RecordIndex = It.GetIndex();
-		RecordBindingToClear.BindingType = ERayTracingLocalShaderBindingType::Clear;
+		uint32 RecordIndexToClear = It.GetIndex();
+		if (RecordIndexToClear < MaxNumValidRecords)
+		{
+			FRayTracingLocalShaderBindings& RecordBindingToClear = LocalBindings[LocalBindingIndex++];
+			RecordBindingToClear = FRayTracingLocalShaderBindings();
+			RecordBindingToClear.RecordIndex = RecordIndexToClear;
+			RecordBindingToClear.BindingType = ERayTracingLocalShaderBindingType::Clear;
+		}
 	}
 
 	const bool bCopyDataToInlineStorage = false; // Storage is already allocated from RHICmdList, no extra copy necessary
diff --git a/Engine/Source/Runtime/Renderer/Private/RayTracing/RayTracingShaderBindingTable.h b/Engine/Source/Runtime/Renderer/Private/RayTracing/RayTracingShaderBindingTable.h
--- a/Engine/Source/Runtime/Renderer/Private/RayTracing/RayTracingShaderBindingTable.h
+++ b/Engine/Source/Runtime/Renderer/Private/RayTracing/RayTracingShaderBindingTable.h
@@ -280,8 +280,8 @@ private:
 
 	FRayTracingShaderBindingTableInitializer PersistentSBTInitializer;	 //< Shared initializer used for all persistent SBTs - so they can all be versioned together
 	TArray<FPersistentSBTData> PersistentSBTs;							 //< All currently allocated persistent SBTs (FRayTracingPersistentShaderBindingTableID contains index into this array so can be sparse)
-	TBitArray<> SetPersistentRecords;									 //< BitArray containing which bits are valid in the cached persistent SBTs (used to build the dirty set)
-	TBitArray<> DirtyPersistentRecords;									 //< BitArray containing which bits are valid in the cached persistent SBTs (used to build the dirty set)
+	TBitArray<> SetPersistentRecords;									 //< BitArray containing which bits are valid in the cached persistent SBTs (use for validation)
+	TBitArray<> DirtyPersistentRecords;									 //< BitArray containing which bits are dirty in the cached persistent SBTs (used during dirty binding collection from currently visible bindings)
 	TBitArray<> PersistentRecordsToClear;								 //< BitArray containing which bits need to be cleared in the persisten SBTs
 	
 	FCriticalSection StaticAllocationCS;								 //< Critical section used to access all static allocation data
