[HWRT] Use num valid bindings instead of all allocated local bindings in SetBindingsOnShaderBindingTable to make sure RHI doesn't iterate unitialized entries

diff --git a/Engine/Source/Runtime/Renderer/Private/RayTracing/RayTracingShaderBindingTable.cpp b/Engine/Source/Runtime/Renderer/Private/RayTracing/RayTracingShaderBindingTable.cpp
--- a/Engine/Source/Runtime/Renderer/Private/RayTracing/RayTracingShaderBindingTable.cpp
+++ b/Engine/Source/Runtime/Renderer/Private/RayTracing/RayTracingShaderBindingTable.cpp
@@ -598,7 +598,7 @@ void FRayTracingShaderBindingTable::FlushAllocationsToClear(FRHICommandList& RHI
 			RHICmdList.SetBindingsOnShaderBindingTable(
 				PersistentSBT.ShaderBindingTable,
 				nullptr,
-				TotalRecordsToClear, LocalBindings,
+				LocalBindingIndex, LocalBindings,
 				ERayTracingBindingType::HitGroup,
 				bCopyDataToInlineStorage);
 		}
