[Backout] - CL45363861
#rnx
Original CL Desc
-----------------------------------------------------------------
[RayTracing]
* Remove call to FlushAllocationsToClear on SBT before setting the current bindings on the SBT because the clear call could flush pending clear ops which are still used this frame. Only flush clear after add & remove of primitive scene infos on RT timeline. Pending clears can be added by async stream in requests in between the building of the dirty bindings and the flush to RHI from RDG passes causing cleared data in the SBT and triggering validation errors and potentially missing entries in the SBT.
* Also keep track of dirty records and clear dirty record when reallocated again before the clear was flushed (otherwise valid entries can be cleared)

diff --git a/Engine/Source/Runtime/Renderer/Private/DeferredShadingRenderer.cpp b/Engine/Source/Runtime/Renderer/Private/DeferredShadingRenderer.cpp
--- a/Engine/Source/Runtime/Renderer/Private/DeferredShadingRenderer.cpp
+++ b/Engine/Source/Runtime/Renderer/Private/DeferredShadingRenderer.cpp
@@ -1117,6 +1117,15 @@ void FDeferredShadingSceneRenderer::SetupRayTracingRenderingData(FRDGBuilder& Gr
 		DirtyPersistentRayTracingShaderBindings = Scene->RayTracingSBT.GetDirtyBindings(VisibleRayTracingShaderBindings, bForceAllDirty);
 	}
 
+	GraphBuilder.AddPass(RDG_EVENT_NAME("FlushRayTracingSBTAllocationsToClear"), ERDGPassFlags::NeverCull,
+		[Scene = Scene](FRDGAsyncTask, FRHICommandList& RHICmdList)
+		{
+			TRACE_CPUPROFILER_EVENT_SCOPE(FlushRayTracingSBTAllocationsToClear);
+
+			Scene->RayTracingSBT.FlushAllocationsToClear(RHICmdList);
+		}
+	);
+
 	FRDGBufferRef LumenHardwareRayTracingHitDataBuffer = nullptr;
 	if (bAnyInlineHardwareRayTracingPassEnabled)
 	{
diff --git a/Engine/Source/Runtime/Renderer/Private/RayTracing/RayTracingShaderBindingTable.cpp b/Engine/Source/Runtime/Renderer/Private/RayTracing/RayTracingShaderBindingTable.cpp
--- a/Engine/Source/Runtime/Renderer/Private/RayTracing/RayTracingShaderBindingTable.cpp
+++ b/Engine/Source/Runtime/Renderer/Private/RayTracing/RayTracingShaderBindingTable.cpp
@@ -188,13 +188,11 @@ FRayTracingShaderBindingDataOneFrameArray FRayTracingShaderBindingTable::GetDirt
 			}
 			else
 			{
-				check(VisibleBinding.BindingType!= ERayTracingLocalShaderBindingType::Persistent || SetPersistentRecords[VisibleBinding.SBTRecordIndex]);
-
-				FBitReference DirtyRecordIndex = DirtyPersistentRecords[VisibleBinding.SBTRecordIndex];
-				if (DirtyRecordIndex || VisibleBinding.BindingType == ERayTracingLocalShaderBindingType::Transient)
+				FBitReference RecordIndex = ValidPersistentRecords[VisibleBinding.SBTRecordIndex];
+				if (!RecordIndex || VisibleBinding.BindingType == ERayTracingLocalShaderBindingType::Transient)
 				{
 					DirtyBindings.Add(VisibleBinding);
-					DirtyRecordIndex = false;
+					RecordIndex = true;
 				}
 				else if (GValidatePersistentBindings)
 				{
@@ -255,7 +253,7 @@ FRayTracingPersistentShaderBindingTableID FRayTracingShaderBindingTable::Allocat
 	PersistentSBTs[ResultID].ShaderBindingMode = ShaderBindingMode;
 	AllocatePersistentShaderBindingTable(RHICmdList, ResultID, ShaderBindingMode);
 
-	DirtyPersistentRecords.Init(true, PersistentSBTInitializer.NumGeometrySegments * PersistentSBTInitializer.NumShaderSlotsPerGeometrySegment);
+	ValidPersistentRecords.Init(false, PersistentSBTInitializer.NumGeometrySegments * PersistentSBTInitializer.NumShaderSlotsPerGeometrySegment);
 
 	return ResultID;
 }
@@ -332,7 +330,7 @@ void FRayTracingShaderBindingTable::CheckPersistentRHI(FRHICommandListBase& RHIC
 				AllocatePersistentShaderBindingTable(RHICmdList, Index, PersistentSBTs[Index].ShaderBindingMode);
 			}
 		}
-		DirtyPersistentRecords.Init(true, NewSBTInitializer.NumGeometrySegments * NewSBTInitializer.NumShaderSlotsPerGeometrySegment);
+		ValidPersistentRecords.Init(false, NewSBTInitializer.NumGeometrySegments * NewSBTInitializer.NumShaderSlotsPerGeometrySegment);
 	}
 }
 
@@ -370,35 +368,13 @@ uint32 FRayTracingShaderBindingTable::GetMaxAllocatedStaticSegmentCount() const
 void FRayTracingShaderBindingTable::MarkDirty(FRayTracingSBTAllocation* SBTAllocation)
 {
 	int32 MaxRecordIndex = SBTAllocation->BaseRecordIndex + SBTAllocation->NumRecords;
-	if (DirtyPersistentRecords.Num() < MaxRecordIndex)
-	{
-		DirtyPersistentRecords.SetNum(FMath::RoundUpToPowerOfTwo(MaxRecordIndex), true);
-	}
-	for (uint32 Index = 0; Index < SBTAllocation->NumRecords; ++Index)
-	{
-		DirtyPersistentRecords[SBTAllocation->BaseRecordIndex + Index] = true;
-	}
-}
-
-void FRayTracingShaderBindingTable::MarkSet(FRayTracingSBTAllocation* SBTAllocation, bool bValue)
-{
-	int32 MaxRecordIndex = SBTAllocation->BaseRecordIndex + SBTAllocation->NumRecords;
-	if (SetPersistentRecords.Num() < MaxRecordIndex)
+	if (ValidPersistentRecords.Num() < MaxRecordIndex)
 	{
-		SetPersistentRecords.SetNum(FMath::RoundUpToPowerOfTwo(MaxRecordIndex), false);
-		PersistentRecordsToClear.SetNum(FMath::RoundUpToPowerOfTwo(MaxRecordIndex), false);
+		ValidPersistentRecords.SetNum(FMath::RoundUpToPowerOfTwo(MaxRecordIndex), false);
 	}
 	for (uint32 Index = 0; Index < SBTAllocation->NumRecords; ++Index)
 	{
-		uint32 RecordIndex = SBTAllocation->BaseRecordIndex + Index;
-		check(SetPersistentRecords[RecordIndex] != bValue);
-		SetPersistentRecords[RecordIndex] = bValue;
-
-		// Mark record to clear and validate that it's not marked for clear when we are trying to clear it out (double clear?)
-		check(bValue || !PersistentRecordsToClear[RecordIndex]);
-		PersistentRecordsToClear[RecordIndex] = !bValue;
-
-		//UE_LOG(LogRenderer, Log, TEXT("Update set SBT record data for RecordIndex %d with new value: %d"), RecordIndex, bValue ? 1 : 0);
+		ValidPersistentRecords[SBTAllocation->BaseRecordIndex + Index] = false;
 	}
 }
 
@@ -420,7 +396,6 @@ FRayTracingSBTAllocation* FRayTracingShaderBindingTable::AllocateStaticRangeInte
 	Allocation->InitStatic(AllocatedLayers, BaseIndex, RecordsPerLayer, RecordCount, Geometry, Flags);
 
 	MarkDirty(Allocation);
-	MarkSet(Allocation, true);
 
 	AllocatedStaticSegmentCount += SegmentCount * LayersCount;
 	
@@ -502,11 +477,7 @@ void FRayTracingShaderBindingTable::FreeStaticRange(const FRayTracingSBTAllocati
 			StaticRangeAllocator.Free(RefAllocation->Allocation->BaseRecordIndex, RefAllocation->Allocation->NumRecords);
 			AllocatedStaticSegmentCount -= (RefAllocation->Allocation->NumRecords / NumShaderSlotsPerGeometrySegment);
 		
-			MarkSet(RefAllocation->Allocation, false);
-
-			// TODO: pool static allocations as well?
-			delete RefAllocation->Allocation;
-
+			PersistentAllocationsToClear.Add(RefAllocation->Allocation);
 			TrackedAllocationMap.Remove(Key);
 		}
 	}
@@ -548,7 +519,7 @@ FRayTracingSBTAllocation* FRayTracingShaderBindingTable::AllocateDynamicRange(ER
 
 bool FRayTracingShaderBindingTable::IsDirty(uint32 RecordIndex) const
 {
-	return DirtyPersistentRecords[RecordIndex];
+	return !ValidPersistentRecords[RecordIndex];
 }
 
 bool FRayTracingShaderBindingTable::IsPersistent() const
@@ -563,9 +534,34 @@ void FRayTracingShaderBindingTable::FlushAllocationsToClear(FRHICommandList& RHI
 	// Don't clear outside of current allocated range - could have allocated and cleared ranges before calling CheckPersistentRHI
 	// and then the RHI SBT doesn't have those new ranges allocated yet
 	uint32 MaxNumValidRecords = PersistentSBTInitializer.NumGeometrySegments * PersistentSBTInitializer.NumShaderSlotsPerGeometrySegment;
-	check(MaxNumValidRecords <= (uint32)DirtyPersistentRecords.Num());
+	check(MaxNumValidRecords <= (uint32)ValidPersistentRecords.Num());
+
+	// Build bit array if records to clear (make sure we don't clear the same record twice)
+	TBitArray<> PersistentRecordsToClear;
+	PersistentRecordsToClear.Init(false, MaxNumValidRecords);
 	
-	uint32 TotalRecordsToClear = PersistentRecordsToClear.CountSetBits();
+	uint32 TotalRecordsToClear = 0;
+	for (const FRayTracingSBTAllocation* Allocation : PersistentAllocationsToClear)
+	{
+		for (uint32 RangeIndex = 0; RangeIndex < Allocation->NumRecords; ++RangeIndex)
+		{
+			uint32 RecordIndex = Allocation->BaseRecordIndex + RangeIndex;
+			if (RecordIndex < MaxNumValidRecords)
+			{
+				FBitReference RecordFlag = PersistentRecordsToClear[RecordIndex];
+				if (!RecordFlag)
+				{
+					TotalRecordsToClear++;
+					RecordFlag = true;
+				}
+			}
+		}
+
+		// TODO: pool static allocations as well?
+		delete Allocation;
+	}
+	PersistentAllocationsToClear.Reset(PersistentAllocationsToClear.Num());
+
 	if (TotalRecordsToClear == 0)
 	{
 		return;
@@ -603,9 +599,6 @@ void FRayTracingShaderBindingTable::FlushAllocationsToClear(FRHICommandList& RHI
 				bCopyDataToInlineStorage);
 		}
 	}
-
-	// Reset the records to clear
-	PersistentRecordsToClear.Init(false, PersistentRecordsToClear.Num());
 }
 
 #endif // RHI_RAYTRACING
\ No newline at end of file
diff --git a/Engine/Source/Runtime/Renderer/Private/RayTracing/RayTracingShaderBindingTable.h b/Engine/Source/Runtime/Renderer/Private/RayTracing/RayTracingShaderBindingTable.h
--- a/Engine/Source/Runtime/Renderer/Private/RayTracing/RayTracingShaderBindingTable.h
+++ b/Engine/Source/Runtime/Renderer/Private/RayTracing/RayTracingShaderBindingTable.h
@@ -232,7 +232,6 @@ private:
 	 * Mark all records used by the SBTAllocation as dirty again
 	 */
 	void MarkDirty(FRayTracingSBTAllocation* SBTAllocation);
-	void MarkSet(FRayTracingSBTAllocation* SBTAllocation, bool bValue);
 
 	struct FAllocationKey
 	{
@@ -280,9 +279,8 @@ private:
 
 	FRayTracingShaderBindingTableInitializer PersistentSBTInitializer;	 //< Shared initializer used for all persistent SBTs - so they can all be versioned together
 	TArray<FPersistentSBTData> PersistentSBTs;							 //< All currently allocated persistent SBTs (FRayTracingPersistentShaderBindingTableID contains index into this array so can be sparse)
-	TBitArray<> SetPersistentRecords;									 //< BitArray containing which bits are valid in the cached persistent SBTs (use for validation)
-	TBitArray<> DirtyPersistentRecords;									 //< BitArray containing which bits are dirty in the cached persistent SBTs (used during dirty binding collection from currently visible bindings)
-	TBitArray<> PersistentRecordsToClear;								 //< BitArray containing which bits need to be cleared in the persisten SBTs
+	TBitArray<> ValidPersistentRecords;									 //< BitArray containing which bits are valid in the cached persistent SBTs (used to build the dirty set)
+	TArray<FRayTracingSBTAllocation*> PersistentAllocationsToClear;		 //< Persistent SBT allocations about which the cached persistent SBTs need to be notified (clear internal cached record data)
 	
 	FCriticalSection StaticAllocationCS;								 //< Critical section used to access all static allocation data
 	bool bStaticAllocationsLocked = false;								 //< Static allocations are not allowed when this bool is set (used for validation)
