Add support for using a seperate Hit group contribution buffer

#DisabledValidations Preflight

diff --git a/Engine/Plugins/Experimental/GPULightmass/Source/GPULightmass/Private/LightmapRenderer.cpp b/Engine/Plugins/Experimental/GPULightmass/Source/GPULightmass/Private/LightmapRenderer.cpp
--- a/Engine/Plugins/Experimental/GPULightmass/Source/GPULightmass/Private/LightmapRenderer.cpp
+++ b/Engine/Plugins/Experimental/GPULightmass/Source/GPULightmass/Private/LightmapRenderer.cpp
@@ -1211,6 +1211,9 @@ bool FSceneRenderState::SetupRayTracingScene(FRDGBuilder& GraphBuilder, FSceneUn
 
 			FRWBufferStructured InstanceBuffer;
 			InstanceBuffer.Initialize(RHICmdList, TEXT("LightmassRayTracingInstanceBuffer"), GRHIRayTracingInstanceDescriptorSize, SceneInitializer.MaxNumInstances);
+			
+			FRWBufferStructured HitGroupContributionsBuffer;
+			InstanceBuffer.Initialize(RHICmdList, TEXT("LightmassRayTracingHitGroupBuffer"), 4, SceneInitializer.MaxNumInstances);
 
 			RayTracingInstanceBufferBuilder.FillRayTracingInstanceUploadBuffer(RHICmdList);
 			RayTracingInstanceBufferBuilder.FillAccelerationStructureAddressesBuffer(RHICmdList);
@@ -1220,6 +1223,7 @@ bool FSceneRenderState::SetupRayTracingScene(FRDGBuilder& GraphBuilder, FSceneUn
 				nullptr,
 				nullptr,
 				InstanceBuffer.UAV,
+				GRHIGlobals.RayTracing.RequiresSeparateHitGroupContributionsBuffer ? HitGroupContributionsBuffer.UAV : nullptr,
 				SceneInitializer.MaxNumInstances,
 				/*bCompactOutput*/ false,
 				nullptr,
@@ -1235,6 +1239,13 @@ bool FSceneRenderState::SetupRayTracingScene(FRDGBuilder& GraphBuilder, FSceneUn
 				BuildParams.ScratchBufferOffset = 0;
 				BuildParams.InstanceBuffer = InstanceBuffer.Buffer;
 				BuildParams.InstanceBufferOffset = 0;
+				
+				if(GRHIGlobals.RayTracing.RequiresSeparateHitGroupContributionsBuffer)
+				{
+					BuildParams.HitGroupContributionsBuffer = HitGroupContributionsBuffer.Buffer;
+					BuildParams.HitGroupContributionsBufferOffset = 0;
+				}
+				
 				BuildParams.ReferencedGeometries = RayTracingInstanceBufferBuilder.GetReferencedGeometries();
 				BuildParams.NumInstances = RayTracingInstanceBufferBuilder.GetMaxNumInstances();
 
diff --git a/Engine/Shaders/Private/RayTracing/RayTracingInstanceBufferUtil.usf b/Engine/Shaders/Private/RayTracing/RayTracingInstanceBufferUtil.usf
--- a/Engine/Shaders/Private/RayTracing/RayTracingInstanceBufferUtil.usf
+++ b/Engine/Shaders/Private/RayTracing/RayTracingInstanceBufferUtil.usf
@@ -36,6 +36,10 @@ float3 PreViewTranslationLow;
 
 RWStructuredBuffer<FPlatformRayTracingInstanceDescriptor> OutPlatformInstanceDescriptors;
 
+#if USE_SEPARATE_HIT_CONTRIBUTION_BUFFER
+RWStructuredBuffer<uint> OutHitGroupContributions;
+#endif
+
 StructuredBuffer<uint> InstanceGroupDescriptors;
 StructuredBuffer<FRayTracingInstanceDescriptor> InstanceDescriptors;
 ByteAddressBuffer AccelerationStructureAddresses;
@@ -213,6 +217,10 @@ void ProcessInstance(FRayTracingInstanceDescriptor InputDesc)
 		InputDesc.InstanceContributionToHitGroupIndex,
 		LocalToTranslatedWorld,
 		BlasAddress);
+		
+#if USE_SEPARATE_HIT_CONTRIBUTION_BUFFER
+	OutHitGroupContributions[InputDesc.OutputDescriptorIndex] = InputDesc.InstanceContributionToHitGroupIndex;
+#endif
 }
 
 [numthreads(THREADGROUP_SIZE, 1, 1)]
diff --git a/Engine/Source/Runtime/RHI/Public/RHIContext.h b/Engine/Source/Runtime/RHI/Public/RHIContext.h
--- a/Engine/Source/Runtime/RHI/Public/RHIContext.h
+++ b/Engine/Source/Runtime/RHI/Public/RHIContext.h
@@ -612,6 +612,11 @@ struct FRayTracingSceneBuildParams
 	// Buffer of native ray tracing instance descriptors. Must be in SRV state.
 	FRHIBuffer* InstanceBuffer = nullptr;
 	uint32 InstanceBufferOffset = 0;
+	
+	// Buffer of native ray tracing hit contribution buffer. Must be in SRV state.
+	// Only used when RequiresSeparateHitGroupContributionsBuffer is set
+	FRHIBuffer* HitGroupContributionsBuffer = nullptr;
+	uint32 HitGroupContributionsBufferOffset = 0;
 
 	uint32 NumInstances = 0;
 
diff --git a/Engine/Source/Runtime/RHI/Public/RHIGlobals.h b/Engine/Source/Runtime/RHI/Public/RHIGlobals.h
--- a/Engine/Source/Runtime/RHI/Public/RHIGlobals.h
+++ b/Engine/Source/Runtime/RHI/Public/RHIGlobals.h
@@ -466,6 +466,9 @@ struct FRHIGlobals
 
 		/** Whether or not the RHI supports serializing acceleration structures to a file */
 		bool SupportsSerializeAccelerationStructure = false;
+		
+		/** Whether or not the RHI supports a seperate hit group buffer when building the acceleration structure */
+		bool RequiresSeparateHitGroupContributionsBuffer = false;
 
 		/** Required alignment for ray tracing acceleration structures. */
 		uint32 AccelerationStructureAlignment = 0;
@@ -961,4 +964,4 @@ inline EShaderPlatform GetFeatureLevelShaderPlatform_Checked(const FStaticFeatur
 {
 	checkf(!IsRunningCookCommandlet(), TEXT("Calling GetFeatureLevelShaderPlatform while inside a cook commandlet is not correct, as it will return the cookers ShaderPlatform and not the ShaderPlatform you are cooking for"));
 	return GetFeatureLevelShaderPlatform(InFeatureLevel);
-}
\ No newline at end of file
+}
diff --git a/Engine/Source/Runtime/Renderer/Private/RayTracing/RayTracingInstanceBufferUtil.cpp b/Engine/Source/Runtime/Renderer/Private/RayTracing/RayTracingInstanceBufferUtil.cpp
--- a/Engine/Source/Runtime/Renderer/Private/RayTracing/RayTracingInstanceBufferUtil.cpp
+++ b/Engine/Source/Runtime/Renderer/Private/RayTracing/RayTracingInstanceBufferUtil.cpp
@@ -577,6 +577,7 @@ struct FRayTracingBuildInstanceBufferCS : public FGlobalShader
 		SHADER_PARAMETER_STRUCT_INCLUDE(FGPUSceneResourceParametersRHI, GPUSceneParameters)
 
 		SHADER_PARAMETER_UAV(RWStructuredBuffer, OutPlatformInstanceDescriptors)
+		SHADER_PARAMETER_UAV(RWStructuredBuffer, OutHitGroupContributions)
 
 		SHADER_PARAMETER_SRV(StructuredBuffer, InstanceGroupDescriptors)
 		SHADER_PARAMETER_SRV(StructuredBuffer, InstanceDescriptors)
@@ -614,7 +615,8 @@ struct FRayTracingBuildInstanceBufferCS : public FGlobalShader
 	class FOutputStatsDim : SHADER_PERMUTATION_BOOL("OUTPUT_STATS");
 	class FCompactOutputDim : SHADER_PERMUTATION_BOOL("COMPACT_OUTPUT");
 	class FUseWaveOpsDim : SHADER_PERMUTATION_BOOL("USE_WAVE_OPS");
-	using FPermutationDomain = TShaderPermutationDomain<FSupportInstanceGroupsDim, FUseGPUSceneDim, FOutputInstanceExtraDataDim, FGpuCullingDim, FOutputStatsDim, FCompactOutputDim, FUseWaveOpsDim>;
+	class FUseSeparateHitContributionBuffer : SHADER_PERMUTATION_BOOL("USE_SEPARATE_HIT_CONTRIBUTION_BUFFER");
+	using FPermutationDomain = TShaderPermutationDomain<FSupportInstanceGroupsDim, FUseGPUSceneDim, FOutputInstanceExtraDataDim, FGpuCullingDim, FOutputStatsDim, FCompactOutputDim, FUseWaveOpsDim, FUseSeparateHitContributionBuffer>;
 		
 	static constexpr uint32 ThreadGroupSize = GRayTracingInstanceGroupSize;
 
@@ -653,6 +655,7 @@ void BuildRayTracingInstanceBuffer(
 	uint32 NumGroups,
 	uint32 NumInstanceDescriptors,
 	FRHIUnorderedAccessView* InstancesUAV,
+	FRHIUnorderedAccessView* HitGroupContributionsUAV,
 	FRHIShaderResourceView* InstanceGroupUploadSRV,
 	uint32 InstanceGroupUploadOffset,
 	FRHIShaderResourceView* InstanceUploadSRV,
@@ -667,6 +670,7 @@ void BuildRayTracingInstanceBuffer(
 {
 	FRayTracingBuildInstanceBufferCS::FParameters PassParams;
 	PassParams.OutPlatformInstanceDescriptors = InstancesUAV;
+	PassParams.OutHitGroupContributions = HitGroupContributionsUAV;
 	PassParams.InstanceGroupDescriptors = InstanceGroupUploadSRV;
 	PassParams.InstanceDescriptors = InstanceUploadSRV;
 	PassParams.AccelerationStructureAddresses = AccelerationStructureAddressesSRV;
@@ -707,6 +711,7 @@ void BuildRayTracingInstanceBuffer(
 	PermutationVector.Set<FRayTracingBuildInstanceBufferCS::FCompactOutputDim>(bCompactOutput);
 	PermutationVector.Set<FRayTracingBuildInstanceBufferCS::FUseWaveOpsDim>(GRHISupportsWaveOperations);
 	PermutationVector.Set<FRayTracingBuildInstanceBufferCS::FGpuCullingDim>(CullingParameters != nullptr);
+	PermutationVector.Set<FRayTracingBuildInstanceBufferCS::FUseSeparateHitContributionBuffer>(GRHIGlobals.RayTracing.RequiresSeparateHitGroupContributionsBuffer);
 
 	auto ComputeShader = GetGlobalShaderMap(GMaxRHIFeatureLevel)->GetShader<FRayTracingBuildInstanceBufferCS>(PermutationVector);
 
@@ -728,6 +733,7 @@ void BuildRayTracingInstanceBuffer(
 	const FGPUScene* GPUScene,
 	const FDFVector3& PreViewTranslation,
 	FRHIUnorderedAccessView* InstancesUAV,
+	FRHIUnorderedAccessView* HitGroupContributionsUAV,
 	FRHIShaderResourceView* InstanceGroupUploadSRV,
 	FRHIShaderResourceView* InstanceUploadSRV,
 	FRHIShaderResourceView* AccelerationStructureAddressesSRV,
@@ -753,6 +759,7 @@ void BuildRayTracingInstanceBuffer(
 			NumGPUGroups,
 			NumGPUInstanceDescriptors,
 			InstancesUAV,
+			HitGroupContributionsUAV,
 			InstanceGroupUploadSRV,
 			0,
 			InstanceUploadSRV,
@@ -776,6 +783,7 @@ void BuildRayTracingInstanceBuffer(
 			NumCPUGroups,
 			NumCPUInstanceDescriptors,
 			InstancesUAV,
+			HitGroupContributionsUAV,
 			InstanceGroupUploadSRV,
 			NumGPUGroups, // CPU instance group descriptors are stored after GPU Scene instance groups
 			InstanceUploadSRV,
@@ -795,6 +803,7 @@ void BuildRayTracingInstanceBuffer(
 	const FGPUScene* GPUScene,
 	const FDFVector3& PreViewTranslation,
 	FRHIUnorderedAccessView* InstancesUAV,
+	FRHIUnorderedAccessView* HitGroupContributionsUAV,
 	FRHIShaderResourceView* InstanceUploadSRV,
 	FRHIShaderResourceView* AccelerationStructureAddressesSRV,
 	FRHIShaderResourceView* CPUInstanceTransformSRV,
@@ -809,6 +818,7 @@ void BuildRayTracingInstanceBuffer(
 		GPUScene,
 		PreViewTranslation,
 		InstancesUAV,
+		HitGroupContributionsUAV,
 		nullptr,
 		InstanceUploadSRV,
 		AccelerationStructureAddressesSRV,
@@ -1000,6 +1010,7 @@ void FRayTracingInstanceBufferBuilder::BuildRayTracingInstanceBuffer(
 	const FGPUScene* GPUScene,
 	const FRayTracingCullingParameters* CullingParameters,
 	FRHIUnorderedAccessView* InstancesUAV,
+	FRHIUnorderedAccessView* HitGroupContributionsUAV,
 	uint32 MaxNumInstances,
 	bool bCompactOutput,
 	FRHIUnorderedAccessView* OutputStatsUAV,
@@ -1011,6 +1022,7 @@ void FRayTracingInstanceBufferBuilder::BuildRayTracingInstanceBuffer(
 		GPUScene,
 		FDFVector3(PreViewTranslation),
 		InstancesUAV,
+		HitGroupContributionsUAV,
 		InstanceGroupUploadSRV,
 		InstanceUploadSRV,
 		AccelerationStructureAddressesBuffer.SRV,
diff --git a/Engine/Source/Runtime/Renderer/Private/RayTracing/RayTracingScene.cpp b/Engine/Source/Runtime/Renderer/Private/RayTracing/RayTracingScene.cpp
--- a/Engine/Source/Runtime/Renderer/Private/RayTracing/RayTracingScene.cpp
+++ b/Engine/Source/Runtime/Renderer/Private/RayTracing/RayTracingScene.cpp
@@ -83,6 +83,7 @@ bool IsRayTracingFeedbackEnabled(const FSceneViewFamily& ViewFamily)
 
 BEGIN_SHADER_PARAMETER_STRUCT(FBuildInstanceBufferPassParams, )
 	SHADER_PARAMETER_RDG_BUFFER_UAV(RWStructuredBuffer, InstanceBuffer)
+	SHADER_PARAMETER_RDG_BUFFER_UAV(RWStructuredBuffer, HitGroupContributionsBuffer)
 	SHADER_PARAMETER_RDG_BUFFER_UAV(RWStructuredBuffer, OutputStats)
 	SHADER_PARAMETER_RDG_BUFFER_UAV(RWStructuredBuffer, InstanceExtraDataBuffer)
 	SHADER_PARAMETER_RDG_UNIFORM_BUFFER(FSceneUniformParameters, Scene)
@@ -314,6 +315,16 @@ void FRayTracingScene::Update(FRDGBuilder& GraphBuilder, FSceneUniformBuffer& Sc
 					AddClearUAVPass(GraphBuilder, GraphBuilder.CreateUAV(LayerView.InstanceBuffer), 0, ComputePassFlags);
 				}
 			}
+			
+			if(GRHIGlobals.RayTracing.RequiresSeparateHitGroupContributionsBuffer)
+			{
+				FRDGBufferDesc HitGroupContributionsDesc;
+				HitGroupContributionsDesc.Usage = EBufferUsageFlags::UnorderedAccess | EBufferUsageFlags::ShaderResource | EBufferUsageFlags::StructuredBuffer;
+				HitGroupContributionsDesc.BytesPerElement = 4;
+				HitGroupContributionsDesc.NumElements = NumNativeInstancesAligned;
+				
+				LayerView.HitGroupContributionsBuffer = GraphBuilder.CreateBuffer(HitGroupContributionsDesc, TEXT("FRayTracingScene::HitGroupContributionsBuffer"));
+			}
 
 			// Feedback
 			if (bTracingFeedbackEnabled)
@@ -382,6 +393,12 @@ void FRayTracingScene::Update(FRDGBuilder& GraphBuilder, FSceneUniformBuffer& Sc
 				{
 					FBuildInstanceBufferPassParams* PassParams = GraphBuilder.AllocParameters<FBuildInstanceBufferPassParams>();
 					PassParams->InstanceBuffer = GraphBuilder.CreateUAV(LayerView.InstanceBuffer);
+					
+					if(GRHIGlobals.RayTracing.RequiresSeparateHitGroupContributionsBuffer)
+					{
+						PassParams->HitGroupContributionsBuffer = GraphBuilder.CreateUAV(LayerView.HitGroupContributionsBuffer);
+					}
+					
 					PassParams->InstanceExtraDataBuffer = InstanceExtraDataBufferUAV;
 					PassParams->Scene = SceneUniformBuffer.GetBuffer(GraphBuilder);
 					PassParams->OutputStats = bCompactInstanceBuffer || bStatsEnabled ? InstanceStatsBufferUAV : nullptr;
@@ -406,6 +423,7 @@ void FRayTracingScene::Update(FRDGBuilder& GraphBuilder, FSceneUniformBuffer& Sc
 								GPUScene,
 								CullingParameters,
 								PassParams->InstanceBuffer->GetRHI(),
+								GRHIGlobals.RayTracing.RequiresSeparateHitGroupContributionsBuffer ? PassParams->HitGroupContributionsBuffer->GetRHI() : nullptr,
 								NumNativeInstances,
 								bCompactInstanceBuffer,
 								PassParams->OutputStats ? PassParams->OutputStats->GetRHI() : nullptr,
@@ -463,6 +481,12 @@ void FRayTracingScene::Build(FRDGBuilder& GraphBuilder, ERDGPassFlags ComputePas
 
 			PassParams->TLASBuildBuffers.Emplace(LayerView.BuildScratchBuffer, ERHIAccess::UAVCompute);
 			PassParams->TLASBuildBuffers.Emplace(LayerView.InstanceBuffer, ERHIAccess::SRVCompute);
+			
+			if(GRHIGlobals.RayTracing.RequiresSeparateHitGroupContributionsBuffer)
+			{
+				PassParams->TLASBuildBuffers.Emplace(LayerView.HitGroupContributionsBuffer, ERHIAccess::SRVCompute);
+			}
+			
 			PassParams->TLASBuildBuffers.Emplace(LayerView.RayTracingSceneBufferRDG, ERHIAccess::BVHWrite);
 		}
 	}
@@ -489,6 +513,14 @@ void FRayTracingScene::Build(FRDGBuilder& GraphBuilder, ERDGPassFlags ComputePas
 					BuildParams.ScratchBufferOffset = 0;
 					BuildParams.InstanceBuffer = LayerView.InstanceBuffer->GetRHI();
 					BuildParams.InstanceBufferOffset = 0;
+					
+					if(GRHIGlobals.RayTracing.RequiresSeparateHitGroupContributionsBuffer)
+					{
+						check(LayerView.HitGroupContributionsBuffer);
+						BuildParams.HitGroupContributionsBuffer = LayerView.HitGroupContributionsBuffer->GetRHI();
+						BuildParams.HitGroupContributionsBufferOffset = 0;
+					}
+					
 					BuildParams.NumInstances = LayerView.MaxNumInstances;
 					BuildParams.ReferencedGeometries = LayerView.InstanceBufferBuilder.GetReferencedGeometries();
 
@@ -1106,6 +1138,7 @@ void FRayTracingScene::Reset()
 			LayerView.RayTracingSceneBufferSRV = nullptr;
 
 			LayerView.InstanceBuffer = nullptr;
+			LayerView.HitGroupContributionsBuffer = nullptr;
 			LayerView.BuildScratchBuffer = nullptr;
 			LayerView.InstanceExtraDataBuffer = nullptr;
 
diff --git a/Engine/Source/Runtime/Renderer/Private/RayTracing/RayTracingScene.h b/Engine/Source/Runtime/Renderer/Private/RayTracing/RayTracingScene.h
--- a/Engine/Source/Runtime/Renderer/Private/RayTracing/RayTracingScene.h
+++ b/Engine/Source/Runtime/Renderer/Private/RayTracing/RayTracingScene.h
@@ -236,6 +236,7 @@ private:
 		FRayTracingSceneRHIRef RayTracingSceneRHI;
 
 		FRDGBufferRef InstanceBuffer = nullptr;
+		FRDGBufferRef HitGroupContributionsBuffer = nullptr;
 		FRDGBufferRef BuildScratchBuffer = nullptr;
 
 		// Feedback		
diff --git a/Engine/Source/Runtime/Renderer/Private/Tests/RayTracingTestbed.cpp b/Engine/Source/Runtime/Renderer/Private/Tests/RayTracingTestbed.cpp
--- a/Engine/Source/Runtime/Renderer/Private/Tests/RayTracingTestbed.cpp
+++ b/Engine/Source/Runtime/Renderer/Private/Tests/RayTracingTestbed.cpp
@@ -204,6 +204,13 @@ bool RunRayTracingTestbed_RenderThread(const FString& Parameters)
 
 	FRWBufferStructured InstanceBuffer;
 	InstanceBuffer.Initialize(RHICmdList, TEXT("RayTracingTestBedInstanceBuffer"), GRHIRayTracingInstanceDescriptorSize, SceneInitializer.MaxNumInstances);
+	
+	FRWBufferStructured HitGroupContributionsBuffer;
+	
+	if(GRHIGlobals.RayTracing.RequiresSeparateHitGroupContributionsBuffer)
+	{
+		HitGroupContributionsBuffer.Initialize(RHICmdList, TEXT("RayTracingTestBedInstanceBuffer"), 4, SceneInitializer.MaxNumInstances);
+	}
 
 	RayTracingInstanceBufferBuilder.FillRayTracingInstanceUploadBuffer(RHICmdList);
 	RayTracingInstanceBufferBuilder.FillAccelerationStructureAddressesBuffer(RHICmdList);
@@ -213,6 +220,7 @@ bool RunRayTracingTestbed_RenderThread(const FString& Parameters)
 		nullptr,
 		nullptr,
 		InstanceBuffer.UAV,
+		GRHIGlobals.RayTracing.RequiresSeparateHitGroupContributionsBuffer ? HitGroupContributionsBuffer.UAV : nullptr,
 		SceneInitializer.MaxNumInstances,
 		/*bCompactOutput*/ false,
 		nullptr,
@@ -232,6 +240,13 @@ bool RunRayTracingTestbed_RenderThread(const FString& Parameters)
 	BuildParams.ScratchBufferOffset = 0;
 	BuildParams.InstanceBuffer = InstanceBuffer.Buffer;
 	BuildParams.InstanceBufferOffset = 0;
+	
+	if(GRHIGlobals.RayTracing.RequiresSeparateHitGroupContributionsBuffer)
+	{
+		BuildParams.HitGroupContributionsBuffer = HitGroupContributionsBuffer.Buffer;
+		BuildParams.HitGroupContributionsBufferOffset = 0;
+	}
+	
 	BuildParams.ReferencedGeometries = RayTracingInstanceBufferBuilder.GetReferencedGeometries();
 	BuildParams.NumInstances = RayTracingInstanceBufferBuilder.GetMaxNumInstances();
 
diff --git a/Engine/Source/Runtime/Renderer/Public/RayTracingInstanceBufferUtil.h b/Engine/Source/Runtime/Renderer/Public/RayTracingInstanceBufferUtil.h
--- a/Engine/Source/Runtime/Renderer/Public/RayTracingInstanceBufferUtil.h
+++ b/Engine/Source/Runtime/Renderer/Public/RayTracingInstanceBufferUtil.h
@@ -89,6 +89,7 @@ RENDERER_API void BuildRayTracingInstanceBuffer(
 	const FGPUScene* GPUScene,
 	const FDFVector3& PreViewTranslation,
 	FRHIUnorderedAccessView* InstancesUAV,
+	FRHIUnorderedAccessView* HitGroupContributionsUAV,
 	FRHIShaderResourceView* InstanceUploadSRV,
 	FRHIShaderResourceView* AccelerationStructureAddressesSRV,
 	FRHIShaderResourceView* CPUInstanceTransformSRV,
@@ -125,6 +126,7 @@ public:
 		const FGPUScene* GPUScene,
 		const FRayTracingCullingParameters* CullingParameters,
 		FRHIUnorderedAccessView* InstancesUAV,
+		FRHIUnorderedAccessView* HitGroupContributionsUAV,
 		uint32 MaxNumInstances,
 		bool bCompactOutput,
 		FRHIUnorderedAccessView* OutputStatsUAV,
