Clean up SceneVisibility nits from bounding box expansion CL

diff --git a/Engine/Source/Runtime/Renderer/Private/SceneVisibility.cpp b/Engine/Source/Runtime/Renderer/Private/SceneVisibility.cpp
--- a/Engine/Source/Runtime/Renderer/Private/SceneVisibility.cpp
+++ b/Engine/Source/Runtime/Renderer/Private/SceneVisibility.cpp
@@ -2839,8 +2839,7 @@ bool FGPUOcclusionPacket::OcclusionCullPrimitive(VisitorType& Visitor, FOcclusio
 				if (GVisualizeOccludedPrimitives == 1 && bIsOccluded)
 				{
 					const FBoxSphereBounds& Bounds = bSubQueries ? (*SubBounds)[SubQuery] : Scene.PrimitiveOcclusionBounds[Index];
-					FViewElementPDI PDI(&View, nullptr, &View.DynamicPrimitiveCollector);
-					DrawWireBox(&PDI, Bounds.GetBox(), FColor(50, 255, 50), SDPG_Foreground);
+					Visitor.AddVisualizeQuery(Bounds.GetBox(), FColor(50, 255, 50));
 				}
 			}
 		}
@@ -2869,8 +2868,7 @@ bool FGPUOcclusionPacket::OcclusionCullPrimitive(VisitorType& Visitor, FOcclusio
 			// Print unexpanded bounding boxes in green if not already visualized by GVisualizeOccludedPrimitives
 			if (GVisualizePrimitiveBBoxes == 2 && !(GVisualizeOccludedPrimitives == 1 && bIsOccluded))
 			{
-				FViewElementPDI PDI(&View, nullptr, &View.DynamicPrimitiveCollector);
-				DrawWireBox(&PDI, PrimitiveOcclusionBounds.GetBox(), FColor(0, 155, 0), SDPG_Foreground);
+				Visitor.AddVisualizeQuery(PrimitiveOcclusionBounds.GetBox(), FColor(0, 155, 0));
 			}
 
 			// scale occlusion bounds by GExpand*TestedBBoxesAmount in world space
@@ -2899,13 +2897,11 @@ bool FGPUOcclusionPacket::OcclusionCullPrimitive(VisitorType& Visitor, FOcclusio
 
 			if (GVisualizeOccludedPrimitives == 2 && bIsOccluded)
 			{
-				FViewElementPDI PDI(&View, nullptr, &View.DynamicPrimitiveCollector);
-				DrawWireBox(&PDI, OcclusionBounds.GetBox(), FColor(50, 255, 50), SDPG_Foreground);
+				Visitor.AddVisualizeQuery(OcclusionBounds.GetBox(), FColor(50, 255, 50));
 			}
 			else if (GVisualizePrimitiveBBoxes > 0)
 			{
-				FViewElementPDI PDI(&View, nullptr, &View.DynamicPrimitiveCollector);
-				DrawWireBox(&PDI, OcclusionBounds.GetBox(), FColor(155, 0, 0), SDPG_Foreground);
+				Visitor.AddVisualizeQuery(OcclusionBounds.GetBox(), FColor(155, 0, 0));
 			}
 
 			bool bAllowBoundsTest;
diff --git a/Engine/Source/Runtime/Renderer/Private/SceneVisibilityPrivate.h b/Engine/Source/Runtime/Renderer/Private/SceneVisibilityPrivate.h
--- a/Engine/Source/Runtime/Renderer/Private/SceneVisibilityPrivate.h
+++ b/Engine/Source/Runtime/Renderer/Private/SceneVisibilityPrivate.h
@@ -978,7 +978,12 @@ public:
 
 		void AddVisualizeQuery(const FBox& Box)
 		{
-			VisualizeQueries.Emplace(Box);
+			VisualizeQueries.Emplace(TPair<FBox, FColor>(Box, FColor(50, 255, 50)));
+		}
+
+		void AddVisualizeQuery(const FBox& Box, const FColor& Color)
+		{
+			VisualizeQueries.Emplace(TPair<FBox, FColor>(Box, Color));
 		}
 
 		void AddOcclusionQuery(const FOcclusionQuery& Query)
@@ -992,7 +997,7 @@ public:
 		TArray<FOcclusionQuery, SceneRenderingAllocator> OcclusionQueries;
 		TArray<FOcclusionFeedbackEntry, SceneRenderingAllocator> OcclusionFeedbacks;
 		TArray<FHZBBound, SceneRenderingAllocator> HZBBounds;
-		TArray<FBox, SceneRenderingAllocator> VisualizeQueries;
+		TArray<TPair<FBox, FColor>, SceneRenderingAllocator> VisualizeQueries;
 	};
 
 	// Visitor variant that will process the requests immediately or replay recorded commands.
@@ -1025,6 +1030,11 @@ public:
 			DrawWireBox(&Packet.ViewElementPDI, Box, FColor(50, 255, 50), SDPG_Foreground);
 		}
 
+		void AddVisualizeQuery(const FBox& Box, const FColor& Color)
+		{
+			DrawWireBox(&Packet.ViewElementPDI, Box, Color, SDPG_Foreground);
+		}
+
 		void AddHZBBounds(const FHZBBound& HZBBounds)
 		{
 			HZBBounds.TargetHistory->HZBTestIndex = Packet.HZBOcclusionTests.AddBounds(HZBBounds.BoundsOrigin, HZBBounds.BoundsExtent);
@@ -1044,9 +1054,9 @@ public:
 				AddOcclusionQuery(Query);
 			}
 
-			for (const FBox& Box : RecordVisitor.VisualizeQueries)
+			for (const TPair<FBox, FColor>& Query : RecordVisitor.VisualizeQueries)
 			{
-				AddVisualizeQuery(Box);
+				AddVisualizeQuery(Query.Key, Query.Value);
 			}
 
 			for (const FHZBBound& HZBBounds : RecordVisitor.HZBBounds)
