Sort lights by substrate complexity and include IESProfile in light sort key to minimize shader switching.

diff --git a/Engine/Shaders/Shared/SceneDefinitions.h b/Engine/Shaders/Shared/SceneDefinitions.h
--- a/Engine/Shaders/Shared/SceneDefinitions.h
+++ b/Engine/Shaders/Shared/SceneDefinitions.h
@@ -16,7 +16,7 @@
 #define INSTANCE_SCENE_DATA_FLAGS_NUM_BITS	(12u) // Max of 12 flags
 #define INSTANCE_RELATIVE_ID_NUM_BITS		(24u) // Max of 16,777,216 instances per primitive
 #define INSTANCE_CUSTOM_DATA_COUNT_NUM_BITS	 (8u) // Max of 255 custom data floats per instance
-#define LIGHT_ID_NUM_BITS					(24u) // Max of 16,777,216 lights in the buffer
+#define LIGHT_ID_NUM_BITS					(23u) // Max of 8,388,608 lights in the buffer
 
 #define MAX_PRIMITIVE_ID					(1U << PRIMITIVE_ID_NUM_BITS)
 #define MAX_INSTANCE_ID						(1U << INSTANCE_ID_NUM_BITS)
diff --git a/Engine/Source/Runtime/Renderer/Private/LightRendering.cpp b/Engine/Source/Runtime/Renderer/Private/LightRendering.cpp
--- a/Engine/Source/Runtime/Renderer/Private/LightRendering.cpp
+++ b/Engine/Source/Runtime/Renderer/Private/LightRendering.cpp
@@ -606,7 +606,9 @@ static void RenderLight(
 	const bool bCanLightUsesAtlasForUnbatchedLight = false,
 	TRDGUniformBufferRef<FVirtualShadowMapUniformParameters> VirtualShadowMapUniformBuffer = nullptr,
 	FRDGTextureRef ShadowMaskBits = nullptr,
-	int32 VirtualShadowMapId = INDEX_NONE);
+	int32 VirtualShadowMapId = INDEX_NONE,
+	ESubstrateTileType SubstrateTileType = ESubstrateTileType::ECount,
+	bool bFirstPass = true);
 
 FDeferredLightUniformStruct GetDeferredLightParameters(const FSceneView& View, const FLightSceneInfo& LightSceneInfo, bool bUseLightFunctionAtlas, uint32 LightFlags)
 {
@@ -1355,6 +1357,7 @@ void FSceneRenderer::GatherAndSortLights(FSortedLightSetSceneInfo& OutSortedLigh
 
 					// Check for shadows and light functions.
 					SortedLightInfo->SortKey.Fields.LightType = LightSceneInfoCompact.LightType;
+					SortedLightInfo->SortKey.Fields.bIESProfile = ViewFamily.EngineShowFlags.TexturedLightProfiles && (LightSceneInfo->Proxy->GetIESTextureResource() != 0);
 					SortedLightInfo->SortKey.Fields.bShadowed = bDynamicShadows && CheckForProjectedShadows(LightSceneInfo);
 					SortedLightInfo->SortKey.Fields.bLightFunction = ViewFamily.EngineShowFlags.LightFunctions && CheckForLightFunction(LightSceneInfo);
 					SortedLightInfo->SortKey.Fields.bUsesLightingChannels = LightSceneInfo->Proxy->GetLightingChannelMask() != GetDefaultLightingChannelMask();
@@ -1594,13 +1597,42 @@ void FDeferredShadingSceneRenderer::RenderLights(
 				RDG_EVENT_SCOPE_CONDITIONAL(GraphBuilder, ViewCount > 1, "View%d", ViewIndex);
 				SCOPED_GPU_MASK(GraphBuilder.RHICmdList, View.GPUMask);
 
-				for (int32 LightIndex = StandardDeferredStart; LightIndex < UnbatchedLightStart; LightIndex++)
+				if (Substrate::IsSubstrateEnabled())
+				{					
+					// Make substrate tile type the outermost sort criteria to minimize shader switching
+					ESubstrateTileType TileType = ESubstrateTileType::EComplexSpecial;
+					bool bFirstPass = true;
+					while(true)
+					{
+						if (Substrate::GetSubstrateUsesTileType(View, TileType))
+						{
+							for (int32 LightIndex = StandardDeferredStart; LightIndex < UnbatchedLightStart; LightIndex++)
+							{
+								// Render the light to the scene color buffer, using a 1x1 white texture as input
+								const FSortedLightSceneInfo& SortedLightInfo = SortedLights[LightIndex];
+								const FLightSceneInfo* LightSceneInfo = SortedLightInfo.LightSceneInfo;
+								RenderLight(GraphBuilder, Scene, View, SceneTextures, LightSceneInfo, nullptr, LightingChannelsTexture,
+									false /*bRenderOverlap*/, false /*bCloudShadow*/, SortedLightInfo.bIsCompatibleWithLightFunctionAtlas,
+									nullptr, nullptr, INDEX_NONE, TileType, bFirstPass);
+								bFirstPass = false;
+							}
+						}
+						if (TileType == ESubstrateTileType::ESimple)
+							break;
+						
+						TileType = ESubstrateTileType(TileType - 1u);
+					}
+				}
+				else
 				{
-					// Render the light to the scene color buffer, using a 1x1 white texture as input
-					const FSortedLightSceneInfo& SortedLightInfo = SortedLights[LightIndex];
-					const FLightSceneInfo* LightSceneInfo = SortedLightInfo.LightSceneInfo;
-					RenderLight(GraphBuilder, Scene, View, SceneTextures, LightSceneInfo, nullptr, LightingChannelsTexture, 
-						false /*bRenderOverlap*/, false /*bCloudShadow*/, SortedLightInfo.bIsCompatibleWithLightFunctionAtlas);
+					for (int32 LightIndex = StandardDeferredStart; LightIndex < UnbatchedLightStart; LightIndex++)
+					{
+						// Render the light to the scene color buffer, using a 1x1 white texture as input
+						const FSortedLightSceneInfo& SortedLightInfo = SortedLights[LightIndex];
+						const FLightSceneInfo* LightSceneInfo = SortedLightInfo.LightSceneInfo;
+						RenderLight(GraphBuilder, Scene, View, SceneTextures, LightSceneInfo, nullptr, LightingChannelsTexture,
+							false /*bRenderOverlap*/, false /*bCloudShadow*/, SortedLightInfo.bIsCompatibleWithLightFunctionAtlas);
+					}
 				}
 			}
 
@@ -2784,7 +2816,9 @@ static void RenderLight(
 	const bool bCanLightUsesAtlasForUnbatchedLight,
 	TRDGUniformBufferRef<FVirtualShadowMapUniformParameters> VirtualShadowMapUniformBuffer,
 	FRDGTextureRef ShadowMaskBits,
-	int32 VirtualShadowMapId)
+	int32 VirtualShadowMapId,
+	ESubstrateTileType SubstrateTileType,
+	bool bFirstPass)
 {
 	// Ensure the light is valid for this view
 	if (!LightSceneInfo->ShouldRenderLight(View))
@@ -2793,8 +2827,12 @@ static void RenderLight(
 	}
 
 	SCOPE_CYCLE_COUNTER(STAT_DirectLightRenderingTime);
-	INC_DWORD_STAT(STAT_NumLightsUsingStandardDeferred);
 
+	if (bFirstPass)
+	{
+		INC_DWORD_STAT(STAT_NumLightsUsingStandardDeferred);
+	}
+	
 	const FLightSceneProxy* RESTRICT LightProxy = LightSceneInfo->Proxy;
 	const bool bUseIESTexture = View.Family->EngineShowFlags.TexturedLightProfiles && (LightSceneInfo->Proxy->GetIESTextureResource() != 0);
 	const bool bTransmission = LightProxy->Transmission();
@@ -2888,7 +2926,7 @@ static void RenderLight(
 			PermutationVector.Set< FDeferredLightPS::FSourceTextureDim >(LightProxy->IsRectLight() && LightProxy->HasSourceTexture());
 			PermutationVector.Set< FDeferredLightPS::FIESProfileDim >(bUseIESTexture);
 			PermutationVector.Set< FDeferredLightPS::FAnistropicMaterials >(bSupportAnisotropyPermutation && !LightSceneInfo->Proxy->IsRectLight());
-			PermutationVector.Set < FDeferredLightPS::FAtmosphereTransmittance >(false);
+			PermutationVector.Set< FDeferredLightPS::FAtmosphereTransmittance >(false);
 			PermutationVector.Set< FDeferredLightPS::FCloudTransmittance >(false);
 		}
 		else // Directional
@@ -2908,7 +2946,12 @@ static void RenderLight(
 		//   so that the geometry get correctly stencil culled on complex/simple part of the screen
 		if (Substrate::IsSubstrateEnabled())
 		{
-			if (Substrate::GetSubstrateUsesTileType(View, ESubstrateTileType::EComplexSpecial))
+			const auto ShouldRenderTileType = [&View, SubstrateTileType](ESubstrateTileType InTileType)
+				{
+					return (InTileType == SubstrateTileType || SubstrateTileType == ESubstrateTileType::ECount) && Substrate::GetSubstrateUsesTileType(View, InTileType);
+				};
+
+			if (ShouldRenderTileType(ESubstrateTileType::EComplexSpecial))
 			{
 				const ESubstrateTileType TileType = ESubstrateTileType::EComplexSpecial;
 				PermutationVector.Set<FDeferredLightPS::FSubstrateTileType>(TileType);
@@ -2916,7 +2959,7 @@ static void RenderLight(
 				TShaderMapRef< FDeferredLightPS > PixelShader(View.ShaderMap, PermutationVectorRemap);
 				InternalRenderLight(GraphBuilder, Scene, View, LightSceneInfo, PixelShader, PassParameters, TileType, TEXT("Light::StandardDeferred(ComplexSpecial)"));
 			}
-			if (Substrate::GetSubstrateUsesTileType(View, ESubstrateTileType::EComplex))
+			if (ShouldRenderTileType(ESubstrateTileType::EComplex))
 			{
 				const ESubstrateTileType TileType = ESubstrateTileType::EComplex;
 				PermutationVector.Set<FDeferredLightPS::FSubstrateTileType>(TileType);
@@ -2924,7 +2967,7 @@ static void RenderLight(
 				TShaderMapRef< FDeferredLightPS > PixelShader(View.ShaderMap, PermutationVectorRemap);
 				InternalRenderLight(GraphBuilder, Scene, View, LightSceneInfo, PixelShader, PassParameters, TileType, TEXT("Light::StandardDeferred(Complex)"));
 			}
-			if (Substrate::GetSubstrateUsesTileType(View, ESubstrateTileType::ESingle))
+			if (ShouldRenderTileType(ESubstrateTileType::ESingle))
 			{
 				const ESubstrateTileType TileType = ESubstrateTileType::ESingle;
 				PermutationVector.Set<FDeferredLightPS::FSubstrateTileType>(TileType);
@@ -2932,7 +2975,7 @@ static void RenderLight(
 				TShaderMapRef< FDeferredLightPS > PixelShader(View.ShaderMap, PermutationVectorRemap);
 				InternalRenderLight(GraphBuilder, Scene, View, LightSceneInfo, PixelShader, PassParameters, TileType, TEXT("Light::StandardDeferred(Single)"));
 			}
-			if (Substrate::GetSubstrateUsesTileType(View, ESubstrateTileType::ESimple))
+			if (ShouldRenderTileType(ESubstrateTileType::ESimple))
 			{
 				const ESubstrateTileType TileType = ESubstrateTileType::ESimple;
 				PermutationVector.Set<FDeferredLightPS::FSubstrateTileType>(TileType);
diff --git a/Engine/Source/Runtime/Renderer/Private/LightSceneInfo.h b/Engine/Source/Runtime/Renderer/Private/LightSceneInfo.h
--- a/Engine/Source/Runtime/Renderer/Private/LightSceneInfo.h
+++ b/Engine/Source/Runtime/Renderer/Private/LightSceneInfo.h
@@ -87,6 +87,8 @@ struct FSortedLightSceneInfo
 			// Currently bHandledByMegaLights is the MSB and LightType is LSB
 			/** The type of light. */
 			uint32 LightType : LightType_NumBits;
+			/** Whether the light uses IES Profile */
+			uint32 bIESProfile : 1;
 			/** Whether the light uses a light function. */
 			uint32 bLightFunction : 1;
 			/** Whether the light uses lighting channels. */
