Remove VALIDATE_UNIFORM_BUFFER_LIFETIME and VALIDATE_UNIFORM_BUFFER_LAYOUT_LIFETIME

diff --git a/Engine/Source/Runtime/RHI/Public/RHIResources.h b/Engine/Source/Runtime/RHI/Public/RHIResources.h
--- a/Engine/Source/Runtime/RHI/Public/RHIResources.h
+++ b/Engine/Source/Runtime/RHI/Public/RHIResources.h
@@ -1125,13 +1125,6 @@ public:
 // Buffers
 //
 
-// Whether to assert in cases where the layout is released before uniform buffers created with that layout
-#define VALIDATE_UNIFORM_BUFFER_LAYOUT_LIFETIME 0
-
-// Whether to assert when a uniform buffer is being deleted while still referenced by a mesh draw command
-// Enabling this requires -norhithread to work correctly since FRHIResource lifetime is managed by both the RT and RHIThread
-#define VALIDATE_UNIFORM_BUFFER_LIFETIME 0
-
 struct FRHIUniformBufferResource
 {
 	/** Byte offset to each resource in the uniform buffer memory. */
@@ -1247,26 +1240,6 @@ public:
 	, LayoutConstantBufferSize(InLayout->ConstantBufferSize)
 	{}
 
-	inline uint32 Release() const
-	{
-		const FRHIUniformBufferLayout* LocalLayout = Layout;
-
-#if VALIDATE_UNIFORM_BUFFER_LIFETIME
-		int32 LocalNumMeshCommandReferencesForDebugging = NumMeshCommandReferencesForDebugging;
-#endif
-
-		uint32 NewRefCount = FRHIResource::Release();
-
-		if (NewRefCount == 0)
-		{
-#if VALIDATE_UNIFORM_BUFFER_LIFETIME
-			check(LocalNumMeshCommandReferencesForDebugging == 0 || IsEngineExitRequested());
-#endif
-		}
-
-		return NewRefCount;
-	}
-
 	/** @return The number of bytes in the uniform buffer. */
 	uint32 GetSize() const
 	{
@@ -1276,10 +1249,6 @@ public:
 	const FRHIUniformBufferLayout& GetLayout() const { return *Layout; }
 	const FRHIUniformBufferLayout* GetLayoutPtr() const { return Layout; }
 
-#if VALIDATE_UNIFORM_BUFFER_LIFETIME
-	mutable std::atomic<int32> NumMeshCommandReferencesForDebugging = 0;
-#endif
-
 	const TArray<TRefCountPtr<FRHIResource>>& GetResourceTable() const { return ResourceTable; }
 
 protected:
diff --git a/Engine/Source/Runtime/Renderer/Private/MeshPassProcessor.cpp b/Engine/Source/Runtime/Renderer/Private/MeshPassProcessor.cpp
--- a/Engine/Source/Runtime/Renderer/Private/MeshPassProcessor.cpp
+++ b/Engine/Source/Runtime/Renderer/Private/MeshPassProcessor.cpp
@@ -838,55 +838,10 @@ void FMeshDrawShaderBindings::CopyFrom(const FMeshDrawShaderBindings& Other)
 	{
 		FPlatformMemory::Memcpy(GetData(), Other.GetData(), Size);
 	}
-
-#if VALIDATE_UNIFORM_BUFFER_LIFETIME
-	uint8* ShaderBindingDataPtr = GetData();
-
-	for (int32 ShaderBindingsIndex = 0; ShaderBindingsIndex < ShaderLayouts.Num(); ShaderBindingsIndex++)
-	{
-		FMeshDrawSingleShaderBindings SingleShaderBindings(ShaderLayouts[ShaderBindingsIndex], ShaderBindingDataPtr);
-		const FRHIUniformBuffer** RESTRICT UniformBufferBindings = SingleShaderBindings.GetUniformBufferStart();
-		const int32 NumUniformBuffers = SingleShaderBindings.ParameterMapInfo.UniformBuffers.Num();
-
-		for (int32 UniformBufferIndex = 0; UniformBufferIndex < NumUniformBuffers; UniformBufferIndex++)
-		{
-			const FRHIUniformBuffer* UniformBuffer = UniformBufferBindings[UniformBufferIndex];
-
-			if (UniformBuffer)
-			{
-				UniformBuffer->NumMeshCommandReferencesForDebugging++;
-			}
-		}
-
-		ShaderBindingDataPtr += ShaderLayouts[ShaderBindingsIndex].GetDataSizeBytes();
-	}
-#endif
 }
 
 void FMeshDrawShaderBindings::Release()
 {
-#if VALIDATE_UNIFORM_BUFFER_LIFETIME
-	uint8* ShaderBindingDataPtr = GetData();
-
-	for (int32 ShaderBindingsIndex = 0; ShaderBindingsIndex < ShaderLayouts.Num(); ShaderBindingsIndex++)
-	{
-		FMeshDrawSingleShaderBindings SingleShaderBindings(ShaderLayouts[ShaderBindingsIndex], ShaderBindingDataPtr);
-		const FRHIUniformBuffer** RESTRICT UniformBufferBindings = SingleShaderBindings.GetUniformBufferStart();
-		const int32 NumUniformBuffers = SingleShaderBindings.ParameterMapInfo.UniformBuffers.Num();
-
-		for (int32 UniformBufferIndex = 0; UniformBufferIndex < NumUniformBuffers; UniformBufferIndex++)
-		{
-			if (const FRHIUniformBuffer* UniformBuffer = UniformBufferBindings[UniformBufferIndex])
-			{
-				const int32 NumMeshCommandReferencesForDebugging = --UniformBuffer->NumMeshCommandReferencesForDebugging;
-				check(NumMeshCommandReferencesForDebugging >= 0);
-			}
-		}
-
-		ShaderBindingDataPtr += ShaderLayouts[ShaderBindingsIndex].GetDataSizeBytes();
-	}
-#endif
-
 	if (Size > sizeof(FData))
 	{
 		delete[] Data.GetHeapData();
diff --git a/Engine/Source/Runtime/Renderer/Public/MeshDrawShaderBindings.h b/Engine/Source/Runtime/Renderer/Public/MeshDrawShaderBindings.h
--- a/Engine/Source/Runtime/Renderer/Public/MeshDrawShaderBindings.h
+++ b/Engine/Source/Runtime/Renderer/Public/MeshDrawShaderBindings.h
@@ -307,15 +307,6 @@ private:
 		int32 FoundIndex = FindSortedArrayBaseIndex(MakeArrayView(ParameterMapInfo.UniformBuffers), BaseIndex);
 		if (FoundIndex >= 0)
 		{
-#if VALIDATE_UNIFORM_BUFFER_LIFETIME
-			if (const FRHIUniformBuffer* Previous = GetUniformBufferStart()[FoundIndex])
-			{
-				const int32 NumMeshCommandReferencesForDebugging = --Previous->NumMeshCommandReferencesForDebugging;
-				check(NumMeshCommandReferencesForDebugging >= 0);
-			}
-			Value->NumMeshCommandReferencesForDebugging++;
-#endif
-
 			GetUniformBufferStart()[FoundIndex] = Value;
 		}
 
