Fix issue where debug primitives with SDPG_World are not correctly occluded, partially reverting 45258202.

#rnx
#tests Ran EngineTests from jira

diff --git a/Engine/Source/Runtime/Renderer/Private/PostProcess/PostProcessCompositeDebugPrimitives.cpp b/Engine/Source/Runtime/Renderer/Private/PostProcess/PostProcessCompositeDebugPrimitives.cpp
--- a/Engine/Source/Runtime/Renderer/Private/PostProcess/PostProcessCompositeDebugPrimitives.cpp
+++ b/Engine/Source/Runtime/Renderer/Private/PostProcess/PostProcessCompositeDebugPrimitives.cpp
@@ -18,6 +18,11 @@ struct FDebugPrimitivesDrawingContext
 	const FViewInfo* DebugView = nullptr;
 	FRDGTextureRef DebugPrimitiveColor = nullptr;
 	FRDGTextureRef DebugPrimitiveDepth = nullptr;
+	FRDGTextureRef UpscaledSceneDepth = nullptr;
+	FScreenPassTexture SceneColor {};
+	FScreenPassTexture SceneDepth {};
+	FScreenPassRenderTarget OutputRT {};
+	FScreenPassRenderTarget OutputDepthRT {};
 };
 
 static FDebugPrimitivesDrawingContext InitializeDebugPrimitiveTextures(
@@ -28,7 +33,8 @@ static FDebugPrimitivesDrawingContext InitializeDebugPrimitiveTextures(
 {
 	// Setup view
 	const FIntRect ViewRect = Output.ViewRect;
-	FIntPoint Extent = Output.Texture->Desc.Extent;
+	FIntPoint OutputExtent = Output.Texture->Desc.Extent;
+	FIntPoint SceneColorExtent = Inputs.SceneColor.Texture->Desc.Extent;
 
 	FViewFamilyInfo* ViewFamily = (FViewFamilyInfo*)View.Family;
 	int NumMSAASamples = 1;
@@ -39,17 +45,22 @@ static FDebugPrimitivesDrawingContext InitializeDebugPrimitiveTextures(
 	const FViewInfo* DebugView = CreateCompositePrimitiveView(View, ViewRect, NumMSAASamples);	
 
 	// Setup textures
-	const FRDGTextureDesc ColorDesc(FRDGTextureDesc::Create2D(Extent, PF_B8G8R8A8, FClearValueBinding::Transparent, TexCreate_ShaderResource | TexCreate_RenderTargetable, 1, NumMSAASamples));
+	const FRDGTextureDesc ColorDesc(FRDGTextureDesc::Create2D(OutputExtent, PF_B8G8R8A8, FClearValueBinding::Transparent, TexCreate_ShaderResource | TexCreate_RenderTargetable, 1, NumMSAASamples));
 	FRDGTextureRef DebugPrimitiveColor = GraphBuilder.CreateTexture(ColorDesc, TEXT("Debug.PrimitivesColor"));
 
-	const FRDGTextureDesc DepthDesc(FRDGTextureDesc::Create2D(Extent, PF_DepthStencil, FClearValueBinding::DepthFar, TexCreate_ShaderResource | TexCreate_DepthStencilTargetable, 1, NumMSAASamples));
+	const FRDGTextureDesc DepthDesc(FRDGTextureDesc::Create2D(OutputExtent, PF_DepthStencil, FClearValueBinding::DepthFar, TexCreate_ShaderResource | TexCreate_DepthStencilTargetable, 1, NumMSAASamples));
 	FRDGTextureRef DebugPrimitiveDepth = GraphBuilder.CreateTexture(DepthDesc, TEXT("Debug.PrimitivesDepth"));
 
-	// Depth testing with scenedepth currently occurs after msaa resolve, so no depth copy/upscale here.
-	// This means we can perform depth sorting/testing amongst the debug primitives that are drawn, while
-	// 'disabling' depth testing of the result when compositing with the rest of the view, i.e. draw occluded lines.
-	const bool bCopyDepth = false;
-	if (bCopyDepth)
+	AddClearRenderTargetPass(GraphBuilder, DebugPrimitiveColor);
+	AddClearDepthStencilPass(GraphBuilder, DebugPrimitiveDepth);
+
+	const FRDGTextureDesc UpscaledSceneDepthDesc(FRDGTextureDesc::Create2D(SceneColorExtent, PF_DepthStencil, FClearValueBinding::DepthFar, TexCreate_ShaderResource | TexCreate_DepthStencilTargetable, 1, Inputs.SceneColor.Texture->Desc.NumSamples));
+	FRDGTextureRef UpscaledSceneDepth = GraphBuilder.CreateTexture(UpscaledSceneDepthDesc, TEXT("Debug.UpscaledWorldDepth"));
+
+	// Debug primitives with SDPG_World are depth tested against the scene depth, and immediately composed on SceneColor.
+	// SceneDepth is upscaled, to match the size of post-TSR SceneColor.
+	// Other debug primitives are only depth tested amongst themselves, and drawn to PrimitivesColor/PrimitivesDepth.
+	// Finally, PrimitivesColor is drawn to SceneColor, with occluded lines (tested against SceneDepth) drawn in a dithered style.
 	{
 		FScreenPassTexture UpsampledDepth = Inputs.SceneDepth;
 		FVector2f SceneDepthJitter = FVector2f(View.TemporalJitterPixels);
@@ -69,23 +80,23 @@ static FDebugPrimitivesDrawingContext InitializeDebugPrimitiveTextures(
 						  *DebugView,
 						  Inputs.SceneColor,
 						  UpsampledDepth,
-						  DebugPrimitiveColor,
-						  DebugPrimitiveDepth,
+						  nullptr,
+						  UpscaledSceneDepth,
 						  SceneDepthJitter,
-						  NumMSAASamples,
-						  true,
+						  UpscaledSceneDepth->Desc.NumSamples,
+						  false,
 						  Inputs.bUseMetalMSAAHDRDecode);
 	}
-	else
-	{
-		AddClearRenderTargetPass(GraphBuilder, DebugPrimitiveColor);
-		AddClearDepthStencilPass(GraphBuilder, DebugPrimitiveDepth);
-	}
 
 	FDebugPrimitivesDrawingContext Result {};
 	Result.DebugView = DebugView;
 	Result.DebugPrimitiveColor = DebugPrimitiveColor;
 	Result.DebugPrimitiveDepth = DebugPrimitiveDepth;
+	Result.UpscaledSceneDepth = UpscaledSceneDepth;
+	Result.SceneColor = Inputs.SceneColor;
+	Result.SceneDepth = Inputs.SceneDepth;
+	Result.OutputRT = Output;
+	Result.OutputDepthRT = Inputs.OverrideDepthOutput;
 
 	return Result;
 }
@@ -100,31 +111,55 @@ static bool DrawDebugPDE(
 		return false;
 	}
 
-	//Composite the debug draw elements into the scene
-	FDebugPrimitivesPassParameters* PassParameters = GraphBuilder.AllocParameters<FDebugPrimitivesPassParameters>();
-	PassParameters->View = Context.DebugView->GetShaderParameters();
-	PassParameters->RenderTargets[0] = FRenderTargetBinding(Context.DebugPrimitiveColor, ERenderTargetLoadAction::EClear);
-	PassParameters->RenderTargets[0].SetLoadAction(ERenderTargetLoadAction::ELoad);
-	PassParameters->RenderTargets.DepthStencil = FDepthStencilBinding(Context.DebugPrimitiveDepth, ERenderTargetLoadAction::ELoad, ERenderTargetLoadAction::ELoad, FExclusiveDepthStencil::DepthWrite_StencilWrite);
-
-	GraphBuilder.AddPass(
-		RDG_EVENT_NAME("DrawDebugPrimitives"),
-		PassParameters,
-		ERDGPassFlags::Raster,
-		[PassParameters, DebugView=Context.DebugView, OutputViewport](FRDGAsyncTask, FRHICommandList& RHICmdList)
+	// Draw SDPG_World elements, depth tested with SceneColor
+	if (Context.DebugView->DebugSimpleElementCollector.HasPrimitives(SDPG_World))
 	{
-		RHICmdList.SetViewport(OutputViewport.Rect.Min.X, OutputViewport.Rect.Min.Y, 0.0f, OutputViewport.Rect.Max.X, OutputViewport.Rect.Max.Y, 1.0f);
+		FDebugPrimitivesPassParameters* PassParameters = GraphBuilder.AllocParameters<FDebugPrimitivesPassParameters>();
+		PassParameters->View = Context.DebugView->GetShaderParameters();
+		PassParameters->RenderTargets[0] = FRenderTargetBinding(Context.SceneColor.Texture, ERenderTargetLoadAction::ELoad);
+		PassParameters->RenderTargets.DepthStencil = FDepthStencilBinding(Context.UpscaledSceneDepth, ERenderTargetLoadAction::ELoad, ERenderTargetLoadAction::ELoad, FExclusiveDepthStencil::DepthWrite_StencilWrite);
+
+		GraphBuilder.AddPass(
+			RDG_EVENT_NAME("DrawDebugPrimitives(SDPG_World)"),
+			PassParameters,
+			ERDGPassFlags::Raster,
+			[PassParameters, DebugView=Context.DebugView, OutputViewport](FRDGAsyncTask, FRHICommandList& RHICmdList)
+		{
+			RHICmdList.SetViewport(OutputViewport.Rect.Min.X, OutputViewport.Rect.Min.Y, 0.0f, OutputViewport.Rect.Max.X, OutputViewport.Rect.Max.Y, 1.0f);
 
-		FMeshPassProcessorRenderState DrawRenderState;
-		DrawRenderState.SetDepthStencilAccess(FExclusiveDepthStencil::DepthWrite_StencilWrite);
-		DrawRenderState.SetBlendState(TStaticBlendState<CW_RGBA, BO_Add, BF_One, BF_One, BO_Add, BF_One, BF_One>::GetRHI());
+			FMeshPassProcessorRenderState DrawRenderState;
+			DrawRenderState.SetDepthStencilAccess(FExclusiveDepthStencil::DepthWrite_StencilWrite);
+			DrawRenderState.SetBlendState(TStaticBlendState<CW_RGBA, BO_Add, BF_One, BF_One, BO_Add, BF_One, BF_One>::GetRHI());
 
-		DrawRenderState.SetDepthStencilState(TStaticDepthStencilState<true, CF_DepthNearOrEqual>::GetRHI());		
-		DebugView->DebugSimpleElementCollector.DrawBatchedElements(RHICmdList, DrawRenderState, *DebugView, EBlendModeFilter::OpaqueAndMasked, SDPG_World);
+			DrawRenderState.SetDepthStencilState(TStaticDepthStencilState<true, CF_DepthNearOrEqual>::GetRHI());		
+			DebugView->DebugSimpleElementCollector.DrawBatchedElements(RHICmdList, DrawRenderState, *DebugView, EBlendModeFilter::OpaqueAndMasked, SDPG_World);	
+		});
+	}
+
+	// Draw SDPG_Foreground elements, depth tested only with other debug primitives
+	if (Context.DebugView->DebugSimpleElementCollector.HasPrimitives(SDPG_Foreground))
+	{
+		FDebugPrimitivesPassParameters* PassParameters = GraphBuilder.AllocParameters<FDebugPrimitivesPassParameters>();
+		PassParameters->View = Context.DebugView->GetShaderParameters();
+		PassParameters->RenderTargets[0] = FRenderTargetBinding(Context.DebugPrimitiveColor, ERenderTargetLoadAction::EClear);
+		PassParameters->RenderTargets.DepthStencil = FDepthStencilBinding(Context.DebugPrimitiveDepth, ERenderTargetLoadAction::EClear, ERenderTargetLoadAction::EClear, FExclusiveDepthStencil::DepthWrite_StencilWrite);
+
+		GraphBuilder.AddPass(
+			RDG_EVENT_NAME("DrawDebugPrimitives(SDPG_Foreground)"),
+			PassParameters,
+			ERDGPassFlags::Raster,
+			[PassParameters, DebugView=Context.DebugView, OutputViewport](FRDGAsyncTask, FRHICommandList& RHICmdList)
+		{
+			RHICmdList.SetViewport(OutputViewport.Rect.Min.X, OutputViewport.Rect.Min.Y, 0.0f, OutputViewport.Rect.Max.X, OutputViewport.Rect.Max.Y, 1.0f);
 
-		DrawRenderState.SetDepthStencilState(TStaticDepthStencilState<true, CF_Always>::GetRHI());
-		DebugView->DebugSimpleElementCollector.DrawBatchedElements(RHICmdList, DrawRenderState, *DebugView, EBlendModeFilter::OpaqueAndMasked, SDPG_Foreground);		
-	});
+			FMeshPassProcessorRenderState DrawRenderState;
+			DrawRenderState.SetDepthStencilAccess(FExclusiveDepthStencil::DepthWrite_StencilWrite);
+			DrawRenderState.SetBlendState(TStaticBlendState<CW_RGBA, BO_Add, BF_One, BF_One, BO_Add, BF_One, BF_One>::GetRHI());
+
+			DrawRenderState.SetDepthStencilState(TStaticDepthStencilState<true, CF_Always>::GetRHI());
+			DebugView->DebugSimpleElementCollector.DrawBatchedElements(RHICmdList, DrawRenderState, *DebugView, EBlendModeFilter::OpaqueAndMasked, SDPG_Foreground);		
+		});
+	}
 
 	return true;
 }
@@ -133,13 +168,12 @@ static void ComposeDebugPrimitives(
 	FRDGBuilder& GraphBuilder,
 	const FViewInfo& View,
 	const FCompositePrimitiveInputs& Inputs,
-	FDebugPrimitivesDrawingContext& Context,
-	const FScreenPassRenderTarget& Output)
+	FDebugPrimitivesDrawingContext& Context)
 {
-	FScreenPassRenderTarget DepthOutput = Inputs.OverrideDepthOutput;
+	const FScreenPassRenderTarget& DepthOutput = Context.OutputDepthRT;
 	int NumMSAASamples = Context.DebugPrimitiveColor->Desc.NumSamples;
 	
-	FScreenPassTexture SceneDepth = Inputs.SceneDepth;
+	const FScreenPassTexture& SceneDepth = Context.SceneDepth;
 	FVector2f SceneDepthJitter = FVector2f(View.TemporalJitterPixels);
 
 	bool bCompositeAnyNonNullDepth = true;
@@ -149,17 +183,17 @@ static void ComposeDebugPrimitives(
 	const bool bOpaqueEditorGizmo = View.Family->EngineShowFlags.OpaqueCompositeEditorPrimitives || View.Family->EngineShowFlags.Wireframe;
 
 	FCompositePostProcessPrimitivesPS::FParameters* PassParameters = GraphBuilder.AllocParameters<FCompositePostProcessPrimitivesPS::FParameters>();
-	PassParameters->RenderTargets[0] = Output.GetRenderTargetBinding();
+	PassParameters->RenderTargets[0] = Context.OutputRT.GetRenderTargetBinding();
 
-	bool bOutputIsMSAA = Output.Texture->Desc.NumSamples > 1;
+	bool bOutputIsMSAA = Context.OutputRT.Texture->Desc.NumSamples > 1;
 	if (DepthOutput.IsValid())
 	{
 		PassParameters->RenderTargets.DepthStencil = FDepthStencilBinding(DepthOutput.Texture, ERenderTargetLoadAction::EClear, ERenderTargetLoadAction::ENoAction, FExclusiveDepthStencil::DepthWrite);
-		verify(Output.Texture->Desc.NumSamples == DepthOutput.Texture->Desc.NumSamples);
+		verify(Context.OutputRT.Texture->Desc.NumSamples == DepthOutput.Texture->Desc.NumSamples);
 	}
 
 	PassParameters->View = View.ViewUniformBuffer;
-	PassParameters->Color = GetScreenPassTextureViewportParameters(FScreenPassTextureViewport(Inputs.SceneColor));
+	PassParameters->Color = GetScreenPassTextureViewportParameters(FScreenPassTextureViewport(Context.SceneColor));
 	PassParameters->Depth = GetScreenPassTextureViewportParameters(FScreenPassTextureViewport(SceneDepth));
 
 	PassParameters->UndistortingDisplacementTexture = GSystemTextures.GetBlackDummy(GraphBuilder);
@@ -169,16 +203,16 @@ static void ComposeDebugPrimitives(
 		PassParameters->UndistortingDisplacementTexture = Inputs.LensDistortionLUT.UndistortingDisplacementTexture;
 	}
 
-	PassParameters->ColorTexture = Inputs.SceneColor.Texture;
+	PassParameters->ColorTexture = Context.SceneColor.Texture;
 	PassParameters->ColorSampler = PointClampSampler;
 	PassParameters->DepthTexture = SceneDepth.Texture;
 	PassParameters->DepthSampler = PointClampSampler;
 	PassParameters->EditorPrimitivesDepth = Context.DebugPrimitiveDepth;
 	PassParameters->EditorPrimitivesColor = Context.DebugPrimitiveColor;
 		
-	PassParameters->PassSvPositionToViewportUV = FScreenTransform::SvPositionToViewportUV(Output.ViewRect);
+	PassParameters->PassSvPositionToViewportUV = FScreenTransform::SvPositionToViewportUV(Context.OutputRT.ViewRect);
 	PassParameters->ViewportUVToColorUV = FScreenTransform::ChangeTextureBasisFromTo(
-		FScreenPassTextureViewport(Inputs.SceneColor), FScreenTransform::ETextureBasis::ViewportUV, FScreenTransform::ETextureBasis::TextureUV);
+		FScreenPassTextureViewport(Context.SceneColor), FScreenTransform::ETextureBasis::ViewportUV, FScreenTransform::ETextureBasis::TextureUV);
 	PassParameters->ViewportUVToDepthUV = FScreenTransform::ChangeTextureBasisFromTo(
 		FScreenPassTextureViewport(SceneDepth), FScreenTransform::ETextureBasis::ViewportUV, FScreenTransform::ETextureBasis::TextureUV);
 		
@@ -209,10 +243,10 @@ static void ComposeDebugPrimitives(
 	FPixelShaderUtils::AddFullscreenPass(
 		GraphBuilder,
 		View.ShaderMap,
-		RDG_EVENT_NAME("Composite %dx%d MSAA=%d", Output.ViewRect.Width(), Output.ViewRect.Height(), NumMSAASamples),
+		RDG_EVENT_NAME("Composite %dx%d MSAA=%d", Context.OutputRT.ViewRect.Width(), Context.OutputRT.ViewRect.Height(), NumMSAASamples),
 		PixelShader,
 		PassParameters,
-		Output.ViewRect,
+		Context.OutputRT.ViewRect,
 		nullptr,
 		nullptr,
 		DepthStencilState);
@@ -264,7 +298,7 @@ FScreenPassTexture AddDebugPrimitivePass(
 
 	if (bHasDrawn)
 	{
-		ComposeDebugPrimitives(GraphBuilder, View, Inputs, Context, Output);
+		ComposeDebugPrimitives(GraphBuilder, View, Inputs, Context);
 	}
 	else
 	{
diff --git a/Engine/Source/Runtime/Renderer/Private/PostProcess/PostProcessCompositePrimitivesCommon.cpp b/Engine/Source/Runtime/Renderer/Private/PostProcess/PostProcessCompositePrimitivesCommon.cpp
--- a/Engine/Source/Runtime/Renderer/Private/PostProcess/PostProcessCompositePrimitivesCommon.cpp
+++ b/Engine/Source/Runtime/Renderer/Private/PostProcess/PostProcessCompositePrimitivesCommon.cpp
@@ -276,8 +276,11 @@ void PopulateDepthPass(FRDGBuilder& GraphBuilder,
 	PassParameters->DepthTexture = InSceneDepth.Texture;
 	PassParameters->DepthSampler = PointClampSampler;
 	PassParameters->DepthTextureJitter = SceneDepthJitter;
-	const ERenderTargetLoadAction ColorLoadAction = View.IsPrimarySceneView() ? ERenderTargetLoadAction::ENoAction : ERenderTargetLoadAction::ELoad;
-	PassParameters->RenderTargets[0] = FRenderTargetBinding(OutPopColor, ColorLoadAction);
+	if (OutPopColor)
+	{
+		const ERenderTargetLoadAction ColorLoadAction = View.IsPrimarySceneView() ? ERenderTargetLoadAction::ENoAction : ERenderTargetLoadAction::ELoad;
+		PassParameters->RenderTargets[0] = FRenderTargetBinding(OutPopColor, ColorLoadAction);
+	}
 	PassParameters->RenderTargets.DepthStencil = FDepthStencilBinding(OutPopDepth, ERenderTargetLoadAction::EClear, ERenderTargetLoadAction::EClear, FExclusiveDepthStencil::DepthWrite_StencilWrite);
 
 	FPopulateCompositeDepthPS::FPermutationDomain PermutationVector;
