Revert page dilation defaults change as there are a few edge cases where it matters

#rnx

diff --git a/Engine/Source/Runtime/Renderer/Private/VirtualShadowMaps/VirtualShadowMapArray.cpp b/Engine/Source/Runtime/Renderer/Private/VirtualShadowMaps/VirtualShadowMapArray.cpp
--- a/Engine/Source/Runtime/Renderer/Private/VirtualShadowMaps/VirtualShadowMapArray.cpp
+++ b/Engine/Source/Runtime/Renderer/Private/VirtualShadowMaps/VirtualShadowMapArray.cpp
@@ -268,7 +268,7 @@ static FAutoConsoleVariableDeprecated CVarVSMShowStats_Deprecated(TEXT("r.Shadow
 
 static TAutoConsoleVariable<float> CVarPageDilationBorderSizeDirectional(
 	TEXT("r.Shadow.Virtual.PageDilationBorderSizeDirectional"),
-	0.0f,
+	0.05f,
 	TEXT("If a screen pixel falls within this fraction of a page border for directional lights, the adacent page will also be mapped.")
 	TEXT("Higher values can reduce page misses at screen edges or disocclusions, but increase total page counts."),
 	ECVF_Scalability | ECVF_RenderThreadSafe
@@ -301,7 +301,7 @@ TAutoConsoleVariable<float> CVarMaxDOFResolutionBias(
 
 static TAutoConsoleVariable<float> CVarPageDilationBorderSizeLocal(
 	TEXT("r.Shadow.Virtual.PageDilationBorderSizeLocal"),
-	0.0f,
+	0.05f,
 	TEXT("If a screen pixel falls within this fraction of a page border for local lights, the adacent page will also be mapped.")
 	TEXT("Higher values can reduce page misses at screen edges or disocclusions, but increase total page counts."),
 	ECVF_Scalability | ECVF_RenderThreadSafe
