Mobile: Disable local light stencil culling by default
#DisabledValidations NonUnityAndStaticAnalysis

diff --git a/Engine/Source/Runtime/Renderer/Private/MobileDeferredShadingPass.cpp b/Engine/Source/Runtime/Renderer/Private/MobileDeferredShadingPass.cpp
--- a/Engine/Source/Runtime/Renderer/Private/MobileDeferredShadingPass.cpp
+++ b/Engine/Source/Runtime/Renderer/Private/MobileDeferredShadingPass.cpp
@@ -37,11 +37,11 @@ static bool UseClusteredDeferredShading(const FStaticShaderPlatform Platform)
 	return GMobileUseClusteredDeferredShading != 0 && MobileForwardEnableLocalLights(Platform);
 }
 
-int32 GMobileUseLightStencilCulling = 1;
+int32 GMobileUseLightStencilCulling = 0;
 static FAutoConsoleVariableRef CVarMobileUseLightStencilCulling(
 	TEXT("r.Mobile.UseLightStencilCulling"),
 	GMobileUseLightStencilCulling,
-	TEXT("Whether to use stencil to cull local lights. 0 is off, 1 is on (default)"),
+	TEXT("Whether to use stencil to cull local lights. 0 is off (default), 1 is on"),
 	ECVF_RenderThreadSafe
 );
 
