Mobile: Swap mesh and a regular decal rendering order to match desktop. MeshDecals then regular decals
#DisabledValidations NonUnityAndStaticAnalysis

diff --git a/Engine/Source/Runtime/Renderer/Private/MobileDecalRendering.cpp b/Engine/Source/Runtime/Renderer/Private/MobileDecalRendering.cpp
--- a/Engine/Source/Runtime/Renderer/Private/MobileDecalRendering.cpp
+++ b/Engine/Source/Runtime/Renderer/Private/MobileDecalRendering.cpp
@@ -58,19 +58,20 @@ void FMobileSceneRenderer::RenderDecals(FRHICommandList& RHICmdList, FViewInfo&
 	const EDecalRenderStage DecalRenderStage = bRequiresDBufferDecals ? EDecalRenderStage::Emissive : bIsMobileDeferred ? EDecalRenderStage::MobileBeforeLighting : EDecalRenderStage::Mobile;
 	const EDecalRenderTargetMode RenderTargetMode = bIsMobileDeferred ? EDecalRenderTargetMode::SceneColorAndGBuffer : EDecalRenderTargetMode::SceneColor;
 
-	// Deferred decals
-	if (Scene->Decals.Num() > 0)
-	{
-		SCOPED_DRAW_EVENT(RHICmdList, Decals);
-		RenderDeferredDecalsMobile(RHICmdList, *Scene, View, DecalRenderStage, RenderTargetMode);
-	}
-
+	// Mesh decals
 	EMeshPass::Type DecalMeshPassType = DecalRendering::GetMeshPassType(RenderTargetMode);
 	if (HasAnyDraw(View.ParallelMeshDrawCommandPasses[DecalMeshPassType]))
 	{
 		RHICmdList.SetViewport(View.ViewRect.Min.X, View.ViewRect.Min.Y, 0, View.ViewRect.Max.X, View.ViewRect.Max.Y, 1);
 		View.ParallelMeshDrawCommandPasses[DecalMeshPassType]->Draw(RHICmdList, InstanceCullingDrawParams);
 	}
+
+	// Deferred decals
+	if (Scene->Decals.Num() > 0)
+	{
+		SCOPED_DRAW_EVENT(RHICmdList, Decals);
+		RenderDeferredDecalsMobile(RHICmdList, *Scene, View, DecalRenderStage, RenderTargetMode);
+	}
 }
 
 void RenderDeferredDecalsMobile(FRHICommandList& RHICmdList, const FScene& Scene, const FViewInfo& View, EDecalRenderStage DecalRenderStage, EDecalRenderTargetMode RenderTargetMode)
