[HWRT]
* Remove call to FlushAllocationsToClear on SBT before setting the current bindings on the SBT because the clear call could flush pending clear ops which are still used this frame. Only flush clear after add & remove of primitive scene infos on RT timeline. Pending clears can be added by async stream in requests in between the building of the dirty bindings and the flush to RHI from RDG passes causing cleared data in the SBT and triggering validation errors and potentially missing entries in the SBT.
* Also keep track of dirty records and clear dirty record when reallocated again before the clear was flushed (otherwise valid entries can be cleared)

diff --git a/Engine/Source/Runtime/Renderer/Private/DeferredShadingRenderer.cpp b/Engine/Source/Runtime/Renderer/Private/DeferredShadingRenderer.cpp
--- a/Engine/Source/Runtime/Renderer/Private/DeferredShadingRenderer.cpp
+++ b/Engine/Source/Runtime/Renderer/Private/DeferredShadingRenderer.cpp
@@ -1117,15 +1117,6 @@ void FDeferredShadingSceneRenderer::SetupRayTracingRenderingData(FRDGBuilder& Gr
 		DirtyPersistentRayTracingShaderBindings = Scene->RayTracingSBT.GetDirtyBindings(VisibleRayTracingShaderBindings, bForceAllDirty);
 	}
 
-	GraphBuilder.AddPass(RDG_EVENT_NAME("FlushRayTracingSBTAllocationsToClear"), ERDGPassFlags::NeverCull,
-		[Scene = Scene](FRDGAsyncTask, FRHICommandList& RHICmdList)
-		{
-			TRACE_CPUPROFILER_EVENT_SCOPE(FlushRayTracingSBTAllocationsToClear);
-
-			Scene->RayTracingSBT.FlushAllocationsToClear(RHICmdList);
-		}
-	);
-
 	FRDGBufferRef LumenHardwareRayTracingHitDataBuffer = nullptr;
 	if (bAnyInlineHardwareRayTracingPassEnabled)
 	{
@@ -1246,6 +1237,13 @@ void FDeferredShadingSceneRenderer::SetupRayTracingRenderingData(FRDGBuilder& Gr
 				});
 		}
 	}
+
+	GraphBuilder.AddPass(RDG_EVENT_NAME("UnlockStaticSBTAllocations"), ERDGPassFlags::NeverCull,
+		[Scene = Scene](FRDGAsyncTask, FRHICommandList& RHICmdList)
+		{
+			Scene->RayTracingSBT.ResetStaticAllocationLock();
+		}
+	);
 }
 
 #endif // RHI_RAYTRACING
diff --git a/Engine/Source/Runtime/Renderer/Private/RayTracing/RayTracing.cpp b/Engine/Source/Runtime/Renderer/Private/RayTracing/RayTracing.cpp
--- a/Engine/Source/Runtime/Renderer/Private/RayTracing/RayTracing.cpp
+++ b/Engine/Source/Runtime/Renderer/Private/RayTracing/RayTracing.cpp
@@ -2089,9 +2089,6 @@ namespace RayTracing
 					}
 					TaskData.VisibleShaderBindings = MoveTemp(DeduplicatedVisibleShaderBindingData);
 				}
-
-				// reset static allocation lock again now that the final array of dirty shader bindings is build
-				RayTracingSBT.ResetStaticAllocationLock();
 			}, AddInstancesTaskEvent);
 
 		// wait for this task here, although it could be done later in the frame
diff --git a/Engine/Source/Runtime/Renderer/Private/RayTracing/RayTracingShaderBindingTable.cpp b/Engine/Source/Runtime/Renderer/Private/RayTracing/RayTracingShaderBindingTable.cpp
--- a/Engine/Source/Runtime/Renderer/Private/RayTracing/RayTracingShaderBindingTable.cpp
+++ b/Engine/Source/Runtime/Renderer/Private/RayTracing/RayTracingShaderBindingTable.cpp
@@ -163,6 +163,8 @@ void FRayTracingShaderBindingTable::EndFrame()
 FRayTracingShaderBindingDataOneFrameArray FRayTracingShaderBindingTable::GetDirtyBindings(TConstArrayView<FRayTracingShaderBindingData> VisibleBindings, bool bForceNonPersistent)
 {
 	TRACE_CPUPROFILER_EVENT_SCOPE(FRayTracingShaderBindingTable::GetDirtyBindings);
+	
+	ensure(bStaticAllocationsLocked);
 
 	FRayTracingShaderBindingDataOneFrameArray DirtyBindings;
 
@@ -188,11 +190,15 @@ FRayTracingShaderBindingDataOneFrameArray FRayTracingShaderBindingTable::GetDirt
 			}
 			else
 			{
-				FBitReference RecordIndex = ValidPersistentRecords[VisibleBinding.SBTRecordIndex];
-				if (!RecordIndex || VisibleBinding.BindingType == ERayTracingLocalShaderBindingType::Transient)
+				check(VisibleBinding.BindingType!= ERayTracingLocalShaderBindingType::Persistent || SetPersistentRecords[VisibleBinding.SBTRecordIndex]);
+
+				FBitReference DirtyRecordIndex = DirtyPersistentRecords[VisibleBinding.SBTRecordIndex];
+				if (DirtyRecordIndex || VisibleBinding.BindingType == ERayTracingLocalShaderBindingType::Transient)
 				{
 					DirtyBindings.Add(VisibleBinding);
-					RecordIndex = true;
+					DirtyRecordIndex = false;
+
+					//UE_LOG(LogRenderer, Log, TEXT("Marking SBT Record %d NOT dirty anymore"), VisibleBinding.SBTRecordIndex);
 				}
 				else if (GValidatePersistentBindings)
 				{
@@ -226,6 +232,8 @@ void FRayTracingShaderBindingTable::AllocatePersistentShaderBindingTable(FRHICom
 	
 FRayTracingPersistentShaderBindingTableID FRayTracingShaderBindingTable::AllocatePersistentSBTID(FRHICommandListBase& RHICmdList, ERayTracingShaderBindingMode ShaderBindingMode)
 {
+	ensure(bStaticAllocationsLocked);
+
 	FRayTracingPersistentShaderBindingTableID ResultID = INDEX_NONE;
 	for (int32 Index = 0; Index < PersistentSBTs.Num(); ++Index)
 	{
@@ -253,13 +261,15 @@ FRayTracingPersistentShaderBindingTableID FRayTracingShaderBindingTable::Allocat
 	PersistentSBTs[ResultID].ShaderBindingMode = ShaderBindingMode;
 	AllocatePersistentShaderBindingTable(RHICmdList, ResultID, ShaderBindingMode);
 
-	ValidPersistentRecords.Init(false, PersistentSBTInitializer.NumGeometrySegments * PersistentSBTInitializer.NumShaderSlotsPerGeometrySegment);
+	DirtyPersistentRecords.Init(true, PersistentSBTInitializer.NumGeometrySegments * PersistentSBTInitializer.NumShaderSlotsPerGeometrySegment);
 
 	return ResultID;
 }
 
 void FRayTracingShaderBindingTable::CheckPersistentRHI(FRHICommandListBase& RHICmdList, uint32 LocalBindingDataSize)
 {		
+	ensure(bStaticAllocationsLocked);
+
 	uint32 NumPersistentStaticGeometrySegments = FMath::Max(uint32(GMinStaticGeometrySegments), FMath::RoundUpToPowerOfTwo(GetMaxAllocatedStaticSegmentCount()));
 	uint32 NumPersistentDynamicGeometrySegments = FMath::Max(uint32(GMinDynamicGeometrySegments), FMath::RoundUpToPowerOfTwo(MaxNumDynamicGeometrySegments));
 	uint32 NumMissShaderSlotsAligned = FMath::RoundUpToPowerOfTwo(NumMissShaderSlots);
@@ -330,7 +340,7 @@ void FRayTracingShaderBindingTable::CheckPersistentRHI(FRHICommandListBase& RHIC
 				AllocatePersistentShaderBindingTable(RHICmdList, Index, PersistentSBTs[Index].ShaderBindingMode);
 			}
 		}
-		ValidPersistentRecords.Init(false, NewSBTInitializer.NumGeometrySegments * NewSBTInitializer.NumShaderSlotsPerGeometrySegment);
+		DirtyPersistentRecords.Init(true, NewSBTInitializer.NumGeometrySegments * NewSBTInitializer.NumShaderSlotsPerGeometrySegment);
 	}
 }
 
@@ -361,20 +371,43 @@ uint32 FRayTracingShaderBindingTable::GetNumGeometrySegments() const
 
 uint32 FRayTracingShaderBindingTable::GetMaxAllocatedStaticSegmentCount() const
 {
-	//ensure(bStaticAllocationsLocked);
+	ensure(bStaticAllocationsLocked);
 	return StaticRangeAllocator.GetMaxSize() / NumShaderSlotsPerGeometrySegment;
 }
 
 void FRayTracingShaderBindingTable::MarkDirty(FRayTracingSBTAllocation* SBTAllocation)
 {
 	int32 MaxRecordIndex = SBTAllocation->BaseRecordIndex + SBTAllocation->NumRecords;
-	if (ValidPersistentRecords.Num() < MaxRecordIndex)
+	if (DirtyPersistentRecords.Num() < MaxRecordIndex)
+	{
+		DirtyPersistentRecords.SetNum(FMath::RoundUpToPowerOfTwo(MaxRecordIndex), true);
+	}
+	for (uint32 Index = 0; Index < SBTAllocation->NumRecords; ++Index)
+	{
+		DirtyPersistentRecords[SBTAllocation->BaseRecordIndex + Index] = true;
+		//UE_LOG(LogRenderer, Log, TEXT("Marking SBT Record %d dirty"), SBTAllocation->BaseRecordIndex + Index);
+	}
+}
+
+void FRayTracingShaderBindingTable::MarkSet(FRayTracingSBTAllocation* SBTAllocation, bool bValue)
+{
+	int32 MaxRecordIndex = SBTAllocation->BaseRecordIndex + SBTAllocation->NumRecords;
+	if (SetPersistentRecords.Num() < MaxRecordIndex)
 	{
-		ValidPersistentRecords.SetNum(FMath::RoundUpToPowerOfTwo(MaxRecordIndex), false);
+		SetPersistentRecords.SetNum(FMath::RoundUpToPowerOfTwo(MaxRecordIndex), false);
+		PersistentRecordsToClear.SetNum(FMath::RoundUpToPowerOfTwo(MaxRecordIndex), false);
 	}
 	for (uint32 Index = 0; Index < SBTAllocation->NumRecords; ++Index)
 	{
-		ValidPersistentRecords[SBTAllocation->BaseRecordIndex + Index] = false;
+		uint32 RecordIndex = SBTAllocation->BaseRecordIndex + Index;
+		check(SetPersistentRecords[RecordIndex] != bValue);
+		SetPersistentRecords[RecordIndex] = bValue;
+
+		// Mark record to clear and validate that it's not marked for clear when we are trying to clear it out (double clear?)
+		check(bValue || !PersistentRecordsToClear[RecordIndex]);
+		PersistentRecordsToClear[RecordIndex] = !bValue;
+
+		//UE_LOG(LogRenderer, Log, TEXT("Update set SBT record data for RecordIndex %d with new value: %d"), RecordIndex, bValue ? 1 : 0);
 	}
 }
 
@@ -396,6 +429,7 @@ FRayTracingSBTAllocation* FRayTracingShaderBindingTable::AllocateStaticRangeInte
 	Allocation->InitStatic(AllocatedLayers, BaseIndex, RecordsPerLayer, RecordCount, Geometry, Flags);
 
 	MarkDirty(Allocation);
+	MarkSet(Allocation, true);
 
 	AllocatedStaticSegmentCount += SegmentCount * LayersCount;
 	
@@ -404,6 +438,7 @@ FRayTracingSBTAllocation* FRayTracingShaderBindingTable::AllocateStaticRangeInte
 
 FRayTracingSBTAllocation* FRayTracingShaderBindingTable::AllocateStaticRange(uint32 SegmentCount, const FRHIRayTracingGeometry* Geometry, FRayTracingCachedMeshCommandFlags Flags)
 {
+	ensure(!bStaticAllocationsLocked);
 	check(Geometry != nullptr);
 
 	// No allocation if we are not rendering decals and all segments are decals
@@ -474,10 +509,13 @@ void FRayTracingShaderBindingTable::FreeStaticRange(const FRayTracingSBTAllocati
 
 		if (RefAllocation->RefCount == 0)
 		{
-			StaticRangeAllocator.Free(RefAllocation->Allocation->BaseRecordIndex, RefAllocation->Allocation->NumRecords);
 			AllocatedStaticSegmentCount -= (RefAllocation->Allocation->NumRecords / NumShaderSlotsPerGeometrySegment);
-		
-			PersistentAllocationsToClear.Add(RefAllocation->Allocation);
+
+			// Already mark unused but only defer the free on the free range allocator until the clear op is actually processed
+			// so the free can happen free threaded while the static allocations are still locked			
+			MarkSet(RefAllocation->Allocation, false);
+			PendingStaticAllocationsToFree.Add(RefAllocation->Allocation);
+
 			TrackedAllocationMap.Remove(Key);
 		}
 	}
@@ -491,6 +529,8 @@ void FRayTracingShaderBindingTable::FreeStaticRange(const FRayTracingSBTAllocati
 
 FRayTracingSBTAllocation* FRayTracingShaderBindingTable::AllocateDynamicRange(ERayTracingShaderBindingLayerMask AllocatedLayers, uint32 SegmentCount)
 {	
+	ensure(bStaticAllocationsLocked);
+
 	// Don't need lock right now because all dynamic allocation are allocated linearly on the same thread
 	// So the FreeDynamicAllocationPool can't be shared with the static allocations right now because those would require a lock
 	FRayTracingSBTAllocation* Allocation = FreeDynamicAllocationPool.IsEmpty() ? nullptr : FreeDynamicAllocationPool.Pop(EAllowShrinking::No);
@@ -519,7 +559,7 @@ FRayTracingSBTAllocation* FRayTracingShaderBindingTable::AllocateDynamicRange(ER
 
 bool FRayTracingShaderBindingTable::IsDirty(uint32 RecordIndex) const
 {
-	return !ValidPersistentRecords[RecordIndex];
+	return DirtyPersistentRecords[RecordIndex];
 }
 
 bool FRayTracingShaderBindingTable::IsPersistent() const
@@ -534,34 +574,9 @@ void FRayTracingShaderBindingTable::FlushAllocationsToClear(FRHICommandList& RHI
 	// Don't clear outside of current allocated range - could have allocated and cleared ranges before calling CheckPersistentRHI
 	// and then the RHI SBT doesn't have those new ranges allocated yet
 	uint32 MaxNumValidRecords = PersistentSBTInitializer.NumGeometrySegments * PersistentSBTInitializer.NumShaderSlotsPerGeometrySegment;
-	check(MaxNumValidRecords <= (uint32)ValidPersistentRecords.Num());
-
-	// Build bit array if records to clear (make sure we don't clear the same record twice)
-	TBitArray<> PersistentRecordsToClear;
-	PersistentRecordsToClear.Init(false, MaxNumValidRecords);
+	check(MaxNumValidRecords <= (uint32)DirtyPersistentRecords.Num());
 	
-	uint32 TotalRecordsToClear = 0;
-	for (const FRayTracingSBTAllocation* Allocation : PersistentAllocationsToClear)
-	{
-		for (uint32 RangeIndex = 0; RangeIndex < Allocation->NumRecords; ++RangeIndex)
-		{
-			uint32 RecordIndex = Allocation->BaseRecordIndex + RangeIndex;
-			if (RecordIndex < MaxNumValidRecords)
-			{
-				FBitReference RecordFlag = PersistentRecordsToClear[RecordIndex];
-				if (!RecordFlag)
-				{
-					TotalRecordsToClear++;
-					RecordFlag = true;
-				}
-			}
-		}
-
-		// TODO: pool static allocations as well?
-		delete Allocation;
-	}
-	PersistentAllocationsToClear.Reset(PersistentAllocationsToClear.Num());
-
+	uint32 TotalRecordsToClear = PersistentRecordsToClear.CountSetBits();
 	if (TotalRecordsToClear == 0)
 	{
 		return;
@@ -599,6 +614,21 @@ void FRayTracingShaderBindingTable::FlushAllocationsToClear(FRHICommandList& RHI
 				bCopyDataToInlineStorage);
 		}
 	}
+
+	PersistentRecordsToClear.Init(false, PersistentRecordsToClear.Num());
+
+	// Only makr ranges as free now on the SBT after the clear ops have been processed by the RHI SBT
+	// to make sure we don't reuse records which haven't been freed yet (due to out of order freeing of records)
+	uint32 FreedRecords = 0;
+	for (FRayTracingSBTAllocation* AllocationToFree : PendingStaticAllocationsToFree)
+	{
+		FreedRecords += AllocationToFree->NumRecords;
+		StaticRangeAllocator.Free(AllocationToFree->BaseRecordIndex, AllocationToFree->NumRecords);
+		// TODO: pool static allocations as well?
+		delete AllocationToFree;
+	}
+	check(FreedRecords == TotalRecordsToClear);
+	PendingStaticAllocationsToFree.Reset();
 }
 
 #endif // RHI_RAYTRACING
\ No newline at end of file
diff --git a/Engine/Source/Runtime/Renderer/Private/RayTracing/RayTracingShaderBindingTable.h b/Engine/Source/Runtime/Renderer/Private/RayTracing/RayTracingShaderBindingTable.h
--- a/Engine/Source/Runtime/Renderer/Private/RayTracing/RayTracingShaderBindingTable.h
+++ b/Engine/Source/Runtime/Renderer/Private/RayTracing/RayTracingShaderBindingTable.h
@@ -232,6 +232,7 @@ private:
 	 * Mark all records used by the SBTAllocation as dirty again
 	 */
 	void MarkDirty(FRayTracingSBTAllocation* SBTAllocation);
+	void MarkSet(FRayTracingSBTAllocation* SBTAllocation, bool bValue);
 
 	struct FAllocationKey
 	{
@@ -279,9 +280,12 @@ private:
 
 	FRayTracingShaderBindingTableInitializer PersistentSBTInitializer;	 //< Shared initializer used for all persistent SBTs - so they can all be versioned together
 	TArray<FPersistentSBTData> PersistentSBTs;							 //< All currently allocated persistent SBTs (FRayTracingPersistentShaderBindingTableID contains index into this array so can be sparse)
-	TBitArray<> ValidPersistentRecords;									 //< BitArray containing which bits are valid in the cached persistent SBTs (used to build the dirty set)
-	TArray<FRayTracingSBTAllocation*> PersistentAllocationsToClear;		 //< Persistent SBT allocations about which the cached persistent SBTs need to be notified (clear internal cached record data)
+	TBitArray<> SetPersistentRecords;									 //< BitArray containing which bits are valid in the cached persistent SBTs (use for validation)
+	TBitArray<> DirtyPersistentRecords;									 //< BitArray containing which bits are dirty in the cached persistent SBTs (used during dirty binding collection from currently visible bindings)
+	TBitArray<> PersistentRecordsToClear;								 //< BitArray containing which bits need to be cleared in the persisten SBTs
 	
+	TArray<FRayTracingSBTAllocation*> PendingStaticAllocationsToFree;	 //< All pending allocations to free during next flush
+
 	FCriticalSection StaticAllocationCS;								 //< Critical section used to access all static allocation data
 	bool bStaticAllocationsLocked = false;								 //< Static allocations are not allowed when this bool is set (used for validation)
 	FSpanAllocator StaticRangeAllocator;								 //< Range allocator to find free static record ranges
