Inline Contact Shadows for mobile deferred multi-pass
ScreenSpace Contact Shadows remain for mobile forward and mobile deferred single pass(ShadowMaskTexture required)
SUPPORT_CONTACT_SHADOWS set to 1 in Mobile Light Pass when mobile deferred multi-pass is active(MobileAllowFramebufferFetch = false)
ShadowedBits now setup for mobile like on SM6(0 = NoShadowCast, 1 = Static Shadow Cast, 2 = Dynamic Shadow Cast, 3 = Both Static and Dynamic Shadow Cast)

diff --git a/Engine/Shaders/Private/DeferredLightingCommon.ush b/Engine/Shaders/Private/DeferredLightingCommon.ush
--- a/Engine/Shaders/Private/DeferredLightingCommon.ush
+++ b/Engine/Shaders/Private/DeferredLightingCommon.ush
@@ -73,6 +73,11 @@ float ShadowRayCast(
 	bOutHitCastContactShadow = false;
 	if (HitDistance > 0.0)
 	{
+#if SHADING_PATH_MOBILE
+		const uint2 HitPixelCoord = uint2(HitUV * View.BufferSizeAndInvSize.xy);
+		FGBufferData SampleGBuffer = MobileFetchAndDecodeGBuffer(HitUV, HitPixelCoord);
+		bOutHitCastContactShadow = CastContactShadow(SampleGBuffer);
+#else
 #if SUBSTRATE_GBUFFER_FORMAT==1
 #if SUBSTRATE_RAYCAST_ENABLED && SUBSTRATE_MATERIALCONTAINER_IS_VIEWRESOURCE
 		uint2 PixelPos = View.ViewRectMin.xy + View.ViewSizeAndInvSize.xy * HitUV;
@@ -85,6 +90,7 @@ float ShadowRayCast(
 #else
 		FGBufferData SampleGBuffer = GetGBufferData(HitUV);
 		bOutHitCastContactShadow = CastContactShadow(SampleGBuffer);
+#endif
 #endif
 	}
 
diff --git a/Engine/Shaders/Private/MobileBasePassPixelShader.usf b/Engine/Shaders/Private/MobileBasePassPixelShader.usf
--- a/Engine/Shaders/Private/MobileBasePassPixelShader.usf
+++ b/Engine/Shaders/Private/MobileBasePassPixelShader.usf
@@ -887,7 +887,7 @@ void Main(
 
 		FLightAccumulator DirectLighting = (FLightAccumulator)0;
 		LightAccumulator_AddSplit(DirectLighting, DiffuseColor, 0.0f, DiffuseColor, 1.0f, false);
-
+		GBuffer.PerObjectGBufferData = GetPrimitive_PerObjectGBufferData(MaterialParameters.PrimitiveId);
 		GBuffer.IndirectIrradiance = IndirectIrradiance * Export.IndirectIrradianceScale;
 		MobileEncodeGBuffer(GBuffer, OutGBufferA, OutGBufferB, OutGBufferC, OutGBufferD);
 
@@ -1074,6 +1074,7 @@ void Main(
 
 #if MOBILE_USE_GBUFFER
 	GBuffer.IndirectIrradiance = IndirectIrradiance;
+	GBuffer.PerObjectGBufferData = GetPrimitive_PerObjectGBufferData(MaterialParameters.PrimitiveId);
 	MobileEncodeGBuffer(GBuffer, OutGBufferA, OutGBufferB, OutGBufferC, OutGBufferD);
 #else
 
@@ -1083,7 +1084,7 @@ void Main(
 	float DirectionalLightShadow = 1.0f;
 	// Directional light
 #if !MATERIAL_SHADINGMODEL_UNLIT
-	AccumulateDirectionalLighting(GBuffer, MaterialParameters.WorldPosition_CamRelative, CameraVector, MaterialParameters.ScreenPosition, DynamicShadowFactors, DirectionalLightShadow, DirectLighting);
+	AccumulateDirectionalLighting(GBuffer, MaterialParameters.WorldPosition_CamRelative, CameraVector, MaterialParameters.ScreenPosition, DynamicShadowFactors, 0, DirectionalLightShadow, DirectLighting);
 #endif
 
 	float2 LocalPosition = SvPosition.xy - ResolvedView.ViewRectMin.xy;
diff --git a/Engine/Shaders/Private/MobileDeferredShading.usf b/Engine/Shaders/Private/MobileDeferredShading.usf
--- a/Engine/Shaders/Private/MobileDeferredShading.usf
+++ b/Engine/Shaders/Private/MobileDeferredShading.usf
@@ -22,6 +22,8 @@
 #include "/Engine/Generated/Material.ush"
 #endif
 
+// In-line Contact Shadows only in mobile deferred multi-pass
+#define SUPPORT_CONTACT_SHADOWS !ALLOW_FRAMEBUFFER_FETCH
 #include "MobileLightingCommon.ush"
 #if MATERIAL_SHADER
 #include "LightFunctionCommon.ush"
@@ -208,7 +210,8 @@ void MobileDirectionalLightPS(
 	// Directional light
 	half4 DynamicShadowFactors = half4(ScreenSpaceShadowMaskSample.rgb, 1.0);
 	float DirectionalLightShadow = 1.0f;
-	AccumulateDirectionalLighting(GBuffer, TranslatedWorldPosition, CameraVector, ScreenPosition, DynamicShadowFactors, DirectionalLightShadow, DirectLighting);
+	const float Dither = InterleavedGradientNoise(SvPosition.xy, View.StateFrameIndexMod8);
+	AccumulateDirectionalLighting(GBuffer, TranslatedWorldPosition, CameraVector, ScreenPosition, DynamicShadowFactors, Dither, DirectionalLightShadow, DirectLighting);
 	// LightFunction should only affect direct lighting result
 	DirectLighting.TotalLight *= ComputeLightFunctionMultiplier(TranslatedWorldPosition);
 
@@ -250,6 +253,7 @@ Texture2D		LocalLightShadowTexture;
 SamplerState	LocalLightShadowSampler;
 float4			LocalLightShadowBufferSize;
 #endif
+uint ShadowedBits;
 
 FDeferredLightData InitDeferredLightFromLightParameters(FLightShaderParameters In, uint InLightType)
 {
@@ -283,8 +287,16 @@ FDeferredLightData InitDeferredLightFromLightParameters(FLightShaderParameters I
 	Out.bRectLight				= InLightType == LIGHT_TYPE_RECT;
 	Out.HairTransmittance		= InitHairTransmittanceData();
 
+// In-line Contact Shadows only in mobile deferred multi-pass
+#if SUPPORT_CONTACT_SHADOWS
+	Out.ContactShadowLength					= abs(ContactShadows.x);
+	Out.ContactShadowLengthInWS				= ContactShadows.x < 0.0f;
+	Out.ContactShadowCastingIntensity		= ContactShadows.y;
+	Out.ContactShadowNonCastingIntensity	= ContactShadows.z;
+#endif
+	Out.ShadowedBits						= ShadowedBits;
+
 #if SUPPORT_SPOTLIGHTS_SHADOW
-	Out.ShadowedBits			= 1;
 	Out.ShadowMapChannelMask	= 1.f;
 #endif
 	return Out;
@@ -377,7 +389,8 @@ void MobileRadialLightPS(
 #else
 	half4 LightAttenuation = half4(1, 1, Shadow, Shadow);
 	float SurfaceShadow = 0;
-	DirectLighting = AccumulateDynamicLighting(TranslatedWorldPosition, CameraVector, GBuffer, 1, LightData, LightAttenuation, 0, uint2(0, 0), SurfaceShadow);
+	const float Dither = InterleavedGradientNoise(SVPos.xy, View.StateFrameIndexMod8);
+	DirectLighting = AccumulateDynamicLighting(TranslatedWorldPosition, CameraVector, GBuffer, 1, LightData, LightAttenuation, Dither, uint2(0, 0), SurfaceShadow);
 	half3 Color = DirectLighting.TotalLight * ComputeLightFunctionMultiplier(TranslatedWorldPosition);
 #endif
 
diff --git a/Engine/Shaders/Private/MobileLightingCommon.ush b/Engine/Shaders/Private/MobileLightingCommon.ush
--- a/Engine/Shaders/Private/MobileLightingCommon.ush
+++ b/Engine/Shaders/Private/MobileLightingCommon.ush
@@ -2,7 +2,9 @@
 
 #pragma once
 
+#ifndef SUPPORT_CONTACT_SHADOWS
 #define SUPPORT_CONTACT_SHADOWS 0
+#endif
 
 #include "ShadingModels.ush"
 #include "ReflectionEnvironmentShared.ush"
@@ -53,6 +55,10 @@
 #define REFLECTION_COMPOSITE_HAS_SPHERE_CAPTURES 1
 #define REFLECTION_COMPOSITE_HAS_BOX_CAPTURES 1
 #include "ReflectionEnvironmentComposite.ush"
+// In-line Contact Shadows only in mobile deferred multi-pass
+#if SUPPORT_CONTACT_SHADOWS
+float3 ContactShadows;
+#endif
 
 half GetSurfaceShadow(FGBufferData GBuffer, FShadowTerms ShadowTerms)
 {
@@ -301,14 +307,23 @@ FDeferredLightData GetDirectionalLightData(float4 ScreenPosition, inout half Out
 	LightData.bRadialLight = false;
 	LightData.SpecularScale = MobileDirectionalLight.DirectionalLightDistanceFadeMADAndSpecularScale.z;
 	LightData.DiffuseScale = MobileDirectionalLight.DirectionalLightDistanceFadeMADAndSpecularScale.w;
-	LightData.ShadowedBits = 1;
+
+// In-line Contact Shadows only in mobile deferred multi-pass
+#if SUPPORT_CONTACT_SHADOWS
+	LightData.ContactShadowLength = abs(ContactShadows.x);
+	LightData.ContactShadowLengthInWS = ContactShadows.x < 0.0f;
+	LightData.ContactShadowCastingIntensity = ContactShadows.y;
+	LightData.ContactShadowNonCastingIntensity = ContactShadows.z;
+#endif
+
+	LightData.ShadowedBits = MobileDirectionalLight.ShadowedBits;
 	LightData.HairTransmittance = InitHairTransmittanceData();
 	LightData.ShadowMapChannelMask = UnpackShadowMapChannelMask(MobileDirectionalLight.DirectionalLightShadowMapChannelMask);
 
 	return LightData;
 }
 
-void AccumulateDirectionalLighting(FGBufferData GBuffer, float3 TranslatedWorldPosition, half3 CameraVector, float4 ScreenPosition, in half4 DynamicShadowFactors, inout float OutDirectionalLightShadow, inout FLightAccumulator DirectLighting)
+void AccumulateDirectionalLighting(FGBufferData GBuffer, float3 TranslatedWorldPosition, half3 CameraVector, float4 ScreenPosition, in half4 DynamicShadowFactors, in float Dither, inout float OutDirectionalLightShadow, inout FLightAccumulator DirectLighting)
 {
 	half DynamicShadowing = 1.0f;
 	FDeferredLightData LightData = GetDirectionalLightData(ScreenPosition, DynamicShadowing);
@@ -329,7 +344,7 @@ void AccumulateDirectionalLighting(FGBufferData GBuffer, float3 TranslatedWorldP
 #if TRANSLUCENCY_SH_LIGHTING
 	FLightAccumulator NewLighting = GetTranslucencySHLighting(GBuffer, LightData, TranslatedWorldPosition);
 #else
-	FLightAccumulator NewLighting = AccumulateDynamicLighting(TranslatedWorldPosition, CameraVector, GBuffer, 1, LightData, LightAttenuation, 0, uint2(0, 0), OutDirectionalLightShadow);
+	FLightAccumulator NewLighting = AccumulateDynamicLighting(TranslatedWorldPosition, CameraVector, GBuffer, 1, LightData, LightAttenuation, Dither, uint2(0, 0), OutDirectionalLightShadow);
 #endif
 	DirectLighting = LightAccumulator_Add(DirectLighting, NewLighting);
 }
diff --git a/Engine/Source/Runtime/Engine/Private/SceneManagement.cpp b/Engine/Source/Runtime/Engine/Private/SceneManagement.cpp
--- a/Engine/Source/Runtime/Engine/Private/SceneManagement.cpp
+++ b/Engine/Source/Runtime/Engine/Private/SceneManagement.cpp
@@ -1149,6 +1149,7 @@ FMobileDirectionalLightShaderParameters::FMobileDirectionalLightShaderParameters
 		DirectionalLightScreenToShadow[i].SetIdentity();
 		DirectionalLightShadowDistances[i] = FLT_MAX; // Unused cascades should compare > all scene depths
 	}
+	ShadowedBits = 0;
 }
 
 FViewUniformShaderParameters::FViewUniformShaderParameters()
diff --git a/Engine/Source/Runtime/Engine/Public/SceneManagement.h b/Engine/Source/Runtime/Engine/Public/SceneManagement.h
--- a/Engine/Source/Runtime/Engine/Public/SceneManagement.h
+++ b/Engine/Source/Runtime/Engine/Public/SceneManagement.h
@@ -1180,6 +1180,7 @@ BEGIN_SHADER_PARAMETER_STRUCT(FMobileMovableLocalLightShadowParameters,ENGINE_AP
 	SHADER_PARAMETER(FVector4f, LocalLightShadowBufferSize)
 	SHADER_PARAMETER_TEXTURE(Texture2D, LocalLightShadowTexture)
 	SHADER_PARAMETER_SAMPLER(SamplerState, LocalLightShadowSampler)
+	SHADER_PARAMETER(uint32, ShadowedBits)
 END_GLOBAL_SHADER_PARAMETER_STRUCT()
 
 /**
diff --git a/Engine/Source/Runtime/Engine/Public/SceneView.h b/Engine/Source/Runtime/Engine/Public/SceneView.h
--- a/Engine/Source/Runtime/Engine/Public/SceneView.h
+++ b/Engine/Source/Runtime/Engine/Public/SceneView.h
@@ -823,6 +823,7 @@ BEGIN_GLOBAL_SHADER_PARAMETER_STRUCT_WITH_CONSTRUCTOR(FMobileDirectionalLightSha
 	SHADER_PARAMETER(uint32, DirectionalLightShadowMapChannelMask)
 	SHADER_PARAMETER_TEXTURE(Texture2D, DirectionalLightShadowTexture)
 	SHADER_PARAMETER_SAMPLER(SamplerState, DirectionalLightShadowSampler)
+	SHADER_PARAMETER(uint32, ShadowedBits)
 END_GLOBAL_SHADER_PARAMETER_STRUCT()
 
 //////////////////////////////////////////////////////////////////////////
diff --git a/Engine/Source/Runtime/Renderer/Private/LightRendering.cpp b/Engine/Source/Runtime/Renderer/Private/LightRendering.cpp
--- a/Engine/Source/Runtime/Renderer/Private/LightRendering.cpp
+++ b/Engine/Source/Runtime/Renderer/Private/LightRendering.cpp
@@ -610,6 +610,16 @@ static void RenderLight(
 	ESubstrateTileType SubstrateTileType = ESubstrateTileType::ECount,
 	bool bFirstPass = true);
 
+uint32 GetShadowedBits(const FSceneView& View, const FLightSceneInfo& LightSceneInfo, bool bUseLightFunctionAtlas)
+{
+	uint32 ShadowedBits = 0;
+	const bool bDynamicShadows = View.Family->EngineShowFlags.DynamicShadows && GetShadowQuality() > 0;
+	const bool bHasLightFunction = LightSceneInfo.Proxy->GetLightFunctionMaterial() != NULL && (!bUseLightFunctionAtlas || !LightSceneInfo.Proxy->HasValidLightFunctionAtlasSlot());
+	ShadowedBits = LightSceneInfo.Proxy->CastsStaticShadow() || bHasLightFunction ? 1 : 0;
+	ShadowedBits |= LightSceneInfo.Proxy->CastsDynamicShadow() && View.Family->EngineShowFlags.DynamicShadows ? 3 : 0;
+	return ShadowedBits;
+}
+
 FDeferredLightUniformStruct GetDeferredLightParameters(const FSceneView& View, const FLightSceneInfo& LightSceneInfo, bool bUseLightFunctionAtlas, uint32 LightFlags)
 {
 	FDeferredLightUniformStruct Out;
@@ -643,11 +653,7 @@ FDeferredLightUniformStruct GetDeferredLightParameters(const FSceneView& View, c
 			ShadowMapChannel == 3 ? 1 : 0);
 	}
 
-	const bool bDynamicShadows = View.Family->EngineShowFlags.DynamicShadows && GetShadowQuality() > 0;
-	const bool bHasLightFunction = LightSceneInfo.Proxy->GetLightFunctionMaterial() != NULL && (!bUseLightFunctionAtlas || !LightSceneInfo.Proxy->HasValidLightFunctionAtlasSlot());
-	Out.ShadowedBits = LightSceneInfo.Proxy->CastsStaticShadow() || bHasLightFunction ? 1 : 0;
-	Out.ShadowedBits |= LightSceneInfo.Proxy->CastsDynamicShadow() && View.Family->EngineShowFlags.DynamicShadows ? 3 : 0;
-
+	Out.ShadowedBits = GetShadowedBits(View, LightSceneInfo, bUseLightFunctionAtlas);
 	Out.VolumetricScatteringIntensity = LightSceneInfo.Proxy->GetVolumetricScatteringIntensity();
 
 	static auto* ContactShadowsCVar = IConsoleManager::Get().FindTConsoleVariableDataInt(TEXT("r.ContactShadows"));
diff --git a/Engine/Source/Runtime/Renderer/Private/LightRendering.h b/Engine/Source/Runtime/Renderer/Private/LightRendering.h
--- a/Engine/Source/Runtime/Renderer/Private/LightRendering.h
+++ b/Engine/Source/Runtime/Renderer/Private/LightRendering.h
@@ -36,6 +36,7 @@ extern uint32 GetShadowQuality();
 
 extern float GetLightFadeFactor(const FSceneView& View, const FLightSceneProxy* Proxy);
 
+extern uint32 GetShadowedBits(const FSceneView& View, const FLightSceneInfo& LightSceneInfo, bool bUseLightFunctionAtlas = false);
 extern FDeferredLightUniformStruct GetDeferredLightParameters(const FSceneView& View, const FLightSceneInfo& LightSceneInfo, bool bUseLightFunctionAtlas=false, uint32 LightFlags=0);
 
 inline void SetDeferredLightParameters(
diff --git a/Engine/Source/Runtime/Renderer/Private/LightSceneInfo.cpp b/Engine/Source/Runtime/Renderer/Private/LightSceneInfo.cpp
--- a/Engine/Source/Runtime/Renderer/Private/LightSceneInfo.cpp
+++ b/Engine/Source/Runtime/Renderer/Private/LightSceneInfo.cpp
@@ -352,6 +352,8 @@ bool FLightSceneInfo::SetupMobileMovableLocalLightShadowParameters(const FViewIn
 		MobileMovableLocalLightShadowParameters.LocalLightShadowSampler = TStaticSamplerState<SF_Point, AM_Clamp, AM_Clamp, AM_Clamp>::GetRHI();
 	}
 
+	MobileMovableLocalLightShadowParameters.ShadowedBits = GetShadowedBits(View, *this);
+
 	return bShouldCastShadow;
 }
 
diff --git a/Engine/Source/Runtime/Renderer/Private/MobileBasePassRendering.cpp b/Engine/Source/Runtime/Renderer/Private/MobileBasePassRendering.cpp
--- a/Engine/Source/Runtime/Renderer/Private/MobileBasePassRendering.cpp
+++ b/Engine/Source/Runtime/Renderer/Private/MobileBasePassRendering.cpp
@@ -465,6 +465,8 @@ void SetupMobileDirectionalLightUniformParameters(
 				}
 			}
 		}
+
+		Params.ShadowedBits = GetShadowedBits(SceneView, *Light);
 	}
 }
 
diff --git a/Engine/Source/Runtime/Renderer/Private/MobileDeferredShadingPass.cpp b/Engine/Source/Runtime/Renderer/Private/MobileDeferredShadingPass.cpp
--- a/Engine/Source/Runtime/Renderer/Private/MobileDeferredShadingPass.cpp
+++ b/Engine/Source/Runtime/Renderer/Private/MobileDeferredShadingPass.cpp
@@ -59,6 +59,14 @@ BEGIN_SHADER_PARAMETER_STRUCT(FMobileDeferredPassParameters, )
 	RENDER_TARGET_BINDING_SLOTS()
 END_SHADER_PARAMETER_STRUCT()
 
+BEGIN_SHADER_PARAMETER_STRUCT(FMobileDeferredCommonParameters, )
+	SHADER_PARAMETER(FMatrix44f, TranslatedWorldToLight)
+	SHADER_PARAMETER(FVector4f, LightFunctionParameters)
+	SHADER_PARAMETER(FVector2f, LightFunctionParameters2)
+	SHADER_PARAMETER(FVector3f, CameraRelativeLightPosition)
+	SHADER_PARAMETER(FVector3f, ContactShadows)
+END_SHADER_PARAMETER_STRUCT()
+
 class FMobileDirectionalLightFunctionPS : public FMaterialShader
 {
 	DECLARE_SHADER_TYPE(FMobileDirectionalLightFunctionPS, Material);
@@ -90,10 +98,7 @@ class FMobileDirectionalLightFunctionPS : public FMaterialShader
 	BEGIN_SHADER_PARAMETER_STRUCT(FParameters, )
 		SHADER_PARAMETER_STRUCT_REF(FMobileDirectionalLightShaderParameters, MobileDirectionalLight)
 		SHADER_PARAMETER_STRUCT_REF(FMobileReflectionCaptureShaderData, MobileReflectionCaptureData)
-		SHADER_PARAMETER(FMatrix44f, TranslatedWorldToLight)
-		SHADER_PARAMETER(FVector4f, LightFunctionParameters)
-		SHADER_PARAMETER(FVector2f, LightFunctionParameters2)
-		SHADER_PARAMETER(FVector3f, CameraRelativeLightPosition)
+		SHADER_PARAMETER_STRUCT_INCLUDE(FMobileDeferredCommonParameters, MobileDeferredCommonParameters)
 		SHADER_PARAMETER_STRUCT_INCLUDE(FDFAOUpsampleParameters, DFAOUpsampleParameters)
 		SHADER_PARAMETER_STRUCT_INCLUDE(FSkyDiffuseLightingParameters, SkyDiffuseLighting)
 	END_SHADER_PARAMETER_STRUCT()
@@ -223,10 +228,7 @@ public:
 	BEGIN_SHADER_PARAMETER_STRUCT(FParameters, )
 		SHADER_PARAMETER_STRUCT(FLightShaderParameters, Light)
 		SHADER_PARAMETER_STRUCT_INCLUDE(FMobileMovableLocalLightShadowParameters, MobileMovableLocalLightShadow)
-		SHADER_PARAMETER(FMatrix44f, TranslatedWorldToLight)
-		SHADER_PARAMETER(FVector4f, LightFunctionParameters)
-		SHADER_PARAMETER(FVector2f, LightFunctionParameters2)
-		SHADER_PARAMETER(FVector3f, CameraRelativeLightPosition)
+		SHADER_PARAMETER_STRUCT_INCLUDE(FMobileDeferredCommonParameters, MobileDeferredCommonParameters)
 	END_SHADER_PARAMETER_STRUCT()
 
 	static bool ShouldCompilePermutation(const FMaterialShaderPermutationParameters& Parameters)
@@ -612,6 +614,27 @@ static void SetDirectionalLightDepthStencilState(FGraphicsPipelineStateInitializ
 	}
 }
 
+extern void GetLightContactShadowParameters(const FLightSceneProxy* Proxy, float& OutLength, bool& bOutLengthInWS, float& OutCastingIntensity, float& OutNonCastingIntensity);
+void SetupContactShadowParameters(const FLightSceneProxy* Proxy, const FViewInfo& View, FVector3f& OutContactShadowsParameter)
+{
+	float ContactShadowLength = 0.f;
+	float ContactShadowCastingIntensity = 1.f;
+	float ContactShadowNonCastingIntensity = 0.f;
+	static auto* ContactShadowsCVar = IConsoleManager::Get().FindTConsoleVariableDataInt(TEXT("r.ContactShadows"));
+	const bool bContactShadows = ContactShadowsCVar && ContactShadowsCVar->GetValueOnRenderThread() != 0;
+	// Only do in-line contact shadows on deferred multi pass
+	if (bContactShadows && View.Family->EngineShowFlags.ContactShadows && !MobileAllowFramebufferFetch(View.GetShaderPlatform()))
+	{
+		bool bContactShadowLengthInWS = false;
+		GetLightContactShadowParameters(Proxy, ContactShadowLength, bContactShadowLengthInWS, ContactShadowCastingIntensity, ContactShadowNonCastingIntensity);
+		// Sign indicates if contact shadow length is in world space or screen space.
+		ContactShadowLength *= bContactShadowLengthInWS ? -1.0f : 1.0f;
+	}
+	OutContactShadowsParameter.X = ContactShadowLength;
+	OutContactShadowsParameter.Y = ContactShadowCastingIntensity;
+	OutContactShadowsParameter.Z = ContactShadowNonCastingIntensity;
+}
+
 static void RenderDirectionalLight(
 	FRHICommandList& RHICmdList, 
 	const FScene& Scene, 
@@ -642,16 +665,16 @@ static void RenderDirectionalLight(
 	FMobileDirectionalLightFunctionPS::FParameters PassParameters;
 	PassParameters.MobileDirectionalLight = Scene.UniformBuffers.MobileDirectionalLightUniformBuffers[LightingChannel + 1];
 	PassParameters.MobileReflectionCaptureData = GetShaderBinding(View.MobileReflectionCaptureUniformBuffer);
-	PassParameters.LightFunctionParameters = FVector4f(1.0f, 1.0f, 0.0f, 0.0f);
-	PassParameters.CameraRelativeLightPosition = GetCamRelativeLightPosition(View.ViewMatrices, DirectionalLight);
+	PassParameters.MobileDeferredCommonParameters.LightFunctionParameters = FVector4f(1.0f, 1.0f, 0.0f, 0.0f);
+	PassParameters.MobileDeferredCommonParameters.CameraRelativeLightPosition = GetCamRelativeLightPosition(View.ViewMatrices, DirectionalLight);
 
 	{
-		PassParameters.LightFunctionParameters2 = FVector2f(DirectionalLight.Proxy->GetLightFunctionFadeDistance(), DirectionalLight.Proxy->GetLightFunctionDisabledBrightness());
+		PassParameters.MobileDeferredCommonParameters.LightFunctionParameters2 = FVector2f(DirectionalLight.Proxy->GetLightFunctionFadeDistance(), DirectionalLight.Proxy->GetLightFunctionDisabledBrightness());
 		const FVector Scale = DirectionalLight.Proxy->GetLightFunctionScale();
 		// Switch x and z so that z of the user specified scale affects the distance along the light direction
 		const FVector InverseScale = FVector(1.f / Scale.Z, 1.f / Scale.Y, 1.f / Scale.X);
 		const FMatrix WorldToLight = DirectionalLight.Proxy->GetWorldToLight() * FScaleMatrix(FVector(InverseScale));
-		PassParameters.TranslatedWorldToLight = FMatrix44f(FTranslationMatrix(-View.ViewMatrices.GetPreViewTranslation()) * WorldToLight);
+		PassParameters.MobileDeferredCommonParameters.TranslatedWorldToLight = FMatrix44f(FTranslationMatrix(-View.ViewMatrices.GetPreViewTranslation()) * WorldToLight);
 	}
 
 	// DFAO
@@ -708,6 +731,7 @@ static void RenderDirectionalLight(
 	const bool bInsetShadows = (CVarMobileSupportInsetShadows != nullptr && CVarMobileSupportInsetShadows->GetInt() != 0);
 	const bool bMobileUsesShadowMaskTexture = (MobileUsesShadowMaskTexture(View.GetShaderPlatform()) || bInsetShadows);
 
+	SetupContactShadowParameters(DirectionalLight.Proxy, View, PassParameters.MobileDeferredCommonParameters.ContactShadows);
 	const uint8 LightingChannelStencilValue = GetLightingChannelStencilValue(LightingChannel);
 	
 	for (int32 PassIndex = 0; PassIndex < NumPasses; ++PassIndex)
@@ -983,15 +1007,16 @@ static void RenderLocalLight(
 
 	PassParameters.Light = GetDeferredLightParameters(View, LightSceneInfo).LightParameters;
 	const float TanOuterAngle = bIsSpotLight ? FMath::Tan(LightSceneInfo.Proxy->GetOuterConeAngle()) : 1.0f;
-	PassParameters.LightFunctionParameters = FVector4f(TanOuterAngle, 1.0f /*ShadowFadeFraction*/, bIsSpotLight ? 1.0f : 0.0f, bIsPointLight ? 1.0f : 0.0f);
-	PassParameters.LightFunctionParameters2 = FVector2f(LightSceneInfo.Proxy->GetLightFunctionFadeDistance(), LightSceneInfo.Proxy->GetLightFunctionDisabledBrightness());
+	PassParameters.MobileDeferredCommonParameters.LightFunctionParameters = FVector4f(TanOuterAngle, 1.0f /*ShadowFadeFraction*/, bIsSpotLight ? 1.0f : 0.0f, bIsPointLight ? 1.0f : 0.0f);
+	PassParameters.MobileDeferredCommonParameters.LightFunctionParameters2 = FVector2f(LightSceneInfo.Proxy->GetLightFunctionFadeDistance(), LightSceneInfo.Proxy->GetLightFunctionDisabledBrightness());
 	const FVector Scale = LightSceneInfo.Proxy->GetLightFunctionScale();
 	// Switch x and z so that z of the user specified scale affects the distance along the light direction
 	const FVector InverseScale = FVector(1.f / Scale.Z, 1.f / Scale.Y, 1.f / Scale.X);
 	const FMatrix WorldToLight = LightSceneInfo.Proxy->GetWorldToLight() * FScaleMatrix(FVector(InverseScale));
-	PassParameters.TranslatedWorldToLight = FMatrix44f(FTranslationMatrix(-View.ViewMatrices.GetPreViewTranslation()) * WorldToLight);
-	PassParameters.CameraRelativeLightPosition = GetCamRelativeLightPosition(View.ViewMatrices, LightSceneInfo);
-
+	PassParameters.MobileDeferredCommonParameters.TranslatedWorldToLight = FMatrix44f(FTranslationMatrix(-View.ViewMatrices.GetPreViewTranslation()) * WorldToLight);
+	PassParameters.MobileDeferredCommonParameters.CameraRelativeLightPosition = GetCamRelativeLightPosition(View.ViewMatrices, LightSceneInfo);
+	
+	SetupContactShadowParameters(LightSceneInfo.Proxy, View, PassParameters.MobileDeferredCommonParameters.ContactShadows);
 	// Do two passes, first masking DefautLit, second masking all other shading models
 	const bool bOnlyDefaultLitInView = IsOnlyDefaultLitShadingModel(View.ShadingModelMaskInView);
 	int32 NumPasses = !bOnlyDefaultLitInView && MobileUsesGBufferCustomData(Scene.GetShaderPlatform()) ? 2 : 1;
@@ -1096,6 +1121,7 @@ static void RenderSimpleLights(
 		// Render light
 		FMobileRadialLightFunctionPS::FParameters ParametersPS;
 		ParametersPS.Light = GetSimpleDeferredLightParameters(View, SimpleLight, SimpleLightPerViewData).LightParameters;
+		ParametersPS.MobileDeferredCommonParameters.ContactShadows = FVector3f(0.f /*ContactShadowLength*/, 1.f /*ContactShadowCastingIntensity*/, 0.f /*ContactShadowNonCastingIntensity*/);
 		SetShaderParameters(RHICmdList, VertexShader, VertexShader.GetVertexShader(), ParametersVS);
 		SetShaderParametersMixedPS(RHICmdList, PixelShader, ParametersPS, View, DefaultMaterial.MaterialProxy, *DefaultMaterial.Material);
 
diff --git a/Engine/Source/Runtime/Renderer/Private/ShadowRendering.cpp b/Engine/Source/Runtime/Renderer/Private/ShadowRendering.cpp
--- a/Engine/Source/Runtime/Renderer/Private/ShadowRendering.cpp
+++ b/Engine/Source/Runtime/Renderer/Private/ShadowRendering.cpp
@@ -2461,7 +2461,10 @@ FRDGTextureRef FMobileSceneRenderer::RenderMobileShadowProjections(FRDGBuilder&
 		}
 		RenderCapsuleDirectShadows(GraphBuilder, *LightSceneInfo, ScreenShadowMaskTexture, VisibleLightInfo.CapsuleShadowsToProject, bProjectingForForwardShading);
 
-		if (LightSceneProxy->GetLightType() == LightType_Directional && LightSceneInfo->GetDynamicShadowMapChannel() != -1)
+		if (LightSceneProxy->GetLightType() == LightType_Directional && 
+			LightSceneInfo->GetDynamicShadowMapChannel() != -1 &&
+			// For deferred multi-pass contact shadows are done inline
+			!(bRequiresMultiPass && bDeferredShading))
 		{
 			// Dynamic shadows are projected into channels of the light attenuation texture based on their assigned DynamicShadowMapChannel
 			// Only render screen space shadows if light is assigned to a valid DynamicShadowMapChannel
