Mobile: Use two-channel ShadowMask instead of four-channel when possible
#DisabledValidations NonUnityAndStaticAnalysis

diff --git a/Engine/Shaders/Private/MobileLightingCommon.ush b/Engine/Shaders/Private/MobileLightingCommon.ush
--- a/Engine/Shaders/Private/MobileLightingCommon.ush
+++ b/Engine/Shaders/Private/MobileLightingCommon.ush
@@ -276,9 +276,15 @@ half4 SampleScreenSpaceShadowMask(float2 BufferUV, uint EyeIndex)
 {
 #if ENABLE_SHADOWMASKTEXTURE 
 	#if MOBILE_MULTI_VIEW
-		return Texture2DArraySample(MobileBasePass.ScreenSpaceShadowMaskTextureArray, MobileBasePass.ScreenSpaceShadowMaskSampler, float3(BufferUV, EyeIndex));
+		half4 Sample = Texture2DArraySample(MobileBasePass.ScreenSpaceShadowMaskTextureArray, MobileBasePass.ScreenSpaceShadowMaskSampler, float3(BufferUV, EyeIndex));
 	#else
-		return Texture2DSample(MobileBasePass.ScreenSpaceShadowMaskTexture, MobileBasePass.ScreenSpaceShadowMaskSampler, BufferUV);
+		half4 Sample = Texture2DSample(MobileBasePass.ScreenSpaceShadowMaskTexture, MobileBasePass.ScreenSpaceShadowMaskSampler, BufferUV);
+	#endif
+
+	#if ALLOW_STATIC_LIGHTING
+		return Sample;
+	#else	
+		return half4(Sample.r, 1.0, 1.0, Sample.g); // PG_RG8 with a AO in g
 	#endif
 #endif
 	return 1.0;
diff --git a/Engine/Source/Runtime/Renderer/Private/PostProcess/PostProcessAmbientOcclusionMobile.cpp b/Engine/Source/Runtime/Renderer/Private/PostProcess/PostProcessAmbientOcclusionMobile.cpp
--- a/Engine/Source/Runtime/Renderer/Private/PostProcess/PostProcessAmbientOcclusionMobile.cpp
+++ b/Engine/Source/Runtime/Renderer/Private/PostProcess/PostProcessAmbientOcclusionMobile.cpp
@@ -91,6 +91,15 @@ static bool IsTwoPassShaderType(EMobileAmbientOcclusionShaderType ShaderType)
 	return ShaderType != EMobileAmbientOcclusionShaderType::ComputeShader;
 }
 
+static FRHIBlendState* GetShadowMaskBlendState(FRDGTextureRef ScreenSpaceShadowMaskTexture)
+{
+	if (ScreenSpaceShadowMaskTexture->Desc.Format == PF_R8G8)
+	{
+		return TStaticBlendState<CW_GREEN>::GetRHI();
+	}
+	return TStaticBlendState<CW_ALPHA>::GetRHI();
+}
+
 static TAutoConsoleVariable<int32> CVarMobileAmbientOcclusionDepthBoundsTest(
 	TEXT("r.Mobile.AmbientOcclusionDepthBoundsTest"),
 	1,
@@ -974,7 +983,7 @@ static void AddMobileAmbientOcclusionPass(
 		MoveTemp(EventName),
 		PassParameters,
 		ERDGPassFlags::Raster,
-		[&View, OutputViewport, InputViewport, VertexShader, PixelShader, PassParameters, bHalfResolution, bDepthBoundsTestEnabled, DepthFar](FRHICommandListImmediate& RHICmdList)
+		[&View, Output, OutputViewport, InputViewport, VertexShader, PixelShader, PassParameters, bHalfResolution, bDepthBoundsTestEnabled, DepthFar](FRHICommandListImmediate& RHICmdList)
 		{
 			const FIntRect InputRect = InputViewport.Rect;
 			const FIntPoint InputSize = InputViewport.Extent;
@@ -992,8 +1001,8 @@ static void AddMobileAmbientOcclusionPass(
 			}
 			else
 			{
-				// Writing to a ScreenSpaceShadowMask alpha channel
-				GraphicsPSOInit.BlendState = TStaticBlendState<CW_ALPHA>::GetRHI();
+				// Writing to a ScreenSpaceShadowMask
+				GraphicsPSOInit.BlendState = GetShadowMaskBlendState(Output.Texture);
 			}
 			GraphicsPSOInit.RasterizerState = TStaticRasterizerState<FM_Solid, CM_None>::GetRHI();
 			GraphicsPSOInit.DepthStencilState = TStaticDepthStencilState<false, CF_Always>::GetRHI();
@@ -1066,17 +1075,17 @@ static void AddMobileAmbientOcclusionUpsamplePass(
 			RDG_EVENT_NAME("DepthBounds ClearQuad(%s)", Output.Texture->Name),
 			ClearParameters,
 			ERDGPassFlags::Raster,
-			[OutputViewport, DepthFar](FRHICommandListImmediate& RHICmdList)
+			[Output, OutputViewport, DepthFar](FRHICommandListImmediate& RHICmdList)
 			{
 				RHICmdList.SetViewport(OutputViewport.Rect.Min.X, OutputViewport.Rect.Min.Y, 0.0f, OutputViewport.Rect.Max.X, OutputViewport.Rect.Max.Y, 1.0f);
 
 				// We must clear all pixels that won't be touched by AO shader.
 				FClearQuadCallbacks Callbacks;
-				Callbacks.PSOModifier = [](FGraphicsPipelineStateInitializer& PSOInitializer)
+				Callbacks.PSOModifier = [Output](FGraphicsPipelineStateInitializer& PSOInitializer)
 				{
 					PSOInitializer.bDepthBounds = true;
-					// Writing to a ScreenSpaceShadowMask alpha channel
-					PSOInitializer.BlendState = TStaticBlendState<CW_ALPHA>::GetRHI();
+					// Writing to a ScreenSpaceShadowMask
+					PSOInitializer.BlendState = GetShadowMaskBlendState(Output.Texture);
 				};
 				Callbacks.PreClear = [DepthFar](FRHICommandList& InRHICmdList)
 				{
@@ -1120,7 +1129,7 @@ static void AddMobileAmbientOcclusionUpsamplePass(
 		RDG_EVENT_NAME("UpsamplePS Quality %d", UpsampleQuality),
 		PassParameters,
 		ERDGPassFlags::Raster,
-		[PassParameters, &View, OutputViewport, VertexShader, PixelShader, bDepthBoundsTestEnabled, DepthFar](FRHICommandList& RHICmdList)
+		[PassParameters, &View, Output, OutputViewport, VertexShader, PixelShader, bDepthBoundsTestEnabled, DepthFar](FRHICommandList& RHICmdList)
 		{
 			RHICmdList.SetViewport(OutputViewport.Rect.Min.X, OutputViewport.Rect.Min.Y, 0.0f, OutputViewport.Rect.Max.X, OutputViewport.Rect.Max.Y, 1.0f);
 
@@ -1128,8 +1137,8 @@ static void AddMobileAmbientOcclusionUpsamplePass(
 			FPixelShaderUtils::InitFullscreenPipelineState(RHICmdList, View.ShaderMap, PixelShader, GraphicsPSOInit);
 			GraphicsPSOInit.bDepthBounds = bDepthBoundsTestEnabled;
 
-			// Writing to a ScreenSpaceShadowMask alpha channel
-			GraphicsPSOInit.BlendState = TStaticBlendState<CW_ALPHA>::GetRHI();
+			// Writing to a ScreenSpaceShadowMask
+			GraphicsPSOInit.BlendState = GetShadowMaskBlendState(Output.Texture);
 
 			SetGraphicsPipelineState(RHICmdList, GraphicsPSOInit, 0);
 			SetShaderParameters(RHICmdList, PixelShader, PixelShader.GetPixelShader(), *PassParameters);
diff --git a/Engine/Source/Runtime/Renderer/Private/ShadowRendering.cpp b/Engine/Source/Runtime/Renderer/Private/ShadowRendering.cpp
--- a/Engine/Source/Runtime/Renderer/Private/ShadowRendering.cpp
+++ b/Engine/Source/Runtime/Renderer/Private/ShadowRendering.cpp
@@ -2429,7 +2429,14 @@ FRDGTextureRef FMobileSceneRenderer::RenderMobileShadowProjections(FRDGBuilder&
 {
 	RDG_RHI_EVENT_SCOPE(GraphBuilder, RenderMobileShadowProjections);
 	
-	FRDGTextureDesc TargetDesc = FRDGTextureDesc::CreateRenderTargetTextureDesc(Extent, PF_B8G8R8A8, FClearValueBinding::White, TexCreate_RenderTargetable | TexCreate_ShaderResource, bIsMobileMultiView);
+	EPixelFormat ShadowMaskPixelFormat = PF_B8G8R8A8;
+	if (!IsStaticLightingAllowed())
+	{
+		// With pre-computed lighting disabled all lights use single channel. And second one for AO
+		ShadowMaskPixelFormat = PF_R8G8;
+	}
+	
+	FRDGTextureDesc TargetDesc = FRDGTextureDesc::CreateRenderTargetTextureDesc(Extent, ShadowMaskPixelFormat, FClearValueBinding::White, TexCreate_RenderTargetable | TexCreate_ShaderResource, bIsMobileMultiView);
 	FRDGTextureRef ScreenShadowMaskTexture = GraphBuilder.CreateTexture(TargetDesc, TEXT("ScreenSpaceShadowMaskTextureMobile"));
 	
 	FRDGTextureClearInfo TextureClearInfo;
