Adding LightCounts stats for the mobilerenderer

diff --git a/Engine/Source/Runtime/Renderer/Private/MobileShadingRenderer.cpp b/Engine/Source/Runtime/Renderer/Private/MobileShadingRenderer.cpp
--- a/Engine/Source/Runtime/Renderer/Private/MobileShadingRenderer.cpp
+++ b/Engine/Source/Runtime/Renderer/Private/MobileShadingRenderer.cpp
@@ -84,6 +84,8 @@
 #include "SingleLayerWaterRendering.h"
 #include "PostProcess/PostProcessCompositeDebugPrimitives.h"
 
+CSV_DECLARE_CATEGORY_EXTERN(LightCount);
+
 uint32 GetShadowQuality();
 
 static TAutoConsoleVariable<int32> CVarMobileForceDepthResolve(
@@ -1248,6 +1250,10 @@ void FMobileSceneRenderer::Render(FRDGBuilder& GraphBuilder, const FSceneRenderU
 			SetDummyForwardLightUniformBufferOnViews(GraphBuilder, ShaderPlatform, Views);
 		}
 
+		CSV_CUSTOM_STAT(LightCount, All, float(SortedLightSet.SortedLights.Num()), ECsvCustomStatOp::Set);
+		CSV_CUSTOM_STAT(LightCount, Batched, float(SortedLightSet.UnbatchedLightStart), ECsvCustomStatOp::Set);
+		CSV_CUSTOM_STAT(LightCount, Unbatched, float(SortedLightSet.SortedLights.Num()) - float(SortedLightSet.UnbatchedLightStart), ECsvCustomStatOp::Set);
+
 		// Generate the Sky/Atmosphere look up tables
 		const bool bShouldRenderSkyAtmosphere = ShouldRenderSkyAtmosphere(Scene, ViewFamily.EngineShowFlags);
 		if (bShouldRenderSkyAtmosphere)
