Fix typo in RenderDebugViewMode function signature.

#rnx

diff --git a/Engine/Source/Runtime/Renderer/Private/DebugViewModeRendering.cpp b/Engine/Source/Runtime/Renderer/Private/DebugViewModeRendering.cpp
--- a/Engine/Source/Runtime/Renderer/Private/DebugViewModeRendering.cpp
+++ b/Engine/Source/Runtime/Renderer/Private/DebugViewModeRendering.cpp
@@ -708,7 +708,7 @@ void RenderDebugViewMode(
 	TArrayView<FViewInfo> Views,
 	FRDGTextureRef QuadOverdrawTexture,
 	const FRenderTargetBindingSlots& RenderTargets,
-	const FSceneUniformBuffer& SceneUniformBuffer,
+	FSceneUniformBuffer& SceneUniformBuffer,
 	FVirtualShadowMapArray& VirtualShadowMapArray)
 {}
 
