[TSR] Update Nanite foliage logic so it can enable translucent clipped depth path.

diff --git a/Engine/Source/Runtime/Renderer/Private/PostProcess/TemporalSuperResolution.cpp b/Engine/Source/Runtime/Renderer/Private/PostProcess/TemporalSuperResolution.cpp
--- a/Engine/Source/Runtime/Renderer/Private/PostProcess/TemporalSuperResolution.cpp
+++ b/Engine/Source/Runtime/Renderer/Private/PostProcess/TemporalSuperResolution.cpp
@@ -1600,7 +1600,7 @@ bool NeedTSRAntiFlickeringPass(const FViewInfo& View)
 bool NeedTSRThinGeometryDetectionPass(const FViewInfo& View)
 {
 	return GetMainTAAPassConfig(View) == EMainTAAPassConfig::TSR && 
-		CVarTSRThinGeometryDetection.GetValueOnRenderThread() != 0 &&
+		ShouldEnableThinGeometryDetection() &&
 		View.Family->EngineShowFlags.PostProcessing;
 }
 
@@ -1848,7 +1848,7 @@ FDefaultTemporalUpscaler::FOutputs AddTemporalSuperResolutionPasses(
 
 	FTSRUpdateHistoryCS::EQuality UpdateHistoryQuality = FTSRUpdateHistoryCS::EQuality(FMath::Clamp(PassConfig.HistoryUpdateQuality, 0, int32(FTSRUpdateHistoryCS::EQuality::MAX) - 1));
 
-	bool bIsSeperateTranslucyTexturesValid = PassInputs.PostDOFTranslucencyResources.IsValid();
+	bool bIsSeparateTranslucyTexturesValid = PassInputs.PostDOFTranslucencyResources.IsValid();
 
 	EPixelFormat ColorFormat = bSupportsAlpha ? PF_FloatRGBA : PF_FloatR11G11B10;
 	EPixelFormat HistoryColorFormat = (PassConfig.HistoryR11G11B10 != 0 && !bSupportsAlpha) ? PF_FloatR11G11B10 : PF_FloatRGBA;
