[Partial backout] - CL45411259
#rnx
Original CL Desc
-----------------------------------------------------------------
[VSM] Implement non-nanite ShadowCasters view, with Nanite ShadowCasters composited on top.
DebugViewShaderMode is now also set lazily, to fix a bug where it was initialized incorrectly before EngineShowFlagOverride was applied.

diff --git a/Engine/Shaders/Private/DebugViewModeCommon.ush b/Engine/Shaders/Private/DebugViewModeCommon.ush
--- a/Engine/Shaders/Private/DebugViewModeCommon.ush
+++ b/Engine/Shaders/Private/DebugViewModeCommon.ush
@@ -19,8 +19,6 @@ struct FDebugVSToPS
 	float3 TangentToWorld0 : TEXCOORD3;
 	float3 TangentToWorld1 : TEXCOORD4;
 	float3 TangentToWorld2 : TEXCOORD5;
-	nointerpolation uint PrimitiveId : PRIMITIVE_ID;
-	nointerpolation uint InstanceId : INSTANCE_ID;
 	FStereoVSOutput StereoOutput;
 };
 
@@ -33,8 +31,6 @@ struct FDebugPSIn
 	float3 TangentToWorld0 : TEXCOORD3;
 	float3 TangentToWorld1 : TEXCOORD4;
 	float3 TangentToWorld2 : TEXCOORD5;
-	nointerpolation uint PrimitiveId : PRIMITIVE_ID;
-	nointerpolation uint InstanceId : INSTANCE_ID;
 	FStereoPSInput StereoInput;
 #if (OUTPUT_QUAD_OVERDRAW)
 	uint SvPrimitiveID : SV_PrimitiveID;
diff --git a/Engine/Shaders/Private/DebugViewModePixelShader.usf b/Engine/Shaders/Private/DebugViewModePixelShader.usf
--- a/Engine/Shaders/Private/DebugViewModePixelShader.usf
+++ b/Engine/Shaders/Private/DebugViewModePixelShader.usf
@@ -15,15 +15,12 @@
 #undef EYE_ADAPTATION_DISABLED
 #define EYE_ADAPTATION_DISABLED 1
 
-#define NANITE_USE_VIEW_UNIFORM_BUFFER 1
-
 #include "Common.ush"
 #include "SHCommon.ush"
 #if OUTPUT_QUAD_OVERDRAW
 #include "QuadOverdraw.ush"
 #endif
 #include "VirtualTextureCommon.ush"
-#include "VirtualShadowMaps/VirtualShadowMapVisualize.ush"
 
 #define UNDEF_ACCURACY 1024
 
@@ -516,18 +513,8 @@ void Main(
 {
 	StereoSetupPS(DebugInputs.StereoInput);
 
-	uint PrimitiveId = DebugInputs.PrimitiveId;
-	uint InstanceId  = DebugInputs.InstanceId;
-
-	FNaniteView NaniteView						= GetNaniteView(ResolvedView.GPUSceneViewId);
-	FInstanceSceneData InstanceData				= GetInstanceSceneData(InstanceId);
-	FInstanceViewData InstanceViewData			= GetInstanceViewData(InstanceId, ResolvedView.GPUSceneViewId);
-	FInstanceDynamicData InstanceDynamicData	= CalculateInstanceDynamicData(NaniteView, InstanceData, true);
-	FPrimitiveSceneData PrimitiveData			= GetPrimitiveData(PrimitiveId);
-
 	// If we see pink something went wrong.
 	float4 VisualizeColor = float4(1, 0, 1, 1);
-	bool bApplyDebugShading = false;
 
 	BRANCH
 	if (VisualizeMode == DVSM_ShaderComplexity ||
@@ -563,22 +550,6 @@ void Main(
 	{
 		VisualizeColor = VisualizeGPUSkinCache(DebugInputs, bIsFrontFace);
 	}
-	else if (VisualizeMode == DVSM_ShadowCasters)
-	{
-		VisualizeColor = float4(GetShadowCasterColor(NaniteView, InstanceData, InstanceViewData, InstanceDynamicData, PrimitiveData), 1);
-		bApplyDebugShading = true;
-	}
-
-	if (bApplyDebugShading)
-	{
-		// Apply N dot V shading using face normal
-		float3 TranslatedWorldPosition = SvPositionToResolvedTranslatedWorld(DebugInputs.SvPosition);
-		float3 N = normalize(cross(ddx(TranslatedWorldPosition), ddy(TranslatedWorldPosition)));
-
-		float3 V = normalize( TranslatedWorldPosition - View.TranslatedWorldCameraOrigin );
-
-		VisualizeColor *= abs( dot( N, V ) ) * 0.7 + 0.3;
-	}
 
 	OutColor = RETURN_COLOR(VisualizeColor);	
 }
diff --git a/Engine/Shaders/Private/DebugViewModeVertexShader.usf b/Engine/Shaders/Private/DebugViewModeVertexShader.usf
--- a/Engine/Shaders/Private/DebugViewModeVertexShader.usf
+++ b/Engine/Shaders/Private/DebugViewModeVertexShader.usf
@@ -89,10 +89,6 @@ void Main(
 #if NUM_MATERIAL_TEXCOORDS_VERTEX > 3
 		Output.TexCoord23.zw = VertexParameters.TexCoords[3];
 #endif
-
-	FSceneDataIntermediates SceneData = VertexParameters.SceneData;
-	Output.PrimitiveId = SceneData.PrimitiveId;
-	Output.InstanceId = SceneData.InstanceId;
 }
 
 #endif // VERTEXSHADER
diff --git a/Engine/Source/Runtime/Renderer/Private/BasePassRendering.cpp b/Engine/Source/Runtime/Renderer/Private/BasePassRendering.cpp
--- a/Engine/Source/Runtime/Renderer/Private/BasePassRendering.cpp
+++ b/Engine/Source/Runtime/Renderer/Private/BasePassRendering.cpp
@@ -1585,8 +1585,7 @@ void FDeferredShadingSceneRenderer::RenderBasePassInternal(
 		else if (ViewFamily.UseDebugViewPS())
 		{
 			// Override the base pass with one of the debug view shader mode (see EDebugViewShaderMode) if required.
-			FSceneUniformBuffer& SceneUniformBuffer = Renderer.GetSceneUniforms();
-			RenderDebugViewMode(GraphBuilder, InViews, SceneTextures.QuadOverdraw, BasePassRenderTargets, SceneUniformBuffer, Renderer.VirtualShadowMapArray);
+			RenderDebugViewMode(GraphBuilder, InViews, SceneTextures.QuadOverdraw, BasePassRenderTargets);
 		}
 	}
 	else
diff --git a/Engine/Source/Runtime/Renderer/Private/DebugViewModeRendering.cpp b/Engine/Source/Runtime/Renderer/Private/DebugViewModeRendering.cpp
--- a/Engine/Source/Runtime/Renderer/Private/DebugViewModeRendering.cpp
+++ b/Engine/Source/Runtime/Renderer/Private/DebugViewModeRendering.cpp
@@ -209,20 +209,12 @@ IMPLEMENT_MATERIAL_SHADER_TYPE(,FDebugViewModeVS,TEXT("/Engine/Private/DebugView
 
 BEGIN_SHADER_PARAMETER_STRUCT(FDebugViewModePassParameters, )
 	SHADER_PARAMETER_STRUCT_INCLUDE(FViewShaderParameters, View)
-	SHADER_PARAMETER_RDG_UNIFORM_BUFFER(FSceneUniformParameters, Scene)
 	SHADER_PARAMETER_RDG_UNIFORM_BUFFER(FDebugViewModePassUniformParameters, Pass)
 	SHADER_PARAMETER_STRUCT_INCLUDE(FInstanceCullingDrawParams, InstanceCullingDrawParams)
-	SHADER_PARAMETER_RDG_UNIFORM_BUFFER(FVirtualShadowMapUniformParameters, VirtualShadowMap)
 	RENDER_TARGET_BINDING_SLOTS()
 END_SHADER_PARAMETER_STRUCT()
 
-void RenderDebugViewMode(
-	FRDGBuilder& GraphBuilder, 
-	TArrayView<FViewInfo> Views, 
-	FRDGTextureRef QuadOverdrawTexture, 
-	const FRenderTargetBindingSlots& RenderTargets, 
-	FSceneUniformBuffer& SceneUniformBuffer,
-	FVirtualShadowMapArray& VirtualShadowMapArray)
+void RenderDebugViewMode(FRDGBuilder& GraphBuilder, TArrayView<FViewInfo> Views, FRDGTextureRef QuadOverdrawTexture, const FRenderTargetBindingSlots& RenderTargets)
 {
 	RDG_EVENT_SCOPE(GraphBuilder, "DebugViewMode");
 
@@ -237,10 +229,8 @@ void RenderDebugViewMode(
 
 			auto* PassParameters = GraphBuilder.AllocParameters<FDebugViewModePassParameters>();
 			PassParameters->View = View.GetShaderParameters();
-			PassParameters->Scene = SceneUniformBuffer.GetBuffer(GraphBuilder);
 			PassParameters->Pass = CreateDebugViewModePassUniformBuffer(GraphBuilder, View, QuadOverdrawTexture);
 			PassParameters->RenderTargets = RenderTargets;
-			PassParameters->VirtualShadowMap = VirtualShadowMapArray.GetUniformBuffer(ViewIndex);
 
 			FScene* Scene = View.Family->Scene->GetRenderScene();
 			check(Scene != nullptr);
@@ -707,9 +697,7 @@ void RenderDebugViewMode(
 	FRDGBuilder& GraphBuilder,
 	TArrayView<FViewInfo> Views,
 	FRDGTextureRef QuadOverdrawTexture,
-	const FRenderTargetBindingSlots& RenderTargets,
-	FSceneUniformBuffer& SceneUniformBuffer,
-	FVirtualShadowMapArray& VirtualShadowMapArray)
+	const FRenderTargetBindingSlots& RenderTargets)
 {}
 
 #endif // WITH_DEBUG_VIEW_MODES
diff --git a/Engine/Source/Runtime/Renderer/Private/DebugViewModeRendering.h b/Engine/Source/Runtime/Renderer/Private/DebugViewModeRendering.h
--- a/Engine/Source/Runtime/Renderer/Private/DebugViewModeRendering.h
+++ b/Engine/Source/Runtime/Renderer/Private/DebugViewModeRendering.h
@@ -23,7 +23,6 @@ class FPrimitiveSceneProxy;
 struct FMeshBatchElement;
 struct FMeshDrawingRenderState;
 class FDebugViewModeInterface;
-class FVirtualShadowMapArray;
 
 static const int32 NumStreamingAccuracyColors = 5;
 static const int32 NumLODColorationColors = 8;
@@ -217,6 +216,4 @@ void RenderDebugViewMode(
 	FRDGBuilder& GraphBuilder,
 	TArrayView<FViewInfo> Views,
 	FRDGTextureRef QuadOverdrawTexture,
-	const FRenderTargetBindingSlots& RenderTargets,
-	FSceneUniformBuffer& SceneUniformBuffer,
-	FVirtualShadowMapArray& VirtualShadowMapArray);
\ No newline at end of file
+	const FRenderTargetBindingSlots& RenderTargets);
\ No newline at end of file
diff --git a/Engine/Source/Runtime/Renderer/Private/DeferredShadingRenderer.cpp b/Engine/Source/Runtime/Renderer/Private/DeferredShadingRenderer.cpp
--- a/Engine/Source/Runtime/Renderer/Private/DeferredShadingRenderer.cpp
+++ b/Engine/Source/Runtime/Renderer/Private/DeferredShadingRenderer.cpp
@@ -3865,7 +3865,7 @@ void FDeferredShadingSceneRenderer::Render(FRDGBuilder& GraphBuilder, const FSce
 					RDG_GPU_MASK_SCOPE(GraphBuilder, View.GPUMask);
 					RDG_EVENT_SCOPE_CONDITIONAL(GraphBuilder, Views.Num() > 1, "View%d", ViewIndex);
 					PostProcessingInputs.TranslucencyViewResourcesMap = FTranslucencyViewResourcesMap(TranslucencyResourceMap, ViewIndex);
-					AddDebugViewPostProcessingPasses(GraphBuilder, View, ViewIndex, GetSceneUniforms(), PostProcessingInputs, NaniteResults, &VirtualShadowMapArray);
+					AddDebugViewPostProcessingPasses(GraphBuilder, View, GetSceneUniforms(), PostProcessingInputs, NaniteResults);
 				}
 			}
 			else
diff --git a/Engine/Source/Runtime/Renderer/Private/PostProcess/PostProcessing.cpp b/Engine/Source/Runtime/Renderer/Private/PostProcess/PostProcessing.cpp
--- a/Engine/Source/Runtime/Renderer/Private/PostProcess/PostProcessing.cpp
+++ b/Engine/Source/Runtime/Renderer/Private/PostProcess/PostProcessing.cpp
@@ -2072,14 +2072,7 @@ void AddPostProcessingPasses(
 	}
 }
 
-void AddDebugViewPostProcessingPasses(
-	FRDGBuilder& GraphBuilder,
-	const FViewInfo& View,
-	int32 ViewIndex,
-	FSceneUniformBuffer& SceneUniformBuffer,
-	const FPostProcessingInputs& Inputs,
-	const Nanite::FRasterResults* NaniteRasterResults,
-	FVirtualShadowMapArray* VirtualShadowMapArray)
+void AddDebugViewPostProcessingPasses(FRDGBuilder& GraphBuilder, const FViewInfo& View, FSceneUniformBuffer &SceneUniformBuffer, const FPostProcessingInputs& Inputs, const Nanite::FRasterResults* NaniteRasterResults)
 {
 	RDG_CSV_STAT_EXCLUSIVE_SCOPE(GraphBuilder, RenderPostProcessing);
 	QUICK_SCOPE_CYCLE_COUNTER(STAT_PostProcessing_Process);
@@ -2108,9 +2101,7 @@ void AddDebugViewPostProcessingPasses(
 	{
 		Visualize,
 		TonemapAfter,
-		VisualizeVirtualShadowMaps_PreEditorPrimitives,
 		SelectionOutline,
-		VisualizeVirtualShadowMaps_PostEditorPrimitives,
 		PrimaryUpscale,
 		SecondaryUpscale,
 		MAX
@@ -2120,9 +2111,7 @@ void AddDebugViewPostProcessingPasses(
 	{
 		TEXT("Visualize"),
 		TEXT("TonemapAfter"),
-		TEXT("VisualizeVirtualShadowMaps_PreEditorPrimitives"),
 		TEXT("SelectionOutline"),
-		TEXT("VisualizeVirtualShadowMaps_PostEditorPrimitives"),
 		TEXT("PrimaryUpscale"),
 		TEXT("SecondaryUpscale")
 	};
@@ -2134,15 +2123,6 @@ void AddDebugViewPostProcessingPasses(
 	PassSequence.SetEnabled(EPass::Visualize, true);
 	PassSequence.SetEnabled(EPass::TonemapAfter, bTonemapAfter);
 	PassSequence.SetEnabled(EPass::SelectionOutline, GIsEditor);
-
-#if WITH_EDITOR || !UE_BUILD_SHIPPING
-	PassSequence.SetEnabled(EPass::VisualizeVirtualShadowMaps_PreEditorPrimitives, VirtualShadowMapArray != nullptr && VirtualShadowMapArray->IsVisualizePassEnabled(View, ViewIndex, EVSMVisualizationPostPass::PreEditorPrimitives));
-	PassSequence.SetEnabled(EPass::VisualizeVirtualShadowMaps_PostEditorPrimitives, VirtualShadowMapArray != nullptr && VirtualShadowMapArray->IsVisualizePassEnabled(View, ViewIndex, EVSMVisualizationPostPass::PostEditorPrimitives));
-#else
-	PassSequence.SetEnabled(EPass::VisualizeVirtualShadowMaps_PreEditorPrimitives, false);
-	PassSequence.SetEnabled(EPass::VisualizeVirtualShadowMaps_PostEditorPrimitives, false);
-#endif
-
 	PassSequence.SetEnabled(EPass::PrimaryUpscale, View.ViewRect.Size() != View.GetSecondaryViewRectSize() && View.PrimaryScreenPercentageMethod != EPrimaryScreenPercentageMethod::TemporalUpscale);
 	PassSequence.SetEnabled(EPass::SecondaryUpscale, View.RequiresSecondaryUpscale() || View.Family->GetSecondarySpatialUpscalerInterface() != nullptr);
 	PassSequence.Finalize();
@@ -2263,13 +2243,6 @@ void AddDebugViewPostProcessingPasses(
 		SceneColor = AddTonemapPass(GraphBuilder, View, PassInputs);
 	}
 
-	if (PassSequence.IsEnabled(EPass::VisualizeVirtualShadowMaps_PreEditorPrimitives))
-	{
-		FScreenPassRenderTarget OverrideOutput;
-		PassSequence.AcceptOverrideIfLastPass(EPass::VisualizeVirtualShadowMaps_PreEditorPrimitives, OverrideOutput);
-		SceneColor = VirtualShadowMapArray->AddVisualizePass(GraphBuilder, View, ViewIndex, EVSMVisualizationPostPass::PreEditorPrimitives, SceneColor, OverrideOutput);
-	}
-
 #if WITH_EDITOR
 	if (PassSequence.IsEnabled(EPass::SelectionOutline))
 	{
@@ -2284,13 +2257,6 @@ void AddDebugViewPostProcessingPasses(
 	}
 #endif
 
-	if (PassSequence.IsEnabled(EPass::VisualizeVirtualShadowMaps_PostEditorPrimitives))
-	{
-		FScreenPassRenderTarget OverrideOutput;
-		PassSequence.AcceptOverrideIfLastPass(EPass::VisualizeVirtualShadowMaps_PostEditorPrimitives, OverrideOutput);
-		SceneColor = VirtualShadowMapArray->AddVisualizePass(GraphBuilder, View, ViewIndex, EVSMVisualizationPostPass::PostEditorPrimitives, SceneColor, OverrideOutput);
-	}
-
 	if (PassSequence.IsEnabled(EPass::PrimaryUpscale))
 	{
 		ISpatialUpscaler::FInputs PassInputs;
@@ -2354,13 +2320,6 @@ void AddDebugViewPostProcessingPasses(
 			SceneColor = ISpatialUpscaler::AddDefaultUpscalePass(GraphBuilder, View, PassInputs, Method);
 		}
 	}
-
-#if WITH_EDITOR || !UE_BUILD_SHIPPING
-	if (ShaderPrint::IsEnabled(View.ShaderPrintData))
-	{
-		ShaderPrint::DrawView(GraphBuilder, View, SceneColor, SceneDepth);
-	}
-#endif
 }
 
 void AddVisualizeCalibrationMaterialPostProcessingPasses(FRDGBuilder& GraphBuilder, const FViewInfo& View, const FPostProcessingInputs& Inputs, const UMaterialInterface* InMaterialInterface)
diff --git a/Engine/Source/Runtime/Renderer/Private/PostProcess/PostProcessing.h b/Engine/Source/Runtime/Renderer/Private/PostProcess/PostProcessing.h
--- a/Engine/Source/Runtime/Renderer/Private/PostProcess/PostProcessing.h
+++ b/Engine/Source/Runtime/Renderer/Private/PostProcess/PostProcessing.h
@@ -46,14 +46,7 @@ void AddPostProcessingPasses(
 	FScreenPassTexture TSRFlickeringInput,
 	FRDGTextureRef& InstancedEditorDepthTexture);
 
-void AddDebugViewPostProcessingPasses(
-	FRDGBuilder& GraphBuilder, 
-	const FViewInfo& View, 
-	int32 ViewIndex,
-	FSceneUniformBuffer& SceneUniformBuffer, 
-	const FPostProcessingInputs& Inputs, 
-	const Nanite::FRasterResults* NaniteRasterResults, 
-	FVirtualShadowMapArray* VirtualShadowMapArray);
+void AddDebugViewPostProcessingPasses(FRDGBuilder& GraphBuilder, const FViewInfo& View, FSceneUniformBuffer &SceneUniformBuffer, const FPostProcessingInputs& Inputs, const Nanite::FRasterResults* NaniteRasterResults);
 
 void AddVisualizeCalibrationMaterialPostProcessingPasses(FRDGBuilder& GraphBuilder, const FViewInfo& View, const FPostProcessingInputs& Inputs, const UMaterialInterface* InMaterialInterface);
 
