Substrate - do not compile tile types that are never going to be used when usinb blendable GBufffer.

#tests PC editor, compiled shader
#rnx
#DisabledValidations Preflight

diff --git a/Engine/Source/Runtime/Core/Private/UObject/FortniteMainBranchObjectVersion.cpp b/Engine/Source/Runtime/Core/Private/UObject/FortniteMainBranchObjectVersion.cpp
--- a/Engine/Source/Runtime/Core/Private/UObject/FortniteMainBranchObjectVersion.cpp
+++ b/Engine/Source/Runtime/Core/Private/UObject/FortniteMainBranchObjectVersion.cpp
@@ -8,7 +8,7 @@ TMap<FGuid, FGuid> FFortniteMainBranchObjectVersion::GetSystemGuids()
 	TMap<FGuid, FGuid> SystemGuids;
 	const FDevSystemGuids& DevGuids = FDevSystemGuids::Get();
 
-	SystemGuids.Add(DevGuids.GLOBALSHADERMAP_DERIVEDDATA_VER, FGuid("2ACD1C7798984B98BECABF4D32E4AA20"));
+	SystemGuids.Add(DevGuids.GLOBALSHADERMAP_DERIVEDDATA_VER, FGuid("539216543F8546C396331084B4997498"));
 	SystemGuids.Add(DevGuids.LANDSCAPE_MOBILE_COOK_VERSION, FGuid("32D02EF867C74B71A0D4E0FA41392732"));
 	SystemGuids.Add(DevGuids.MATERIALSHADERMAP_DERIVEDDATA_VER, FGuid("F051877974B44D409435F321BDB39DBB"));
 	SystemGuids.Add(DevGuids.NIAGARASHADERMAP_DERIVEDDATA_VER, FGuid("6360A977062842A29ED30E3A7ACB0E64"));
diff --git a/Engine/Source/Runtime/Renderer/Private/ClusteredDeferredShadingPass.cpp b/Engine/Source/Runtime/Renderer/Private/ClusteredDeferredShadingPass.cpp
--- a/Engine/Source/Runtime/Renderer/Private/ClusteredDeferredShadingPass.cpp
+++ b/Engine/Source/Runtime/Renderer/Private/ClusteredDeferredShadingPass.cpp
@@ -97,11 +97,10 @@ class FClusteredShadingPS : public FGlobalShader
 
 	class FVisualizeLightCullingDim : SHADER_PERMUTATION_BOOL("VISUALIZE_LIGHT_CULLING");
 	class FHairStrandsLighting : SHADER_PERMUTATION_BOOL("USE_HAIR_LIGHTING");
-	class FSubstrateTileType : SHADER_PERMUTATION_INT("SUBSTRATE_TILETYPE", 4);
 	class FLightFunctionAtlasDim : SHADER_PERMUTATION_BOOL("USE_LIGHT_FUNCTION_ATLAS");
 	class FRectLight : SHADER_PERMUTATION_BOOL("USE_RECT_LIGHT");
 	class FAnistropicMaterials : SHADER_PERMUTATION_BOOL("SUPPORTS_ANISOTROPIC_MATERIALS");
-	using FPermutationDomain = TShaderPermutationDomain<FVisualizeLightCullingDim, FHairStrandsLighting, FSubstrateTileType, FLightFunctionAtlasDim, FRectLight, FAnistropicMaterials>;
+	using FPermutationDomain = TShaderPermutationDomain<FVisualizeLightCullingDim, FHairStrandsLighting, Substrate::FSubstrateTileType, FLightFunctionAtlasDim, FRectLight, FAnistropicMaterials>;
 
 	BEGIN_SHADER_PARAMETER_STRUCT(FParameters, )
 		SHADER_PARAMETER_RDG_UNIFORM_BUFFER(FForwardLightUniformParameters, ForwardLightStruct)
@@ -122,7 +121,8 @@ class FClusteredShadingPS : public FGlobalShader
 	static bool ShouldCompilePermutation(const FGlobalShaderPermutationParameters& Parameters)
 	{
 		FPermutationDomain PermutationVector(Parameters.PermutationId);
-		if (!Substrate::IsSubstrateEnabled() && PermutationVector.Get<FSubstrateTileType>() != 0)
+
+		if (!Substrate::ShouldCompileSubstrateTileTypePermutations(PermutationVector.Get<Substrate::FSubstrateTileType>(), Parameters.Platform))
 		{
 			return false;
 		}
@@ -239,7 +239,7 @@ static void InternalAddClusteredDeferredShadingPass(
 		FClusteredShadingPS::FPermutationDomain PermutationVector;
 		PermutationVector.Set<FClusteredShadingPS::FVisualizeLightCullingDim>(View.Family->EngineShowFlags.VisualizeLightCulling);
 		PermutationVector.Set<FClusteredShadingPS::FHairStrandsLighting>(bHairStrands);
-		PermutationVector.Set<FClusteredShadingPS::FSubstrateTileType>(bSubstrate ? TileType : 0);
+		PermutationVector.Set<Substrate::FSubstrateTileType>(bSubstrate ? TileType : 0);
 		PermutationVector.Set<FClusteredShadingPS::FRectLight>(bHasRectLights);
 		PermutationVector.Set<FClusteredShadingPS::FLightFunctionAtlasDim>(bLightFunctionAtlas);
 		PermutationVector.Set<FClusteredShadingPS::FAnistropicMaterials>(bSupportAnisotropyPermutation);
diff --git a/Engine/Source/Runtime/Renderer/Private/IndirectLightRendering.cpp b/Engine/Source/Runtime/Renderer/Private/IndirectLightRendering.cpp
--- a/Engine/Source/Runtime/Renderer/Private/IndirectLightRendering.cpp
+++ b/Engine/Source/Runtime/Renderer/Private/IndirectLightRendering.cpp
@@ -125,12 +125,11 @@ class FDiffuseIndirectCompositePS : public FGlobalShader
 	class FUpscaleDiffuseIndirectDim : SHADER_PERMUTATION_BOOL("DIM_UPSCALE_DIFFUSE_INDIRECT");
 	class FScreenBentNormalMode : SHADER_PERMUTATION_RANGE_INT("DIM_SCREEN_BENT_NORMAL_MODE", 0, 3);
 	class FShortRangeGI : SHADER_PERMUTATION_BOOL("DIM_SHORT_RANGE_GI");
-	class FSubstrateTileType : SHADER_PERMUTATION_INT("SUBSTRATE_TILETYPE", 4);
 	class FEnableDualSrcBlending : SHADER_PERMUTATION_BOOL("ENABLE_DUAL_SRC_BLENDING");
 	class FHairComplexTransmittance : SHADER_PERMUTATION_BOOL("USE_HAIR_COMPLEX_TRANSMITTANCE");
 	class FScreenSpaceReflectionTiledComposition : SHADER_PERMUTATION_BOOL("SSR_TILED_COMPOSITION");
 
-	using FPermutationDomain = TShaderPermutationDomain<FApplyDiffuseIndirectDim, FUpscaleDiffuseIndirectDim, FScreenBentNormalMode, FShortRangeGI, FSubstrateTileType, FEnableDualSrcBlending, FScreenSpaceReflectionTiledComposition, FHairComplexTransmittance>;
+	using FPermutationDomain = TShaderPermutationDomain<FApplyDiffuseIndirectDim, FUpscaleDiffuseIndirectDim, FScreenBentNormalMode, FShortRangeGI, Substrate::FSubstrateTileType, FEnableDualSrcBlending, FScreenSpaceReflectionTiledComposition, FHairComplexTransmittance>;
 
 	static bool ShouldCompilePermutation(const FGlobalShaderPermutationParameters& Parameters)
 	{
@@ -155,7 +154,7 @@ class FDiffuseIndirectCompositePS : public FGlobalShader
 		}
 
 		// Build Substrate tile permutation only for Lumen
-		if (PermutationVector.Get<FSubstrateTileType>() != ESubstrateTileType::EComplex)
+		if (PermutationVector.Get<Substrate::FSubstrateTileType>() != ESubstrateTileType::EComplex)
 		{
 			return Substrate::IsSubstrateEnabled() && PermutationVector.Get<FApplyDiffuseIndirectDim>() == 2;
 		}
@@ -315,7 +314,6 @@ class FReflectionEnvironmentSkyLightingPS : public FGlobalShader
 	class FDynamicSkyLight : SHADER_PERMUTATION_BOOL("ENABLE_DYNAMIC_SKY_LIGHT");
 	class FSkyShadowing : SHADER_PERMUTATION_BOOL("APPLY_SKY_SHADOWING");
 	class FReflectionSourceType : SHADER_PERMUTATION_ENUM_CLASS("REFLECTION_SOURCE_TYPE", EReflectionSourceType);
-	class FSubstrateTileType : SHADER_PERMUTATION_INT("SUBSTRATE_TILETYPE", 4);
 	class FHairComplexTransmittance : SHADER_PERMUTATION_BOOL("USE_HAIR_COMPLEX_TRANSMITTANCE");
 
 	using FPermutationDomain = TShaderPermutationDomain<
@@ -326,7 +324,7 @@ class FReflectionEnvironmentSkyLightingPS : public FGlobalShader
 		FDynamicSkyLight,
 		FSkyShadowing,
 		FReflectionSourceType,
-		FSubstrateTileType,
+		Substrate::FSubstrateTileType,
 		FHairComplexTransmittance>;
 
 	static EReflectionSourceType GetReflectionSourceType(bool bIsLumenStandalone, bool bIsTiledSSR)
@@ -357,9 +355,9 @@ class FReflectionEnvironmentSkyLightingPS : public FGlobalShader
 			PermutationVector.Set<FSkyShadowing>(false);
 		}
 
-		if (PermutationVector.Get<FSubstrateTileType>() && !Substrate::IsSubstrateEnabled())
+		if (PermutationVector.Get<Substrate::FSubstrateTileType>() && !Substrate::IsSubstrateEnabled())
 		{
-			PermutationVector.Set<FSubstrateTileType>(0);
+			PermutationVector.Set<Substrate::FSubstrateTileType>(0);
 		}
 
 		return PermutationVector;
@@ -376,12 +374,12 @@ class FReflectionEnvironmentSkyLightingPS : public FGlobalShader
 		PermutationVector.Set<FDynamicSkyLight>(bEnableDynamicSkyLight);
 		PermutationVector.Set<FSkyShadowing>(bApplySkyShadowing);
 		PermutationVector.Set<FReflectionSourceType>(GetReflectionSourceType(bLumenStandaloneReflections,bIsTiledSSR));
-		PermutationVector.Set<FSubstrateTileType>(0);
+		PermutationVector.Set<Substrate::FSubstrateTileType>(0);
 		PermutationVector.Set<FHairComplexTransmittance>(bHasHairCardsOrMeshes);
 		if (Substrate::IsSubstrateEnabled())
 		{
 			check(TileType <= ESubstrateTileType::EComplexSpecial);
-			PermutationVector.Set<FSubstrateTileType>(TileType);
+			PermutationVector.Set<Substrate::FSubstrateTileType>(TileType);
 		}
 		return RemapPermutation(PermutationVector);
 	}
@@ -394,6 +392,11 @@ class FReflectionEnvironmentSkyLightingPS : public FGlobalShader
 		}
 
 		FPermutationDomain PermutationVector(Parameters.PermutationId);
+		if (!Substrate::ShouldCompileSubstrateTileTypePermutations(PermutationVector.Get<Substrate::FSubstrateTileType>(), Parameters.Platform))
+		{
+			return false;
+		}
+
 		return PermutationVector == RemapPermutation(PermutationVector);
 	}
 
@@ -1407,7 +1410,7 @@ void FDeferredShadingSceneRenderer::RenderDiffuseIndirectAndAmbientOcclusion(
 
 			const TCHAR* DiffuseIndirectSampling = TEXT("Disabled");
 			FDiffuseIndirectCompositePS::FPermutationDomain PermutationVector;
-			PermutationVector.Set<FDiffuseIndirectCompositePS::FSubstrateTileType>(ESubstrateTileType::EComplex);
+			PermutationVector.Set<Substrate::FSubstrateTileType>(ESubstrateTileType::EComplex);
 
 			bool bUpscale = false;
 			bool bApplyAOToSceneColor = true;
@@ -1422,7 +1425,7 @@ void FDeferredShadingSceneRenderer::RenderDiffuseIndirectAndAmbientOcclusion(
 
 					if (Substrate::IsSubstrateEnabled() && TileType != ESubstrateTileType::ECount)
 					{
-						PermutationVector.Set<FDiffuseIndirectCompositePS::FSubstrateTileType>(TileType);
+						PermutationVector.Set<Substrate::FSubstrateTileType>(TileType);
 					}
 					DiffuseIndirectSampling = TEXT("ScreenProbeGather");
 					bApplyAOToSceneColor = false;
diff --git a/Engine/Source/Runtime/Renderer/Private/LightRendering.cpp b/Engine/Source/Runtime/Renderer/Private/LightRendering.cpp
--- a/Engine/Source/Runtime/Renderer/Private/LightRendering.cpp
+++ b/Engine/Source/Runtime/Renderer/Private/LightRendering.cpp
@@ -914,7 +914,6 @@ class FDeferredLightPS : public FGlobalShader
 	class FAtmosphereTransmittance : SHADER_PERMUTATION_BOOL("USE_ATMOSPHERE_TRANSMITTANCE");
 	class FCloudTransmittance 	: SHADER_PERMUTATION_BOOL("USE_CLOUD_TRANSMITTANCE");
 	class FAnistropicMaterials 	: SHADER_PERMUTATION_BOOL("SUPPORTS_ANISOTROPIC_MATERIALS");
-	class FSubstrateTileType	: SHADER_PERMUTATION_INT("SUBSTRATE_TILETYPE", 4);				// Simple, Single, Complex, ComplexSpecial
 	class FVirtualShadowMapMask : SHADER_PERMUTATION_BOOL("USE_VIRTUAL_SHADOW_MAP_MASK");
 
 	using FPermutationDomain = TShaderPermutationDomain<
@@ -930,7 +929,7 @@ class FDeferredLightPS : public FGlobalShader
 		FAtmosphereTransmittance,
 		FCloudTransmittance,
 		FAnistropicMaterials,
-		FSubstrateTileType,
+		Substrate::FSubstrateTileType,
 		FVirtualShadowMapMask>;
 
 	BEGIN_SHADER_PARAMETER_STRUCT(FParameters, )
@@ -1023,7 +1022,7 @@ class FDeferredLightPS : public FGlobalShader
 			return false;
 		}
 
-		if (Substrate::IsSubstrateEnabled() && PermutationVector.Get<FHairLighting>() == 1 && PermutationVector.Get<FSubstrateTileType>() != ESubstrateTileType::ESimple)
+		if (Substrate::IsSubstrateEnabled() && PermutationVector.Get<FHairLighting>() == 1 && PermutationVector.Get<Substrate::FSubstrateTileType>() != ESubstrateTileType::ESimple)
 		{
 			return false;
 		}
@@ -1053,20 +1052,20 @@ class FDeferredLightPS : public FGlobalShader
 			return false;
 		}
 
-		if (!Substrate::IsSubstrateEnabled() && PermutationVector.Get<FSubstrateTileType>() != 0)
+		if (!Substrate::ShouldCompileSubstrateTileTypePermutations(PermutationVector.Get<Substrate::FSubstrateTileType>(), Parameters.Platform))
 		{
 			return false;
 		}
 
 		// We must keep FAnistropicMaterials 0/1 permutations for tiles of Complex and ComplexSpecial complexity. That is because rect light or hair can disable anisotropy explicitely.
-		if (Substrate::IsSubstrateEnabled() && PermutationVector.Get<FSubstrateTileType>() < ESubstrateTileType::EComplex && PermutationVector.Get<FAnistropicMaterials>() == 1)
+		if (Substrate::IsSubstrateEnabled() && PermutationVector.Get<Substrate::FSubstrateTileType>() < ESubstrateTileType::EComplex && PermutationVector.Get<FAnistropicMaterials>() == 1)
 		{
 			// Remove all permutations for tiles of Simple and Single complexity, when anisotropy is true. No anisotropy is allowed in this case with Substrate.
 			return false; 
 		}
 
 		// Skip special complex tile with GBuffer Format=0, as GBuffer Format=0 does not handle glint
-		if (Substrate::IsSubstrateEnabled() && Substrate::IsSubstrateBlendableGBufferEnabled(Parameters.Platform) && PermutationVector.Get<FSubstrateTileType>() == ESubstrateTileType::EComplexSpecial)
+		if (Substrate::IsSubstrateEnabled() && Substrate::IsSubstrateBlendableGBufferEnabled(Parameters.Platform) && PermutationVector.Get<Substrate::FSubstrateTileType>() == ESubstrateTileType::EComplexSpecial)
 		{
 			return false;
 		}
@@ -1092,7 +1091,7 @@ class FDeferredLightPS : public FGlobalShader
 		}
 
 		// We must keep FAnistropicMaterials 0/1 permutations for tiles of Complex and ComplexSpecial complexity. That is because rect light or hair can disable anisotropy explicitely.
-		if (Substrate::IsSubstrateEnabled() && PermutationVector.Get<FSubstrateTileType>() < ESubstrateTileType::EComplex)
+		if (Substrate::IsSubstrateEnabled() && PermutationVector.Get<Substrate::FSubstrateTileType>() < ESubstrateTileType::EComplex)
 		{
 			PermutationVector.Set<FAnistropicMaterials>(false);		// No anisotropy for tiles of Simple and Single complexity
 		}
@@ -2920,7 +2919,7 @@ static void RenderLight(
 		PermutationVector.Set< FDeferredLightPS::FLightingChannelsDim >(View.bUsesLightingChannels);
 		PermutationVector.Set< FDeferredLightPS::FVisualizeCullingDim >(View.Family->EngineShowFlags.VisualizeLightCulling);
 		PermutationVector.Set< FDeferredLightPS::FVirtualShadowMapMask >(bUseVirtualShadowMapMask);
-		PermutationVector.Set< FDeferredLightPS::FSubstrateTileType >(0);
+		PermutationVector.Set< Substrate::FSubstrateTileType >(0);
 		PermutationVector.Set< FDeferredLightPS::FHairComplexTransmittance >(bNeedComplexTransmittanceSupport);
 		PermutationVector.Set< FDeferredLightPS::FLightFunctionAtlasDim >(
 			LightFunctionAtlas::IsEnabled(View, ELightFunctionAtlasSystem::DeferredLighting) && LightSceneInfo->Proxy->HasValidLightFunctionAtlasSlot() &&
@@ -2960,7 +2959,7 @@ static void RenderLight(
 			if (ShouldRenderTileType(ESubstrateTileType::EComplexSpecial))
 			{
 				const ESubstrateTileType TileType = ESubstrateTileType::EComplexSpecial;
-				PermutationVector.Set<FDeferredLightPS::FSubstrateTileType>(TileType);
+				PermutationVector.Set<Substrate::FSubstrateTileType>(TileType);
 				FDeferredLightPS::FPermutationDomain PermutationVectorRemap = FDeferredLightPS::RemapPermutation(PermutationVector);
 				TShaderMapRef< FDeferredLightPS > PixelShader(View.ShaderMap, PermutationVectorRemap);
 				InternalRenderLight(GraphBuilder, Scene, View, LightSceneInfo, PixelShader, PassParameters, TileType, TEXT("Light::StandardDeferred(ComplexSpecial)"));
@@ -2968,7 +2967,7 @@ static void RenderLight(
 			if (ShouldRenderTileType(ESubstrateTileType::EComplex))
 			{
 				const ESubstrateTileType TileType = ESubstrateTileType::EComplex;
-				PermutationVector.Set<FDeferredLightPS::FSubstrateTileType>(TileType);
+				PermutationVector.Set<Substrate::FSubstrateTileType>(TileType);
 				FDeferredLightPS::FPermutationDomain PermutationVectorRemap = FDeferredLightPS::RemapPermutation(PermutationVector);
 				TShaderMapRef< FDeferredLightPS > PixelShader(View.ShaderMap, PermutationVectorRemap);
 				InternalRenderLight(GraphBuilder, Scene, View, LightSceneInfo, PixelShader, PassParameters, TileType, TEXT("Light::StandardDeferred(Complex)"));
@@ -2976,7 +2975,7 @@ static void RenderLight(
 			if (ShouldRenderTileType(ESubstrateTileType::ESingle))
 			{
 				const ESubstrateTileType TileType = ESubstrateTileType::ESingle;
-				PermutationVector.Set<FDeferredLightPS::FSubstrateTileType>(TileType);
+				PermutationVector.Set<Substrate::FSubstrateTileType>(TileType);
 				FDeferredLightPS::FPermutationDomain PermutationVectorRemap = FDeferredLightPS::RemapPermutation(PermutationVector);
 				TShaderMapRef< FDeferredLightPS > PixelShader(View.ShaderMap, PermutationVectorRemap);
 				InternalRenderLight(GraphBuilder, Scene, View, LightSceneInfo, PixelShader, PassParameters, TileType, TEXT("Light::StandardDeferred(Single)"));
@@ -2984,7 +2983,7 @@ static void RenderLight(
 			if (ShouldRenderTileType(ESubstrateTileType::ESimple))
 			{
 				const ESubstrateTileType TileType = ESubstrateTileType::ESimple;
-				PermutationVector.Set<FDeferredLightPS::FSubstrateTileType>(TileType);
+				PermutationVector.Set<Substrate::FSubstrateTileType>(TileType);
 				FDeferredLightPS::FPermutationDomain PermutationVectorRemap = FDeferredLightPS::RemapPermutation(PermutationVector);
 				TShaderMapRef< FDeferredLightPS > PixelShader(View.ShaderMap, PermutationVectorRemap);
 				InternalRenderLight(GraphBuilder, Scene, View, LightSceneInfo, PixelShader, PassParameters, TileType, TEXT("Light::StandardDeferred(Simple)"));
@@ -2992,7 +2991,7 @@ static void RenderLight(
 		}
 		else
 		{
-			PermutationVector.Set< FDeferredLightPS::FSubstrateTileType>(0);
+			PermutationVector.Set< Substrate::FSubstrateTileType>(0);
 			PermutationVector = FDeferredLightPS::RemapPermutation(PermutationVector);
 			TShaderMapRef< FDeferredLightPS > PixelShader(View.ShaderMap, PermutationVector);
 			InternalRenderLight(GraphBuilder, Scene, View, LightSceneInfo, PixelShader, PassParameters, ESubstrateTileType::ECount, TEXT("Light::StandardDeferred"));
@@ -3097,7 +3096,7 @@ void FDeferredShadingSceneRenderer::RenderLightForHair(
 	PermutationVector.Set< FDeferredLightPS::FHairLighting>(1);
 	PermutationVector.Set< FDeferredLightPS::FHairComplexTransmittance>(true);
 	PermutationVector.Set< FDeferredLightPS::FVirtualShadowMapMask >(bUseVirtualShadowMapMask);
-	PermutationVector.Set< FDeferredLightPS::FSubstrateTileType>(Substrate::IsSubstrateEnabled() ? ESubstrateTileType::ESimple : 0);
+	PermutationVector.Set< Substrate::FSubstrateTileType>(Substrate::IsSubstrateEnabled() ? ESubstrateTileType::ESimple : 0);
 	PermutationVector.Set< FDeferredLightPS::FLightFunctionAtlasDim >(
 		LightFunctionAtlas::IsEnabled(View, ELightFunctionAtlasSystem::DeferredLighting) && LightSceneInfo->Proxy->HasValidLightFunctionAtlasSlot() &&
 		LightSceneInfo->Proxy->GetLightFunctionMaterial() != nullptr && !View.Family->EngineShowFlags.VisualizeLightCulling && bCanLightUsesAtlasForUnbatchedLight);
@@ -3339,7 +3338,7 @@ static void InternalRenderSimpleLightsForHair(
 	PermutationVector.Set< FDeferredLightPS::FHairComplexTransmittance>(true /*bNeedComplexTransmittanceSupport*/);
 	PermutationVector.Set< FDeferredLightPS::FAtmosphereTransmittance >(false);
 	PermutationVector.Set< FDeferredLightPS::FCloudTransmittance >(false);
-	PermutationVector.Set< FDeferredLightPS::FSubstrateTileType>(Substrate::IsSubstrateEnabled() ? ESubstrateTileType::ESimple : 0);
+	PermutationVector.Set< Substrate::FSubstrateTileType>(Substrate::IsSubstrateEnabled() ? ESubstrateTileType::ESimple : 0);
 
 	PermutationVector = FDeferredLightPS::RemapPermutation(PermutationVector);
 	TShaderMapRef<FDeferredLightPS> PixelShader(View.ShaderMap, PermutationVector);
@@ -3429,7 +3428,7 @@ static void InternalRenderSimpleLightsStandardDeferred(
 	PermutationVector.Set< FDeferredLightPS::FHairComplexTransmittance>(bNeedComplexTransmittanceSupport);
 	PermutationVector.Set< FDeferredLightPS::FAtmosphereTransmittance >(false);
 	PermutationVector.Set< FDeferredLightPS::FCloudTransmittance >(false);
-	PermutationVector.Set< FDeferredLightPS::FSubstrateTileType>(TileType != ESubstrateTileType::ECount ? TileType : 0);
+	PermutationVector.Set< Substrate::FSubstrateTileType>(TileType != ESubstrateTileType::ECount ? TileType : 0);
 
 	PermutationVector = FDeferredLightPS::RemapPermutation(PermutationVector);
 	TShaderMapRef<FDeferredLightPS> PixelShader(View.ShaderMap, PermutationVector);
@@ -3757,7 +3756,7 @@ void DeferredLightGlobalPSOCollector(const FSceneTexturesConfig& SceneTexturesCo
 
 				// Extract useful information from the permutation vector
 				bool bIsDirectional = (PermutationVector.Get< FDeferredLightPS::FSourceShapeDim >() == ELightSourceShape::Directional);
-				ESubstrateTileType SubstrateTileMaterialType = Substrate::IsSubstrateEnabled() ? (ESubstrateTileType) PermutationVector.Get< FDeferredLightPS::FSubstrateTileType >() : ESubstrateTileType::ECount;
+				ESubstrateTileType SubstrateTileMaterialType = Substrate::IsSubstrateEnabled() ? (ESubstrateTileType) PermutationVector.Get< Substrate::FSubstrateTileType >() : ESubstrateTileType::ECount;
 				bool bIsHairShader = PermutationVector.Get<FDeferredLightPS::FHairLighting>() > 0;
 
 				AddPSOInitializer(bIsHairShader, bCameraInsideLightGeometry, bReverseCulling, bIsDirectional, SubstrateTileMaterialType, RHIPixelShader);
diff --git a/Engine/Source/Runtime/Renderer/Private/Substrate/Substrate.cpp b/Engine/Source/Runtime/Renderer/Private/Substrate/Substrate.cpp
--- a/Engine/Source/Runtime/Renderer/Private/Substrate/Substrate.cpp
+++ b/Engine/Source/Runtime/Renderer/Private/Substrate/Substrate.cpp
@@ -195,7 +195,7 @@ bool GetSubstrateUsesTileType(const FViewInfo& View, ESubstrateTileType TileType
 	if (Substrate::IsSubstrateEnabled())
 	{
 		static_assert(uint32(ESubstrateTileType::EComplexSpecial) < sizeof(uint8)*8u);
-		check(TileType <= ESubstrateTileType::EComplexSpecial);		
+		check(TileType <= ESubstrateTileType::EComplexSpecial);
 
 		const bool bCompatible = TileType == ESubstrateTileType::EComplexSpecial ? !Substrate::IsSubstrateBlendableGBufferEnabled(View.GetShaderPlatform()) : true;
 		const uint8 Mask = GetSubstrateTileTypeAsUint8(TileType);
@@ -1967,4 +1967,21 @@ void AddSubstrateSampleMaterialPass(FRDGBuilder& GraphBuilder, const FScene* Sce
 	}
 }
 
+bool ShouldCompileSubstrateTileTypePermutations(const int32 SubstrateTileType, const EShaderPlatform Platform)
+{
+	const bool bIsSubstrateEnabled = Substrate::IsSubstrateEnabled();
+
+	if (!bIsSubstrateEnabled && SubstrateTileType != 0)
+	{
+		return false; // When Substrate is disabled, only light type 0 is used.
+	}
+
+	if (bIsSubstrateEnabled && SubstrateTileType == SUBSTRATE_TILE_TYPE_COMPLEX_SPECIAL && IsSubstrateBlendableGBufferEnabled(Platform))
+	{
+		return false; // When Substrate is enabled and uses blendable GBuffer, ComplexSpecialTiles are nevers used (e.g. used for glints).
+	}
+
+	return true;
+}
+
 } // namespace Substrate
diff --git a/Engine/Source/Runtime/Renderer/Private/Substrate/Substrate.h b/Engine/Source/Runtime/Renderer/Private/Substrate/Substrate.h
--- a/Engine/Source/Runtime/Renderer/Private/Substrate/Substrate.h
+++ b/Engine/Source/Runtime/Renderer/Private/Substrate/Substrate.h
@@ -291,4 +291,8 @@ void AddSubstrateDBufferBasePass(
 	FInstanceCullingManager& InstanceCullingManager, 
 	const FSubstrateSceneData& SubstrateSceneData);
 
+
+class FSubstrateTileType : SHADER_PERMUTATION_INT("SUBSTRATE_TILETYPE", 4);
+bool ShouldCompileSubstrateTileTypePermutations(const int32 SubstrateTileType, const EShaderPlatform Platform);
+
 }; // namespace Substrate
