Substrate - fixed comment typos.

#rnx

diff --git a/Engine/Source/Runtime/Renderer/Private/Substrate/Substrate.cpp b/Engine/Source/Runtime/Renderer/Private/Substrate/Substrate.cpp
--- a/Engine/Source/Runtime/Renderer/Private/Substrate/Substrate.cpp
+++ b/Engine/Source/Runtime/Renderer/Private/Substrate/Substrate.cpp
@@ -1973,12 +1973,12 @@ bool ShouldCompileSubstrateTileTypePermutations(const int32 SubstrateTileType, c
 
 	if (!bIsSubstrateEnabled && SubstrateTileType != 0)
 	{
-		return false; // When Substrate is disabled, only light type 0 is used.
+		return false; // When Substrate is disabled, only tile type 0 is used.
 	}
 
 	if (bIsSubstrateEnabled && SubstrateTileType == SUBSTRATE_TILE_TYPE_COMPLEX_SPECIAL && IsSubstrateBlendableGBufferEnabled(Platform))
 	{
-		return false; // When Substrate is enabled and uses blendable GBuffer, ComplexSpecialTiles are nevers used (e.g. used for glints).
+		return false; // When Substrate is enabled and uses blendable GBuffer, ComplexSpecial tile types are never used (e.g. used for glints).
 	}
 
 	return true;
