Added project cvar to disable clustered shading shaders cook and runtime compilation.

#tests pc and cooked shaders on consoles
#rnx
#DisabledValidations Preflight

diff --git a/Engine/Source/Runtime/Core/Private/UObject/FortniteMainBranchObjectVersion.cpp b/Engine/Source/Runtime/Core/Private/UObject/FortniteMainBranchObjectVersion.cpp
--- a/Engine/Source/Runtime/Core/Private/UObject/FortniteMainBranchObjectVersion.cpp
+++ b/Engine/Source/Runtime/Core/Private/UObject/FortniteMainBranchObjectVersion.cpp
@@ -8,7 +8,7 @@ TMap<FGuid, FGuid> FFortniteMainBranchObjectVersion::GetSystemGuids()
 	TMap<FGuid, FGuid> SystemGuids;
 	const FDevSystemGuids& DevGuids = FDevSystemGuids::Get();
 
-	SystemGuids.Add(DevGuids.GLOBALSHADERMAP_DERIVEDDATA_VER, FGuid("539216543F8546C396331084B4997498"));
+	SystemGuids.Add(DevGuids.GLOBALSHADERMAP_DERIVEDDATA_VER, FGuid("E5A68CF5F40C43B3B7BE0C7110B5A66A"));
 	SystemGuids.Add(DevGuids.LANDSCAPE_MOBILE_COOK_VERSION, FGuid("32D02EF867C74B71A0D4E0FA41392732"));
 	SystemGuids.Add(DevGuids.MATERIALSHADERMAP_DERIVEDDATA_VER, FGuid("F051877974B44D409435F321BDB39DBB"));
 	SystemGuids.Add(DevGuids.NIAGARASHADERMAP_DERIVEDDATA_VER, FGuid("6360A977062842A29ED30E3A7ACB0E64"));
diff --git a/Engine/Source/Runtime/Engine/Classes/Engine/RendererSettings.h b/Engine/Source/Runtime/Engine/Classes/Engine/RendererSettings.h
--- a/Engine/Source/Runtime/Engine/Classes/Engine/RendererSettings.h
+++ b/Engine/Source/Runtime/Engine/Classes/Engine/RendererSettings.h
@@ -1024,6 +1024,15 @@ class URendererSettings : public UDeveloperSettings
 		ToolTip = "Enable cloud shadow on translucent surface not relying on the translucenct lighting volume, e.g. using Forward lighting. This is evaluated per vertex to reduce GPU cost and requires extra samplers/textures to be bound to vertex shaders. This is not implemented on mobile as VolumetricClouds are not available on these platforms.",
 		ConfigRestartRequired = true))
 		uint32 bSupportCloudShadowOnForwardLitTranslucent : 1;
+		
+	/**
+	"When enabled, the height fog scattering lob will match the volumetric fog phase function."
+	*/
+	UPROPERTY(config, EditAnywhere, Category = Optimizations, meta = (
+		ConsoleVariable = "r.ClusteredDeferredShading.EnableForProject", DisplayName = "Support Clustered Deferred Shading",
+		ToolTip = "Whether to support Clustered Deferred Shading for the project. When enabled, more shaders are compiled and the feature can be toggled on/off at runtime using r.UseClusteredDeferredShading.",
+		ConfigRestartRequired = true))
+		uint32 bClusteredDeferredShadingenableForProject : 1;
 
 	/**
 	"Select the format of the light function atlas texture."
diff --git a/Engine/Source/Runtime/RenderCore/Private/Shader.cpp b/Engine/Source/Runtime/RenderCore/Private/Shader.cpp
--- a/Engine/Source/Runtime/RenderCore/Private/Shader.cpp
+++ b/Engine/Source/Runtime/RenderCore/Private/Shader.cpp
@@ -2374,6 +2374,15 @@ void ShaderMapAppendKey(EShaderPlatform Platform, FShaderKeyGenerator& KeyGen)
 		}
 	}
 
+	{
+		static const auto CVar = IConsoleManager::Get().FindTConsoleVariableDataInt(TEXT("r.ClusteredDeferredShading.EnableForProject"));
+		if (CVar && CVar->GetValueOnAnyThread() > 0)
+		{
+			KeyGen.AppendSeparator();
+			KeyGen.Append(TEXT("CLUDEFF"));
+		}
+	}
+
 	const bool bNeedsSeparateMainDirLightTexture = IsWaterSeparateMainDirLightEnabled(Platform);
 	if (bNeedsSeparateMainDirLightTexture)
 	{
diff --git a/Engine/Source/Runtime/Renderer/Private/ClusteredDeferredShadingPass.cpp b/Engine/Source/Runtime/Renderer/Private/ClusteredDeferredShadingPass.cpp
--- a/Engine/Source/Runtime/Renderer/Private/ClusteredDeferredShadingPass.cpp
+++ b/Engine/Source/Runtime/Renderer/Private/ClusteredDeferredShadingPass.cpp
@@ -40,6 +40,13 @@ static FAutoConsoleVariableRef CVarUseClusteredDeferredShading(
 	ECVF_RenderThreadSafe
 );
 
+// The project setting (disable runtime and shader code)
+static TAutoConsoleVariable<int32> CVarClusteredDeferredShadingEnableForProject(
+	TEXT("r.ClusteredDeferredShading.EnableForProject"),
+	1,
+	TEXT("Enables Clustered Deferred Shading rendering and shader code."),
+	ECVF_ReadOnly | ECVF_RenderThreadSafe);
+
 DECLARE_GPU_STAT_NAMED(ClusteredShading, TEXT("Clustered Shading"));
 
 using namespace LightFunctionAtlas;
@@ -49,7 +56,7 @@ bool FDeferredShadingSceneRenderer::ShouldUseClusteredDeferredShading() const
 	// The feature level is the same as in the shader compile conditions below, maybe we don't need SM5?
 	// NOTE: should also take into account the conditions for building the light grid, since these 
 	//       shaders might have another feature level.
-	return GUseClusteredDeferredShading != 0 && Scene->GetFeatureLevel() >= ERHIFeatureLevel::SM5;
+	return CVarClusteredDeferredShadingEnableForProject.GetValueOnAnyThread() > 0 && GUseClusteredDeferredShading != 0 && Scene->GetFeatureLevel() >= ERHIFeatureLevel::SM5;
 }
 
 bool FDeferredShadingSceneRenderer::AreLightsInLightGrid() const
@@ -74,8 +81,7 @@ class FClusteredShadingVS : public FGlobalShader
 
 	static bool ShouldCompilePermutation(const FGlobalShaderPermutationParameters& Parameters)
 	{
-		// OLATODO: what level do we actually need for this?
-		return IsFeatureLevelSupported(Parameters.Platform, ERHIFeatureLevel::SM5);
+		return IsFeatureLevelSupported(Parameters.Platform, ERHIFeatureLevel::SM5) && CVarClusteredDeferredShadingEnableForProject.GetValueOnAnyThread() > 0;
 	}
 
 	static void ModifyCompilationEnvironment(const FGlobalShaderPermutationParameters& Parameters, FShaderCompilerEnvironment& OutEnvironment)
@@ -131,8 +137,7 @@ class FClusteredShadingPS : public FGlobalShader
 			return false;
 		}
 
-		// OLATODO: what level do we actually need for this?
-		return IsFeatureLevelSupported(Parameters.Platform, ERHIFeatureLevel::SM5);
+		return IsFeatureLevelSupported(Parameters.Platform, ERHIFeatureLevel::SM5) && CVarClusteredDeferredShadingEnableForProject.GetValueOnAnyThread() > 0;
 	}
 
 	static void ModifyCompilationEnvironment(const FGlobalShaderPermutationParameters& Parameters, FShaderCompilerEnvironment& OutEnvironment)
@@ -310,7 +315,7 @@ void FDeferredShadingSceneRenderer::AddClusteredDeferredShadingPass(
 	FRDGTextureRef HairStrandsShadowMaskBits,
 	FRDGTextureRef LightingChannelsTexture)
 {
-	check(GUseClusteredDeferredShading);
+	check(ShouldUseClusteredDeferredShading());
 
 	const int32 NumLightsToRender = SortedLightsSet.ClusteredSupportedEnd;
 
