Fix maths used for screen size culling of velocity pass.

diff --git a/Engine/Source/Runtime/Renderer/Private/VelocityRendering.cpp b/Engine/Source/Runtime/Renderer/Private/VelocityRendering.cpp
--- a/Engine/Source/Runtime/Renderer/Private/VelocityRendering.cpp
+++ b/Engine/Source/Runtime/Renderer/Private/VelocityRendering.cpp
@@ -830,14 +830,12 @@ bool FVelocityMeshProcessor::PrimitiveHasVelocityForView(const FViewInfo& View,
 	}
 
 	const FBoxSphereBounds& PrimitiveBounds = PrimitiveSceneProxy->GetBounds();
-	const float LODFactorDistanceSquared = (PrimitiveBounds.Origin - View.ViewMatrices.GetViewOrigin()).SizeSquared() * FMath::Square(View.LODDistanceFactor);
+	const float PrimitiveScreenRadiusSq = ComputeBoundsScreenRadiusSquared(PrimitiveBounds.Origin, PrimitiveBounds.SphereRadius, View);
 
-	// The minimum projected screen radius for a primitive to be drawn in the velocity pass, as a fraction of half the horizontal screen width.
-	float MinScreenRadiusForVelocityPass = View.FinalPostProcessSettings.MotionBlurPerObjectSize * (2.0f / 100.0f);
-	float MinScreenRadiusForVelocityPassSquared = FMath::Square(MinScreenRadiusForVelocityPass);
+	const float MinScreenRadiusForVelocityPass = (View.FinalPostProcessSettings.MotionBlurPerObjectSize * 0.5f / 100.0f) * View.LODDistanceFactor;
+	const float MinScreenRadiusForVelocityPassSquared = FMath::Square(MinScreenRadiusForVelocityPass);
 
-	// Skip primitives that only cover a small amount of screen space, motion blur on them won't be noticeable.
-	if (FMath::Square(PrimitiveBounds.SphereRadius) <= MinScreenRadiusForVelocityPassSquared * LODFactorDistanceSquared)
+	if (PrimitiveScreenRadiusSq < MinScreenRadiusForVelocityPassSquared)
 	{
 		return false;
 	}
