Fix temporal responsiveness bug due to default texture output format not supported on some consoles.

#tests Test on PS5 and the artifact is gone.
#DisabledValidations NonUnityAndStaticAnalysis

diff --git a/Engine/Source/Runtime/Renderer/Private/VelocityRendering.cpp b/Engine/Source/Runtime/Renderer/Private/VelocityRendering.cpp
--- a/Engine/Source/Runtime/Renderer/Private/VelocityRendering.cpp
+++ b/Engine/Source/Runtime/Renderer/Private/VelocityRendering.cpp
@@ -253,6 +253,7 @@ public:
 	{
 		FMaterialShader::ModifyCompilationEnvironment(Parameters, OutEnvironment);
 		OutEnvironment.SetRenderTargetOutputFormat(0, PF_A16B16G16R16);
+		OutEnvironment.SetRenderTargetOutputFormat(1, PF_A16B16G16R16);
 
 		// We support velocity on thin translucent only with masking, and only if the material is only made of thin translucent shading model.
 		OutEnvironment.SetDefine(TEXT("VELOCITY_THIN_TRANSLUCENT_MODE"), Parameters.MaterialParameters.ShadingModels.HasOnlyShadingModel(MSM_ThinTranslucent));
