Remove deprecated RequestVirtualTextureTilesForRegion() function.
#rnx

diff --git a/Engine/Source/Runtime/RenderCore/Public/RendererInterface.h b/Engine/Source/Runtime/RenderCore/Public/RendererInterface.h
--- a/Engine/Source/Runtime/RenderCore/Public/RendererInterface.h
+++ b/Engine/Source/Runtime/RenderCore/Public/RendererInterface.h
@@ -838,9 +838,6 @@ public:
 	 */
 	virtual void RequestVirtualTextureTiles(IAllocatedVirtualTexture* AllocatedVT, const FVector2D& InScreenSpaceSize, const FVector2D& InViewportPosition, const FVector2D& InViewportSize, const FVector2D& InUV0, const FVector2D& InUV1, int32 InMipLevel) = 0;
 
-	UE_DEPRECATED(5.4, "Use RequestVirtualTextureTiles() overloads that takes similar parameters. Make sure not to negate the InViewportPosition.")
-	virtual void RequestVirtualTextureTilesForRegion(IAllocatedVirtualTexture* AllocatedVT, const FVector2D& InScreenSpaceSize, const FVector2D& InViewportPosition, const FVector2D& InViewportSize, const FVector2D& InUV0, const FVector2D& InUV1, int32 InMipLevel) = 0;
-
 	/** Ensure that any tiles requested by 'RequestVirtualTextureTiles' are loaded, must be called from render thread */
 	virtual void LoadPendingVirtualTextureTiles(FRHICommandListImmediate& RHICmdList, ERHIFeatureLevel::Type FeatureLevel) = 0;
 	
diff --git a/Engine/Source/Runtime/Renderer/Private/RendererModule.h b/Engine/Source/Runtime/Renderer/Private/RendererModule.h
--- a/Engine/Source/Runtime/Renderer/Private/RendererModule.h
+++ b/Engine/Source/Runtime/Renderer/Private/RendererModule.h
@@ -114,9 +114,6 @@ public:
 	virtual void RequestVirtualTextureTiles(IAllocatedVirtualTexture* AllocatedVT, const FVector2D& InScreenSpaceSize, const FVector2D& InViewportPosition, const FVector2D& InViewportSize, const FVector2D& InUV0, const FVector2D& InUV1, int32 InMipLevel) override;
 	virtual IVirtualTexture* FindProducer(const FVirtualTextureProducerHandle& Handle) override;
 
-	UE_DEPRECATED(5.4, "Use RequestVirtualTextureTiles() overloads that takes similar parameters. Make sure not to negate the InViewportPosition.")
-	virtual void RequestVirtualTextureTilesForRegion(IAllocatedVirtualTexture* AllocatedVT, const FVector2D& InScreenSpaceSize, const FVector2D& InViewportPosition, const FVector2D& InViewportSize, const FVector2D& InUV0, const FVector2D& InUV1, int32 InMipLevel) override;
-	
 	virtual void LoadPendingVirtualTextureTiles(FRHICommandListImmediate& RHICmdList, ERHIFeatureLevel::Type FeatureLevel) override;
 	virtual void LockVirtualTextureTiles(FVirtualTextureProducerHandle ProducerHandle, int32 InMipLevel) override;
 	virtual void SetVirtualTextureRequestRecordBuffer(uint64 Handle) override;
diff --git a/Engine/Source/Runtime/Renderer/Private/SceneRendering.cpp b/Engine/Source/Runtime/Renderer/Private/SceneRendering.cpp
--- a/Engine/Source/Runtime/Renderer/Private/SceneRendering.cpp
+++ b/Engine/Source/Runtime/Renderer/Private/SceneRendering.cpp
@@ -5531,13 +5531,6 @@ IVirtualTexture* FRendererModule::FindProducer(const FVirtualTextureProducerHand
 	return nullptr;
 }
 
-void FRendererModule::RequestVirtualTextureTilesForRegion(IAllocatedVirtualTexture* AllocatedVT, const FVector2D& InScreenSpaceSize, const FVector2D& InViewportPosition, const FVector2D& InViewportSize, const FVector2D& InUV0, const FVector2D& InUV1, int32 InMipLevel)
-{
-PRAGMA_DISABLE_DEPRECATION_WARNINGS
-	FVirtualTextureSystem::Get().RequestTilesForRegion(AllocatedVT, InScreenSpaceSize, InViewportPosition, InViewportSize, InUV0, InUV1, InMipLevel);
-PRAGMA_ENABLE_DEPRECATION_WARNINGS
-}
-
 void FRendererModule::LoadPendingVirtualTextureTiles(FRHICommandListImmediate& RHICmdList, ERHIFeatureLevel::Type FeatureLevel)
 {
 	FRDGBuilder GraphBuilder(RHICmdList);
diff --git a/Engine/Source/Runtime/Renderer/Private/VT/VirtualTextureSystem.cpp b/Engine/Source/Runtime/Renderer/Private/VT/VirtualTextureSystem.cpp
--- a/Engine/Source/Runtime/Renderer/Private/VT/VirtualTextureSystem.cpp
+++ b/Engine/Source/Runtime/Renderer/Private/VT/VirtualTextureSystem.cpp
@@ -1261,12 +1261,6 @@ void FVirtualTextureSystem::RequestTiles(IAllocatedVirtualTexture* AllocatedVT,
 	}
 }
 
-void FVirtualTextureSystem::RequestTilesForRegion(IAllocatedVirtualTexture* AllocatedVT, const FVector2D& InScreenSpaceSize, const FVector2D& InViewportPosition, const FVector2D& InViewportSize, const FVector2D& InUV0, const FVector2D& InUV1, int32 InMipLevel)
-{
-	// RequestTilesForRegion() used to require that the viewport position was negated which felt wrong. Correct this when calling the new implementation.
-	RequestTiles(AllocatedVT, InScreenSpaceSize, -InViewportPosition, InViewportSize, InUV0, InUV1, InMipLevel);
-}
-
 void FVirtualTextureSystem::LoadPendingTiles(FRDGBuilder& GraphBuilder, ERHIFeatureLevel::Type FeatureLevel)
 {
 	check(!bUpdating);
diff --git a/Engine/Source/Runtime/Renderer/Private/VT/VirtualTextureSystem.h b/Engine/Source/Runtime/Renderer/Private/VT/VirtualTextureSystem.h
--- a/Engine/Source/Runtime/Renderer/Private/VT/VirtualTextureSystem.h
+++ b/Engine/Source/Runtime/Renderer/Private/VT/VirtualTextureSystem.h
@@ -152,9 +152,6 @@ public:
 	 */
 	void RequestTiles(IAllocatedVirtualTexture* AllocatedVT, const FVector2D& InScreenSpaceSize, const FVector2D& InViewportPosition, const FVector2D& InViewportSize, const FVector2D& InUV0, const FVector2D& InUV1, int32 InMipLevel = -1);
 
-	UE_DEPRECATED(5.4, "Use RequestTiles() overloads that takes similar parameters. Make sure not to negate the InViewportPosition.")
-	void RequestTilesForRegion(IAllocatedVirtualTexture* AllocatedVT, const FVector2D& InScreenSpaceSize, const FVector2D& InViewportPosition, const FVector2D& InViewportSize, const FVector2D& InUV0, const FVector2D& InUV1, int32 InMipLevel = -1);
-
 	void LoadPendingTiles(FRDGBuilder& GraphBuilder, ERHIFeatureLevel::Type FeatureLevel);
 	
 	void SetMipLevelToLock(FVirtualTextureProducerHandle ProducerHandle, int32 InMipLevel);
