Filter out all material domains except MD_Surface for shadow depth PSO pre-caching

diff --git a/Engine/Source/Runtime/Renderer/Private/ShadowDepthRendering.cpp b/Engine/Source/Runtime/Renderer/Private/ShadowDepthRendering.cpp
--- a/Engine/Source/Runtime/Renderer/Private/ShadowDepthRendering.cpp
+++ b/Engine/Source/Runtime/Renderer/Private/ShadowDepthRendering.cpp
@@ -2198,7 +2198,7 @@ void FShadowDepthPassMeshProcessor::CollectPSOInitializers(const FSceneTexturesC
 	if (!PreCacheParams.bDefaultMaterial && 
 		(GetShadowQuality() == 0 ||
 		!Material.ShouldCastDynamicShadows() ||
-		!ShouldIncludeDomainInMeshPass(Material.GetMaterialDomain()) ||
+		Material.GetMaterialDomain() != MD_Surface ||
 		!ShouldIncludeMaterialInDefaultOpaquePass(Material) ||
 		!EnumHasAnyFlags(MeshSelectionMask, VertexFactoryData.VertexFactoryType->SupportsPrimitiveIdStream() ? EShadowMeshSelection::VSM : EShadowMeshSelection::SM)))
 	{
