Add permutation for contact shadows to the mobile deferred shaders.

diff --git a/Engine/Shaders/Private/MobileDeferredShading.usf b/Engine/Shaders/Private/MobileDeferredShading.usf
--- a/Engine/Shaders/Private/MobileDeferredShading.usf
+++ b/Engine/Shaders/Private/MobileDeferredShading.usf
@@ -22,8 +22,9 @@
 #include "/Engine/Generated/Material.ush"
 #endif
 
-// In-line Contact Shadows only in mobile deferred multi-pass
-#define SUPPORT_CONTACT_SHADOWS !ALLOW_FRAMEBUFFER_FETCH
+#ifndef SUPPORT_CONTACT_SHADOWS
+#define SUPPORT_CONTACT_SHADOWS 0
+#endif
 #include "MobileLightingCommon.ush"
 #if MATERIAL_SHADER
 #include "LightFunctionCommon.ush"
diff --git a/Engine/Source/Runtime/Renderer/Private/MobileDeferredShadingPass.cpp b/Engine/Source/Runtime/Renderer/Private/MobileDeferredShadingPass.cpp
--- a/Engine/Source/Runtime/Renderer/Private/MobileDeferredShadingPass.cpp
+++ b/Engine/Source/Runtime/Renderer/Private/MobileDeferredShadingPass.cpp
@@ -54,6 +54,12 @@ static FAutoConsoleVariableRef CVarMobileIgnoreDeferredShadingSkyLightChannels(
 	ECVF_RenderThreadSafe
 );
 
+bool IsMobileInlineContactShadowsEnabled(EShaderPlatform ShaderPlatform)
+{
+	static auto* ContactShadowsCVar = IConsoleManager::Get().FindTConsoleVariableDataInt(TEXT("r.ContactShadows"));
+	return !MobileAllowFramebufferFetch(ShaderPlatform) && ContactShadowsCVar && ContactShadowsCVar->GetValueOnRenderThread() != 0;
+}
+
 BEGIN_SHADER_PARAMETER_STRUCT(FMobileDeferredPassParameters, )
 	SHADER_PARAMETER_RDG_UNIFORM_BUFFER(FMobileSceneTextureUniformParameters, MobileSceneTextures)
 	RENDER_TARGET_BINDING_SLOTS()
@@ -82,6 +88,7 @@ class FMobileDirectionalLightFunctionPS : public FMaterialShader
 	class FShadowQuality				: SHADER_PERMUTATION_RANGE_INT("MOBILE_SHADOW_QUALITY", 1, 3); // not using Quality=0
 	class FMobileSSRQuality 			: SHADER_PERMUTATION_ENUM_CLASS("MOBILE_SSR_QUALITY", EMobileSSRQuality);
 	class FSkyShadowing					: SHADER_PERMUTATION_BOOL("APPLY_SKY_SHADOWING");
+	class FInlineContactShadows			: SHADER_PERMUTATION_BOOL("SUPPORT_CONTACT_SHADOWS");
 
 	using FPermutationDomain = TShaderPermutationDomain<
 		FEnableShadingModelSupport,
@@ -93,7 +100,8 @@ class FMobileDirectionalLightFunctionPS : public FMaterialShader
 		FEnableCSM, 
 		FShadowQuality,
 		FMobileSSRQuality,
-		FSkyShadowing>;
+		FSkyShadowing,
+		FInlineContactShadows>;
 
 	BEGIN_SHADER_PARAMETER_STRUCT(FParameters, )
 		SHADER_PARAMETER_STRUCT_REF(FMobileDirectionalLightShaderParameters, MobileDirectionalLight)
@@ -145,6 +153,11 @@ class FMobileDirectionalLightFunctionPS : public FMaterialShader
 			PermutationVector.Set<FSkyShadowing>(false);
 		}
 		
+		if (MobileAllowFramebufferFetch(Platform))
+		{
+			PermutationVector.Set<FInlineContactShadows>(false);
+		}
+
 		return PermutationVector;
 	}
 
@@ -178,7 +191,8 @@ class FMobileDirectionalLightFunctionPS : public FMaterialShader
 		EMobileSSRQuality MobileSSRQuality,
 		bool bApplySkyShadowing)
 	{
-		bool bUseClusteredLights = UseClusteredDeferredShading(View.GetShaderPlatform());
+		EShaderPlatform ShaderPlatform = View.GetShaderPlatform();
+		bool bUseClusteredLights = UseClusteredDeferredShading(ShaderPlatform);
 		bool bClustredReflection = bInlineReflectionAndSky && (View.NumBoxReflectionCaptures + View.NumSphereReflectionCaptures) > 0;
 		bool bEnableSkyLight = bInlineReflectionAndSky && bSkyLight;
 		int32 ShadowQuality = bDynamicShadows && !bScreenSpaceShadowMask ? (int32)GetShadowQuality() : 0;
@@ -194,6 +208,7 @@ class FMobileDirectionalLightFunctionPS : public FMaterialShader
 		PermutationVector.Set<FShadowQuality>(FMath::Clamp(ShadowQuality, 1, 3));
 		PermutationVector.Set<FMobileSSRQuality>(bInlineReflectionAndSky ? MobileSSRQuality : EMobileSSRQuality::Disabled);
 		PermutationVector.Set<FSkyShadowing>(bInlineReflectionAndSky && bApplySkyShadowing);
+		PermutationVector.Set<FInlineContactShadows>(IsMobileInlineContactShadowsEnabled(ShaderPlatform));
 		return PermutationVector;
 	}
 
@@ -223,7 +238,8 @@ public:
 	class FIESProfileDim			: SHADER_PERMUTATION_BOOL("USE_IES_PROFILE");
 	class FSpotLightShadowDim		: SHADER_PERMUTATION_BOOL("SUPPORT_SPOTLIGHTS_SHADOW");
 	class FSimpleLightDim			: SHADER_PERMUTATION_BOOL("SIMPLE_LIGHT");
-	using FPermutationDomain = TShaderPermutationDomain<FEnableShadingModelSupport, FRadialLightTypeDim, FIESProfileDim, FSpotLightShadowDim, FSimpleLightDim>;
+	class FInlineContactShadows		: SHADER_PERMUTATION_BOOL("SUPPORT_CONTACT_SHADOWS");
+	using FPermutationDomain = TShaderPermutationDomain<FEnableShadingModelSupport, FRadialLightTypeDim, FIESProfileDim, FSpotLightShadowDim, FSimpleLightDim, FInlineContactShadows>;
 
 	BEGIN_SHADER_PARAMETER_STRUCT(FParameters, )
 		SHADER_PARAMETER_STRUCT(FLightShaderParameters, Light)
@@ -268,12 +284,18 @@ public:
 			PermutationVector.Set<FEnableShadingModelSupport>(false);
 		}
 
+		if (MobileAllowFramebufferFetch(Platform))
+		{
+			PermutationVector.Set<FInlineContactShadows>(false);
+		}
+
 		if (PermutationVector.Get<FSimpleLightDim>() == true)
 		{
 			PermutationVector.Set<FEnableShadingModelSupport>(false);
 			PermutationVector.Set<FRadialLightTypeDim>(LIGHT_TYPE_POINT);
 			PermutationVector.Set<FIESProfileDim>(false);
 			PermutationVector.Set<FSpotLightShadowDim>(false);
+			PermutationVector.Set<FInlineContactShadows>(false);
 		}
 
 		return PermutationVector;
@@ -620,10 +642,8 @@ void SetupContactShadowParameters(const FLightSceneProxy* Proxy, const FViewInfo
 	float ContactShadowLength = 0.f;
 	float ContactShadowCastingIntensity = 1.f;
 	float ContactShadowNonCastingIntensity = 0.f;
-	static auto* ContactShadowsCVar = IConsoleManager::Get().FindTConsoleVariableDataInt(TEXT("r.ContactShadows"));
-	const bool bContactShadows = ContactShadowsCVar && ContactShadowsCVar->GetValueOnRenderThread() != 0;
 	// Only do in-line contact shadows on deferred multi pass
-	if (bContactShadows && View.Family->EngineShowFlags.ContactShadows && !MobileAllowFramebufferFetch(View.GetShaderPlatform()))
+	if (IsMobileInlineContactShadowsEnabled(View.GetShaderPlatform()) && View.Family->EngineShowFlags.ContactShadows)
 	{
 		bool bContactShadowLengthInWS = false;
 		GetLightContactShadowParameters(Proxy, ContactShadowLength, bContactShadowLengthInWS, ContactShadowCastingIntensity, ContactShadowNonCastingIntensity);
@@ -1031,6 +1051,7 @@ static void RenderLocalLight(
 		PermutationVector.Set<FMobileRadialLightFunctionPS::FRadialLightTypeDim>(LightType);
 		PermutationVector.Set<FMobileRadialLightFunctionPS::FIESProfileDim>(bUseIESTexture);
 		PermutationVector.Set<FMobileRadialLightFunctionPS::FSpotLightShadowDim>(bShouldCastShadow);
+		PermutationVector.Set<FMobileRadialLightFunctionPS::FInlineContactShadows>(IsMobileInlineContactShadowsEnabled(View.GetShaderPlatform()));
 		FCachedLightMaterial LightMaterial;
 		TShaderRef<FMobileRadialLightFunctionPS> PixelShader;
 		GetLightMaterial(DefaultLightMaterial, LightFunctionMaterialProxy, PermutationVector.ToDimensionValueId(), LightMaterial, PixelShader);
