Mobile: Revert ShadowMask and AO texture fusion - A/B performance tests show fusion is slightly worse on GPU.
Use single-channel PF_R8 pixel format for ShadowMask when possible.
Restored r.Mobile.AmbientOcclusion cvar (now applies only to forward shading).

diff --git a/Engine/Shaders/Private/MobileBasePassPixelShader.usf b/Engine/Shaders/Private/MobileBasePassPixelShader.usf
--- a/Engine/Shaders/Private/MobileBasePassPixelShader.usf
+++ b/Engine/Shaders/Private/MobileBasePassPixelShader.usf
@@ -37,7 +37,7 @@
 #define PreIntegratedGFSampler	MobileBasePass.PreIntegratedGFSampler
 #endif
 
-#define APPLY_AO (MATERIALBLENDING_MASKED || MATERIALBLENDING_SOLID) && !MATERIAL_SHADINGMODEL_UNLIT
+#define APPLY_AO (ENABLE_AMBIENT_OCCLUSION && (MATERIALBLENDING_MASKED || MATERIALBLENDING_SOLID) && !MATERIAL_SHADINGMODEL_UNLIT)
 
 #if MATERIAL_SHADINGMODEL_SINGLELAYERWATER
 	#ifdef SINGLE_LAYER_WATER_SHADING_QUALITY
diff --git a/Engine/Shaders/Private/MobileDeferredShading.usf b/Engine/Shaders/Private/MobileDeferredShading.usf
--- a/Engine/Shaders/Private/MobileDeferredShading.usf
+++ b/Engine/Shaders/Private/MobileDeferredShading.usf
@@ -201,15 +201,17 @@ void MobileDirectionalLightPS(
 
 	float4 ScreenPosition = SvPositionToScreenPosition(float4(SvPosition.xyz, GBuffer.Depth));
 
-	half4 ScreenSpaceShadowMaskSample = SampleScreenSpaceShadowMask(UVAndScreenPos.xy, EyeIndex);
-	half GatheredAmbientOcclusion = ScreenSpaceShadowMaskSample.a;
+	half GatheredAmbientOcclusion = 1;
+#if ENABLE_AMBIENT_OCCLUSION
+	GatheredAmbientOcclusion = Texture2DSample(MobileBasePass.AmbientOcclusionTexture, MobileBasePass.AmbientOcclusionSampler, UVAndScreenPos.xy).r;
 	#if ALLOW_STATIC_LIGHTING
 		GatheredAmbientOcclusion = lerp(1.0, GatheredAmbientOcclusion, MobileBasePass.AmbientOcclusionStaticFraction);
 	#endif
 	GBuffer.GBufferAO *= GatheredAmbientOcclusion;
+#endif
 
 	// Directional light
-	half4 DynamicShadowFactors = half4(ScreenSpaceShadowMaskSample.rgb, 1.0);
+	half4 DynamicShadowFactors = SampleScreenSpaceShadowMask(UVAndScreenPos.xy, EyeIndex);
 	float DirectionalLightShadow = 1.0f;
 	const float Dither = InterleavedGradientNoise(SvPosition.xy, View.StateFrameIndexMod8);
 	AccumulateDirectionalLighting(GBuffer, TranslatedWorldPosition, CameraVector, ScreenPosition, DynamicShadowFactors, Dither, DirectionalLightShadow, DirectLighting);
@@ -433,11 +435,14 @@ void MobileReflectionEnvironmentSkyLightingPS(
 	float2 LocalPosition = SvPosition.xy - View.ViewRectMin.xy;
 	uint GridIndex = ComputeLightGridCellIndex(uint2(LocalPosition.x, LocalPosition.y), GBuffer.Depth, EyeIndex);
 
-	half GatheredAmbientOcclusion = SampleScreenSpaceShadowMask(UVAndScreenPos.xy, EyeIndex).a;
-#if ALLOW_STATIC_LIGHTING
-	GatheredAmbientOcclusion = lerp(1.0, GatheredAmbientOcclusion, MobileBasePass.AmbientOcclusionStaticFraction);
-#endif
+	half GatheredAmbientOcclusion = 1;
+#if ENABLE_AMBIENT_OCCLUSION
+	GatheredAmbientOcclusion = Texture2DSample(MobileBasePass.AmbientOcclusionTexture, MobileBasePass.AmbientOcclusionSampler, UVAndScreenPos.xy).r;
+	#if ALLOW_STATIC_LIGHTING
+		GatheredAmbientOcclusion = lerp(1.0, GatheredAmbientOcclusion, MobileBasePass.AmbientOcclusionStaticFraction);
+	#endif
 	GBuffer.GBufferAO *= GatheredAmbientOcclusion;
+#endif
 	
 	FLightAccumulator DirectLighting = (FLightAccumulator)0;
 	ReflectionEnvironmentSkyLighting(GBuffer, SvPosition, CameraVector, TranslatedWorldPosition, ReflectionVector, GridIndex, GatheredAmbientOcclusion, DirectLighting);
diff --git a/Engine/Shaders/Private/MobileLightingCommon.ush b/Engine/Shaders/Private/MobileLightingCommon.ush
--- a/Engine/Shaders/Private/MobileLightingCommon.ush
+++ b/Engine/Shaders/Private/MobileLightingCommon.ush
@@ -284,7 +284,7 @@ half4 SampleScreenSpaceShadowMask(float2 BufferUV, uint EyeIndex)
 	#if ALLOW_STATIC_LIGHTING
 		return Sample;
 	#else	
-		return half4(Sample.r, 1.0, 1.0, Sample.g); // PG_RG8 with a AO in g
+		return half4(Sample.r, 1.0, 1.0, 1.0); // PG_R8
 	#endif
 #endif
 	return 1.0;
diff --git a/Engine/Source/Runtime/RenderCore/Private/RenderUtils.cpp b/Engine/Source/Runtime/RenderCore/Private/RenderUtils.cpp
--- a/Engine/Source/Runtime/RenderCore/Private/RenderUtils.cpp
+++ b/Engine/Source/Runtime/RenderCore/Private/RenderUtils.cpp
@@ -553,6 +553,12 @@ RENDERCORE_API uint32 GetPlatformShadingModelsMask(const FStaticShaderPlatform P
 
 using ShaderPlatformMaskType = TStaticBitArray<EShaderPlatform::SP_NumPlatforms>;
 
+RENDERCORE_API bool IsMobileAmbientOcclusionEnabled(const FStaticShaderPlatform Platform)
+{
+	static FShaderPlatformCachedIniValue<bool> MobileAmbientOcclusionIniValue(TEXT("r.Mobile.AmbientOcclusion"));
+	return IsMobilePlatform(Platform) && MobileAmbientOcclusionIniValue.Get(Platform);
+}
+
 RENDERCORE_API bool AreMobileScreenSpaceReflectionsEnabled(const FStaticShaderPlatform Platform)
 {
 	static FShaderPlatformCachedIniValue<bool> MobileScreenSpaceReflectionsIniValue(TEXT("r.Mobile.ScreenSpaceReflections"));
@@ -639,7 +645,7 @@ RENDERCORE_API bool MobileUsesFullDepthPrepass(const FStaticShaderPlatform Platf
 		if (!IsMobileDeferredShadingEnabled(Platform) || MobileAllowFramebufferFetch(Platform))
 		{
 			// Features that require full depth before lighting
-			return MobileUsesShadowMaskTexture(Platform);
+			return MobileUsesShadowMaskTexture(Platform) || IsMobileAmbientOcclusionEnabled(Platform);
 		}
 	}
 	return bForcedDepthPrepass;
diff --git a/Engine/Source/Runtime/RenderCore/Public/RenderUtils.h b/Engine/Source/Runtime/RenderCore/Public/RenderUtils.h
--- a/Engine/Source/Runtime/RenderCore/Public/RenderUtils.h
+++ b/Engine/Source/Runtime/RenderCore/Public/RenderUtils.h
@@ -303,6 +303,9 @@ RENDERCORE_API bool AllowPerPixelShadingModels(const FStaticShaderPlatform Platf
 
 RENDERCORE_API uint32 GetPlatformShadingModelsMask(const FStaticShaderPlatform Platform);
 
+/** Returns whether AO is enabled. Only valid for a mobile forward shading. */
+RENDERCORE_API bool IsMobileAmbientOcclusionEnabled(const FStaticShaderPlatform Platform);
+
 RENDERCORE_API bool AreMobileScreenSpaceReflectionsEnabled(const FStaticShaderPlatform Platform);
 
 RENDERCORE_API bool IsMobileDistanceFieldEnabled(const FStaticShaderPlatform Platform);
diff --git a/Engine/Source/Runtime/Renderer/Private/CompositionLighting/PostProcessAmbientOcclusion.cpp b/Engine/Source/Runtime/Renderer/Private/CompositionLighting/PostProcessAmbientOcclusion.cpp
--- a/Engine/Source/Runtime/Renderer/Private/CompositionLighting/PostProcessAmbientOcclusion.cpp
+++ b/Engine/Source/Runtime/Renderer/Private/CompositionLighting/PostProcessAmbientOcclusion.cpp
@@ -1648,9 +1648,6 @@ FScreenPassTexture AddGTAOUpsamplePass(
 	const FScreenPassTextureViewport InputViewport(Input);
 	const FScreenPassTextureViewport OutputViewport(Output);
 
-	// Mobile path writes to texture alpha channel only
-	FRHIBlendState* BlendState = (View.GetFeatureLevel() == ERHIFeatureLevel::ES3_1 ? TStaticBlendState<CW_ALPHA>::GetRHI() : TStaticBlendState<>::GetRHI());
-
 	// Pixel Shader
 	FGTAOUpsamplePS::FParameters* PassParameters = GraphBuilder.AllocParameters<FGTAOUpsamplePS::FParameters>();
 
@@ -1670,7 +1667,7 @@ FScreenPassTexture AddGTAOUpsamplePass(
 		InputViewport,
 		VertexShader,
 		PixelShader,
-		BlendState,
+		TStaticBlendState<>::GetRHI(),
 		TStaticDepthStencilState<false, CF_Always>::GetRHI(),
 		PassParameters);
 
diff --git a/Engine/Source/Runtime/Renderer/Private/MobileBasePassRendering.cpp b/Engine/Source/Runtime/Renderer/Private/MobileBasePassRendering.cpp
--- a/Engine/Source/Runtime/Renderer/Private/MobileBasePassRendering.cpp
+++ b/Engine/Source/Runtime/Renderer/Private/MobileBasePassRendering.cpp
@@ -311,6 +311,17 @@ void SetupMobileBasePassUniformParameters(
 	const bool bIsMobileMultiViewEnabled = View.bIsMobileMultiViewEnabled || View.Aspects.IsMobileMultiViewEnabled();
 	BasePassParameters.DBuffer = GetDBufferParameters(GraphBuilder, MobileBasePassTextures.DBufferTextures, View.GetShaderPlatform(), bIsMobileMultiViewEnabled);
 
+	if (HasBeenProduced(MobileBasePassTextures.AmbientOcclusionTexture))
+	{
+		BasePassParameters.AmbientOcclusionTexture = MobileBasePassTextures.AmbientOcclusionTexture;
+	}
+	else
+	{
+		BasePassParameters.AmbientOcclusionTexture = SystemTextures.White;
+	}
+	BasePassParameters.AmbientOcclusionSampler = TStaticSamplerState<SF_Point, AM_Clamp, AM_Clamp, AM_Clamp>::GetRHI();
+	BasePassParameters.AmbientOcclusionStaticFraction = FMath::Clamp(View.FinalPostProcessSettings.AmbientOcclusionStaticFraction, 0.0f, 1.0f);
+
 	if (HasBeenProduced(MobileBasePassTextures.ScreenSpaceShadowMaskTexture))
 	{
 		BasePassParameters.ScreenSpaceShadowMaskTexture = MobileBasePassTextures.ScreenSpaceShadowMaskTexture;
@@ -321,7 +332,6 @@ void SetupMobileBasePassUniformParameters(
 		BasePassParameters.ScreenSpaceShadowMaskTexture = SystemTextures.White;
 		BasePassParameters.ScreenSpaceShadowMaskTextureArray = GSystemTextures.GetDefaultTexture(GraphBuilder, ETextureDimension::Texture2DArray, PF_DepthStencil, FClearValueBinding::White);
 	}
-	BasePassParameters.AmbientOcclusionStaticFraction = FMath::Clamp(View.FinalPostProcessSettings.AmbientOcclusionStaticFraction, 0.0f, 1.0f);
 	BasePassParameters.ScreenSpaceShadowMaskSampler = TStaticSamplerState<SF_Point, AM_Clamp, AM_Clamp, AM_Clamp>::GetRHI();
 
 	FRDGBuffer* OcclusionBuffer = View.ViewState ? View.ViewState->OcclusionFeedback.GetGPUFeedbackBuffer() : nullptr;
diff --git a/Engine/Source/Runtime/Renderer/Private/MobileBasePassRendering.h b/Engine/Source/Runtime/Renderer/Private/MobileBasePassRendering.h
--- a/Engine/Source/Runtime/Renderer/Private/MobileBasePassRendering.h
+++ b/Engine/Source/Runtime/Renderer/Private/MobileBasePassRendering.h
@@ -38,6 +38,7 @@ struct FMobileBasePassTextures
 {
 	FDBufferTextures DBufferTextures = {};
 	FRDGTextureRef ScreenSpaceShadowMaskTexture = nullptr;
+	FRDGTextureRef AmbientOcclusionTexture = nullptr;
 };
 
 BEGIN_GLOBAL_SHADER_PARAMETER_STRUCT(FMobileBasePassUniformParameters, )
@@ -56,6 +57,8 @@ BEGIN_GLOBAL_SHADER_PARAMETER_STRUCT(FMobileBasePassUniformParameters, )
 	SHADER_PARAMETER_SAMPLER(SamplerState, PreIntegratedGFSampler)
 	SHADER_PARAMETER_RDG_BUFFER_SRV(StructuredBuffer<float4>, EyeAdaptationBuffer)
 	SHADER_PARAMETER_RDG_BUFFER_UAV(RWStructuredBuffer<uint>, RWOcclusionBufferUAV)
+	SHADER_PARAMETER_RDG_TEXTURE(Texture2D, AmbientOcclusionTexture)
+	SHADER_PARAMETER_SAMPLER(SamplerState, AmbientOcclusionSampler)
 	SHADER_PARAMETER_RDG_TEXTURE(Texture2D, ScreenSpaceShadowMaskTexture)
 	SHADER_PARAMETER_RDG_TEXTURE(Texture2DArray, ScreenSpaceShadowMaskTextureArray)
 	SHADER_PARAMETER_SAMPLER(SamplerState, ScreenSpaceShadowMaskSampler)
@@ -402,6 +405,7 @@ public:
 		// Only non-static skylights contribute into forward basepass
 		const bool bProjectSupportsNonStaticSkyLights = FReadOnlyCVARCache::EnableStationarySkylight() || !IsStaticLightingAllowed();
 		OutEnvironment.SetDefine(TEXT("ENABLE_SKY_LIGHT"), bIsLit && bForwardShading && bProjectSupportsNonStaticSkyLights);
+		OutEnvironment.SetDefine(TEXT("ENABLE_AMBIENT_OCCLUSION"), bForwardShading && IsMobileAmbientOcclusionEnabled(Parameters.Platform) ? 1u : 0u);
 		
 		FForwardLightingParameters::ModifyCompilationEnvironment(Parameters.Platform, OutEnvironment);
 		OutEnvironment.SetDefine(TEXT("ENABLE_CLUSTERED_LIGHTS"), (LocalLightSetting == EMobileLocalLightSetting::LOCAL_LIGHTS_ENABLED) ? 1u : 0u);
diff --git a/Engine/Source/Runtime/Renderer/Private/MobileDeferredShadingPass.cpp b/Engine/Source/Runtime/Renderer/Private/MobileDeferredShadingPass.cpp
--- a/Engine/Source/Runtime/Renderer/Private/MobileDeferredShadingPass.cpp
+++ b/Engine/Source/Runtime/Renderer/Private/MobileDeferredShadingPass.cpp
@@ -37,6 +37,12 @@ static bool UseClusteredDeferredShading(const FStaticShaderPlatform Platform)
 	return GMobileUseClusteredDeferredShading != 0 && MobileForwardEnableLocalLights(Platform);
 }
 
+static bool MobileDeferredEnableAmbientOcclusion(const FStaticShaderPlatform Platform)
+{
+	// AO requires a full depth before shading
+	return MobileUsesFullDepthPrepass(Platform) || !MobileAllowFramebufferFetch(Platform);
+}
+
 int32 GMobileUseLightStencilCulling = 0;
 static FAutoConsoleVariableRef CVarMobileUseLightStencilCulling(
 	TEXT("r.Mobile.UseLightStencilCulling"),
@@ -122,7 +128,7 @@ class FMobileDirectionalLightFunctionPS : public FMaterialShader
 
 		const bool bMobileForceDepthRead = MobileUsesFullDepthPrepass(Parameters.Platform);
 		OutEnvironment.SetDefine(TEXT("FORCE_DEPTH_TEXTURE_READS"), bMobileForceDepthRead ? 1u : 0u);
-
+		OutEnvironment.SetDefine(TEXT("ENABLE_AMBIENT_OCCLUSION"), MobileDeferredEnableAmbientOcclusion(Parameters.Platform) ? 1u : 0u);
 		OutEnvironment.SetDefine(TEXT("MOBILE_SSR_ENABLED"), PermutationVector.Get<FMobileSSRQuality>() != EMobileSSRQuality::Disabled ? 1u : 0u);
 	}
 
@@ -396,6 +402,7 @@ public:
 
 		const bool bMobileForceDepthRead = MobileUsesFullDepthPrepass(Parameters.Platform);
 		OutEnvironment.SetDefine(TEXT("FORCE_DEPTH_TEXTURE_READS"), bMobileForceDepthRead ? 1u : 0u);
+		OutEnvironment.SetDefine(TEXT("ENABLE_AMBIENT_OCCLUSION"), MobileDeferredEnableAmbientOcclusion(Parameters.Platform) ? 1u : 0u);
 		OutEnvironment.SetDefine(TEXT("MOBILE_SSR_ENABLED"), PermutationVector.Get<FMobileSSRQuality>() != EMobileSSRQuality::Disabled ? 1u : 0u);
 	}
 };
diff --git a/Engine/Source/Runtime/Renderer/Private/MobileShadingRenderer.cpp b/Engine/Source/Runtime/Renderer/Private/MobileShadingRenderer.cpp
--- a/Engine/Source/Runtime/Renderer/Private/MobileShadingRenderer.cpp
+++ b/Engine/Source/Runtime/Renderer/Private/MobileShadingRenderer.cpp
@@ -1448,6 +1448,11 @@ void FMobileSceneRenderer::Render(FRDGBuilder& GraphBuilder, const FSceneRenderU
 			RenderHZB(GraphBuilder, SceneTextures.Depth.Resolve);
 		}
 
+		// When renderer is in ERendererOutput::DepthPrepassOnly mode, bRequiresAmbientOcclusionPass is set to false in InitViews()
+		if (bRendererOutputFinalSceneColor)
+		{
+			BasePassTextures.AmbientOcclusionTexture = RenderAmbientOcclusion(GraphBuilder, SceneTextures);
+		}
 
 		if (bEnableDistanceFieldAO && (!bDeferredShading || !bRequiresMultiPass))
 		{
@@ -1468,12 +1473,6 @@ void FMobileSceneRenderer::Render(FRDGBuilder& GraphBuilder, const FSceneRenderU
 			BasePassTextures.ScreenSpaceShadowMaskTexture = RenderMobileShadowProjections(GraphBuilder, SceneTextures.Config.Extent, bIsMobileMultiView);
 		}
 
-		// When renderer is in ERendererOutput::DepthPrepassOnly mode, bRequiresAmbientOcclusionPass is set to false in InitViews()
-		if (bRendererOutputFinalSceneColor)
-		{
-			RenderAmbientOcclusion(GraphBuilder, SceneTextures, BasePassTextures.ScreenSpaceShadowMaskTexture);
-		}
-
 		// Local Light prepass
 		if (bRendererOutputFinalSceneColor)
 		{
@@ -2548,6 +2547,11 @@ void FMobileSceneRenderer::RenderDeferredMultiPass(FRDGBuilder& GraphBuilder, FS
 		}
 	}
 
+	if (!bIsFullDepthPrepassEnabled)
+	{
+		BasePassTextures.AmbientOcclusionTexture = RenderAmbientOcclusion(GraphBuilder, SceneTextures);
+	}
+
 	TArray<FRDGTextureRef> DynamicBentNormalAOTextures;
 	if (bEnableDistanceFieldAO)
 	{
@@ -2566,11 +2570,6 @@ void FMobileSceneRenderer::RenderDeferredMultiPass(FRDGBuilder& GraphBuilder, FS
 		BasePassTextures.ScreenSpaceShadowMaskTexture = RenderMobileShadowProjections(GraphBuilder, SceneTextures.Config.Extent, bIsMobileMultiView);
 	}
 
-	if (!bIsFullDepthPrepassEnabled)
-	{
-		RenderAmbientOcclusion(GraphBuilder, SceneTextures, BasePassTextures.ScreenSpaceShadowMaskTexture);
-	}
-
 	const bool bSupportsSM5Nodes = MobileSupportsSM5MaterialNodes(ShaderPlatform);
 	const bool bShouldRenderFullSLW = bSupportsSM5Nodes && ShouldRenderSingleLayerWater(Views);
 	const bool bNeedsDepthCopy = bShouldRenderFullSLW;
diff --git a/Engine/Source/Runtime/Renderer/Private/PostProcess/PostProcessAmbientOcclusionMobile.cpp b/Engine/Source/Runtime/Renderer/Private/PostProcess/PostProcessAmbientOcclusionMobile.cpp
--- a/Engine/Source/Runtime/Renderer/Private/PostProcess/PostProcessAmbientOcclusionMobile.cpp
+++ b/Engine/Source/Runtime/Renderer/Private/PostProcess/PostProcessAmbientOcclusionMobile.cpp
@@ -20,6 +20,15 @@
 #include "ClearQuad.h"
 #include "PixelShaderUtils.h"
 
+static TAutoConsoleVariable<int32> CVarMobileAmbientOcclusion(
+	TEXT("r.Mobile.AmbientOcclusion"),
+	0,
+	TEXT("Caution: Affects only mobile Forward shading. An extra sampler will be occupied in mobile base pass pixel shader after enable the mobile ambient occlusion.\n")
+	TEXT("0: Disable Ambient Occlusion on mobile platform. [default]\n")
+	TEXT("1: Enable Ambient Occlusion on mobile platform.\n"),
+	ECVF_ReadOnly | ECVF_RenderThreadSafe
+);
+
 static TAutoConsoleVariable<int32> CVarMobileAmbientOcclusionTechnique(
 	TEXT("r.Mobile.AmbientOcclusionTechnique"),
 	1,
@@ -91,15 +100,6 @@ static bool IsTwoPassShaderType(EMobileAmbientOcclusionShaderType ShaderType)
 	return ShaderType != EMobileAmbientOcclusionShaderType::ComputeShader;
 }
 
-static FRHIBlendState* GetShadowMaskBlendState(FRDGTextureRef ScreenSpaceShadowMaskTexture)
-{
-	if (ScreenSpaceShadowMaskTexture->Desc.Format == PF_R8G8)
-	{
-		return TStaticBlendState<CW_GREEN>::GetRHI();
-	}
-	return TStaticBlendState<CW_ALPHA>::GetRHI();
-}
-
 static TAutoConsoleVariable<int32> CVarMobileAmbientOcclusionDepthBoundsTest(
 	TEXT("r.Mobile.AmbientOcclusionDepthBoundsTest"),
 	1,
@@ -455,26 +455,7 @@ public:
 
 IMPLEMENT_GLOBAL_SHADER(FGTAOMobile_SpatialFilterPS, "/Engine/Private/PostProcessAmbientOcclusionMobile.usf", "GTAOSpatialFilterPS", SF_Pixel);
 
-static EPixelFormat GetMobileScreenSpaceAOTexturePixelFormat(FStaticShaderPlatform ShaderPlatform)
-{
-	EPixelFormat Format = PF_G8;
-
-	bool bGTAO = (CVarMobileAmbientOcclusionTechnique.GetValueOnRenderThread() == 0);
-	const EMobileAmbientOcclusionShaderType ShaderType = GetCVarMobileAmbientOcclusionShaderTypeOnRenderThread();
-
-	// G8 isn't supported as UAV on Android OpenGLES, fall back to RGBA8 for compute shader usage.
-	if (bGTAO
-		&& IsOpenGLPlatform(ShaderPlatform)
-		&& IsComputeShaderType(ShaderType)
-		)
-	{
-		Format = PF_R8G8B8A8;
-	}
-
-	return Format;
-}
-
-static void RenderGTAO(FRDGBuilder& GraphBuilder, const FSceneTextures& SceneTextures, FRDGTextureRef ScreenSpaceShadowMaskTexture, const TArray<FViewInfo>& Views)
+static void RenderGTAO(FRDGBuilder& GraphBuilder, const FSceneTextures& SceneTextures, FRDGTextureRef AmbientOcclusionTexture, const TArray<FViewInfo>& Views)
 {
 	RDG_EVENT_SCOPE_STAT(GraphBuilder, MobileGTAO, "MobileGTAO");
 	RDG_GPU_STAT_SCOPE(GraphBuilder, MobileGTAO);
@@ -490,7 +471,7 @@ static void RenderGTAO(FRDGBuilder& GraphBuilder, const FSceneTextures& SceneTex
 		
 	FRDGTextureRef SceneDepthTexture = SceneTextures.Depth.Resolve;
 	const FIntPoint DepthBufferSize = SceneDepthTexture->Desc.Extent;
-	const FIntPoint BufferSize = FIntPoint::DivideAndRoundUp(ScreenSpaceShadowMaskTexture->Desc.Extent, DownsampleFactor);
+	const FIntPoint BufferSize = AmbientOcclusionTexture->Desc.Extent;
 
 	float FallOffEnd = GTAOFalloffEndCVar ? GTAOFalloffEndCVar->GetValueOnRenderThread() : 200.0f;
 	float FallOffStartRatio = GTAOFalloffStartRatioCVar ? FMath::Clamp(GTAOFalloffStartRatioCVar->GetValueOnRenderThread(), 0.0f, 0.999f) : 0.5f;
@@ -511,10 +492,7 @@ static void RenderGTAO(FRDGBuilder& GraphBuilder, const FSceneTextures& SceneTex
 	FMath::SinCos(&SinDeltaAngle, &CosDeltaAngle, PI / NumAngles);
 
 	const EMobileAmbientOcclusionShaderType ShaderType = GetCVarMobileAmbientOcclusionShaderTypeOnRenderThread();
-
-	ETextureCreateFlags AOCreateFlags = TexCreate_ShaderResource | (IsComputeShaderType(ShaderType) ? TexCreate_UAV : TexCreate_RenderTargetable);
-	const EPixelFormat AOPixelFormat = GetMobileScreenSpaceAOTexturePixelFormat(SceneTextures.Config.ShaderPlatform);
-	FRDGTextureRef AmbientOcclusionTexture = GraphBuilder.CreateTexture(FRDGTextureDesc::Create2D(BufferSize, AOPixelFormat, FClearValueBinding::None, AOCreateFlags), TEXT("HalfResolutionScreenSpaceAO"));
+	
 	FRDGTextureUAVRef AmbientOcclusionTextureUAV{};
 	if (IsComputeShaderType(ShaderType))
 	{
@@ -546,8 +524,6 @@ static void RenderGTAO(FRDGBuilder& GraphBuilder, const FSceneTextures& SceneTex
 		const FIntRect& ViewRect = FIntRect::DivideAndRoundUp(View.ViewRect, DownsampleFactor);
 		FVector4f WorldRadiusAdjSinCosDeltaAngleThickness(FallOffStartEndScaleBias.Y * DepthBufferSize.Y * View.ViewMatrices.GetProjectionMatrix().M[0][0], SinDeltaAngle, CosDeltaAngle, ThicknessBlend);
 
-		FScreenPassRenderTarget AmbientOcclusionOutput = FScreenPassRenderTarget(AmbientOcclusionTexture, ViewRect, ERenderTargetLoadAction::ENoAction);
-
 		if (ShaderType == EMobileAmbientOcclusionShaderType::ComputeShader)
 		{
 			FGTAOMobile_HorizonSearchIntegralSpatialFilterCS::FParameters* HorizonSearchIntegralSpatialFilterParameters = GraphBuilder.AllocParameters<FGTAOMobile_HorizonSearchIntegralSpatialFilterCS::FParameters>();
@@ -618,10 +594,11 @@ static void RenderGTAO(FRDGBuilder& GraphBuilder, const FSceneTextures& SceneTex
 					EDRF_UseTriangleOptimization);
 			});
 
-			
+			FScreenPassRenderTarget AmbientOcclusionRT(AmbientOcclusionTexture, ViewRect, ViewIndex > 0 ? ERenderTargetLoadAction::ELoad : ERenderTargetLoadAction::EClear);
+
 			FGTAOMobile_SpatialFilterPS::FParameters* SpatialFilterParameters = GraphBuilder.AllocParameters<FGTAOMobile_SpatialFilterPS::FParameters>();
 			FGTAOMobile_SpatialFilter::SetupShaderParameters(SpatialFilterParameters->Common, View, ViewRect, BufferSize, HorizonSearchIntegralTexture);
-			SpatialFilterParameters->RenderTargets[0] = AmbientOcclusionOutput.GetRenderTargetBinding();
+			SpatialFilterParameters->RenderTargets[0] = AmbientOcclusionRT.GetRenderTargetBinding();
 
 			TShaderMapRef<FGTAOMobile_SpatialFilterPS> SpatialFilterShader(View.ShaderMap);
 
@@ -695,13 +672,6 @@ static void RenderGTAO(FRDGBuilder& GraphBuilder, const FSceneTextures& SceneTex
 				SpatialFilterParameters,
 				FComputeShaderUtils::GetGroupCount(ViewRect.Size(), FGTAOMobile_SpatialFilterCS::TexelsPerThreadGroup));
 		}
-
-		check(DownsampleFactor > 1.0);
-		{
-			FScreenPassRenderTarget ScreenSpaceShadowMaskRT(ScreenSpaceShadowMaskTexture, View.ViewRect, ERenderTargetLoadAction::ELoad);
-			FScreenPassTexture SceneDepth(SceneDepthTexture);
-			AddGTAOUpsamplePass(GraphBuilder, View, AmbientOcclusionOutput, SceneDepth, ScreenSpaceShadowMaskRT);
-		}
 	}
 }
 
@@ -995,15 +965,7 @@ static void AddMobileAmbientOcclusionPass(
 			FGraphicsPipelineStateInitializer GraphicsPSOInit;
 			RHICmdList.ApplyCachedRenderTargets(GraphicsPSOInit);
 
-			if (bHalfResolution)
-			{
-				GraphicsPSOInit.BlendState = TStaticBlendState<>::GetRHI();
-			}
-			else
-			{
-				// Writing to a ScreenSpaceShadowMask
-				GraphicsPSOInit.BlendState = GetShadowMaskBlendState(Output.Texture);
-			}
+			GraphicsPSOInit.BlendState = TStaticBlendState<>::GetRHI();
 			GraphicsPSOInit.RasterizerState = TStaticRasterizerState<FM_Solid, CM_None>::GetRHI();
 			GraphicsPSOInit.DepthStencilState = TStaticDepthStencilState<false, CF_Always>::GetRHI();
 			GraphicsPSOInit.BoundShaderState.VertexDeclarationRHI = GFilterVertexDeclaration.VertexDeclarationRHI;
@@ -1084,8 +1046,6 @@ static void AddMobileAmbientOcclusionUpsamplePass(
 				Callbacks.PSOModifier = [Output](FGraphicsPipelineStateInitializer& PSOInitializer)
 				{
 					PSOInitializer.bDepthBounds = true;
-					// Writing to a ScreenSpaceShadowMask
-					PSOInitializer.BlendState = GetShadowMaskBlendState(Output.Texture);
 				};
 				Callbacks.PreClear = [DepthFar](FRHICommandList& InRHICmdList)
 				{
@@ -1136,10 +1096,6 @@ static void AddMobileAmbientOcclusionUpsamplePass(
 			FGraphicsPipelineStateInitializer GraphicsPSOInit;
 			FPixelShaderUtils::InitFullscreenPipelineState(RHICmdList, View.ShaderMap, PixelShader, GraphicsPSOInit);
 			GraphicsPSOInit.bDepthBounds = bDepthBoundsTestEnabled;
-
-			// Writing to a ScreenSpaceShadowMask
-			GraphicsPSOInit.BlendState = GetShadowMaskBlendState(Output.Texture);
-
 			SetGraphicsPipelineState(RHICmdList, GraphicsPSOInit, 0);
 			SetShaderParameters(RHICmdList, PixelShader, PixelShader.GetPixelShader(), *PassParameters);
 
@@ -1157,7 +1113,7 @@ static void AddMobileAmbientOcclusionUpsamplePass(
 		});
 }
 
-static void RenderSSAO(FRDGBuilder& GraphBuilder, const FSceneTextures& SceneTextures, FRDGTextureRef ScreenSpaceShadowMaskTexture, const TArray<FViewInfo>& Views)
+static void RenderSSAO(FRDGBuilder& GraphBuilder, const FSceneTextures& SceneTextures, FRDGTextureRef AmbientOcclusionTexture, const TArray<FViewInfo>& Views)
 {
 	RDG_EVENT_SCOPE_STAT(GraphBuilder, MobileSSAO, "MobileSSAO");
 	RDG_GPU_STAT_SCOPE(GraphBuilder, MobileSSAO);
@@ -1172,8 +1128,8 @@ static void RenderSSAO(FRDGBuilder& GraphBuilder, const FSceneTextures& SceneTex
 	if (bHalfResolution)
 	{
 		// Bilateral requires 32bit format for AO + depth.
-		const EPixelFormat Format = bBilateralUpsample ? PF_R8G8B8A8 : GetMobileScreenSpaceAOTexturePixelFormat(SceneTextures.Config.ShaderPlatform);
-		const FIntPoint Extent = GetDownscaledExtent(ScreenSpaceShadowMaskTexture->Desc.Extent, 2);
+		const EPixelFormat Format = bBilateralUpsample ? PF_R8G8B8A8 : PF_R8;
+		const FIntPoint Extent = GetDownscaledExtent(AmbientOcclusionTexture->Desc.Extent, 2);
 		HalfResolutionTexture = GraphBuilder.CreateTexture(
 			FRDGTextureDesc::Create2D(Extent, Format, FClearValueBinding::None, TexCreate_ShaderResource | TexCreate_RenderTargetable),
 			TEXT("HalfResolutionScreenSpaceAO"));
@@ -1182,7 +1138,7 @@ static void RenderSSAO(FRDGBuilder& GraphBuilder, const FSceneTextures& SceneTex
 	for (FViewInfo const& View : Views)
 	{
 		FMobileSSAOCommonParameters Parameters = GetMobileSSAOCommonParameters(GraphBuilder, View, SceneDepthTexture, SceneTextures.MobileUniformBuffer);
-		FScreenPassRenderTarget FinalTarget = FScreenPassRenderTarget(ScreenSpaceShadowMaskTexture, View.ViewRect, ERenderTargetLoadAction::ELoad);
+		FScreenPassRenderTarget FinalTarget = FScreenPassRenderTarget(AmbientOcclusionTexture, View.ViewRect, ERenderTargetLoadAction::ENoAction);
 
 		if (HalfResolutionTexture != nullptr)
 		{
@@ -1211,23 +1167,54 @@ bool RequiresMobileAmbientOcclusionPass(FScene* Scene, const FViewInfo& View)
 		&& !View.Family->UseDebugViewPS();
 }
 
+FRDGTextureRef CreateMobileScreenSpaceAOTexture(FRDGBuilder& GraphBuilder, const FSceneTexturesConfig& Config)
+{
+	bool bGTAO = (CVarMobileAmbientOcclusionTechnique.GetValueOnRenderThread() == 0);
+	const uint32 DownsampleFactor = bGTAO ? 2 : 1;
+	const FIntPoint Extent = FIntPoint::DivideAndRoundUp(Config.Extent, DownsampleFactor);
+	EPixelFormat Format = PF_R8;
+
+	const EMobileAmbientOcclusionShaderType ShaderType = GetCVarMobileAmbientOcclusionShaderTypeOnRenderThread();
+	// G8 isn't supported as UAV on Android OpenGLES, fall back to RGBA8 for compute shader usage.
+	if (bGTAO
+		&& IsOpenGLPlatform(Config.ShaderPlatform)
+		&& IsComputeShaderType(ShaderType)
+		)
+	{
+		Format = PF_R8G8B8A8;
+	}
+	ETextureCreateFlags TextureCreateFlags = TexCreate_ShaderResource | TexCreate_RenderTargetable;
+	if (IsComputeShaderType(ShaderType))
+	{
+		TextureCreateFlags |= TexCreate_UAV;
+	}
+	if (UE::PixelFormat::HasCapabilities(Format, EPixelFormatCapabilities::LossyCompressible) && CVarMobileAmbientOcclusionCompression.GetValueOnRenderThread() != 0)
+	{
+		TextureCreateFlags |= TexCreate_LossyCompressionLowBitrate;
+	}
+	return GraphBuilder.CreateTexture(FRDGTextureDesc::Create2D(Extent, Format, FClearValueBinding::Black, TextureCreateFlags), TEXT("ScreenSpaceAO"));
+}
+
 // --------------------------------------------------------------------------------------------------------------------
-void FMobileSceneRenderer::RenderAmbientOcclusion(FRDGBuilder& GraphBuilder, const FSceneTextures& SceneTextures, FRDGTextureRef ScreenSpaceShadowMaskTexture)
+FRDGTextureRef FMobileSceneRenderer::RenderAmbientOcclusion(FRDGBuilder& GraphBuilder, const FSceneTextures& SceneTextures)
 {
-	if (!RequiresMobileAmbientOcclusionPass(Scene, Views[0]) 
-		|| ScreenSpaceShadowMaskTexture == nullptr)
+	if (!RequiresMobileAmbientOcclusionPass(Scene, Views[0]))
 	{
-		return;
+		return nullptr;
 	}
 
+	FRDGTextureRef AmbientOcclusionTexture = CreateMobileScreenSpaceAOTexture(GraphBuilder, SceneTextures.Config);
+
 	const int32 Technique = CVarMobileAmbientOcclusionTechnique.GetValueOnRenderThread();
 	switch (Technique)
 	{
 	case 0:
-		RenderGTAO(GraphBuilder, SceneTextures, ScreenSpaceShadowMaskTexture, Views);
+		RenderGTAO(GraphBuilder, SceneTextures, AmbientOcclusionTexture, Views);
 		break;
 	case 1:
-		RenderSSAO(GraphBuilder, SceneTextures, ScreenSpaceShadowMaskTexture, Views);
+		RenderSSAO(GraphBuilder, SceneTextures, AmbientOcclusionTexture, Views);
 		break;
 	}
+
+	return AmbientOcclusionTexture;
 }
diff --git a/Engine/Source/Runtime/Renderer/Private/SceneRendering.h b/Engine/Source/Runtime/Renderer/Private/SceneRendering.h
--- a/Engine/Source/Runtime/Renderer/Private/SceneRendering.h
+++ b/Engine/Source/Runtime/Renderer/Private/SceneRendering.h
@@ -2923,7 +2923,7 @@ protected:
 	void RenderStereoMotionVectors(FRDGBuilder& GraphBuilder, FViewInfo& View, const FSceneTextures& SceneTextures, FInstanceCullingManager& InstanceCullingManager);
 	bool HasStereoMotionVectorRenderTargets(const FSceneTextures& SceneTextures);
 
-	void RenderAmbientOcclusion(FRDGBuilder& GraphBuilder, const FSceneTextures& SceneTextures, FRDGTextureRef ScreenSpaceShadowMaskTexture);
+	FRDGTextureRef RenderAmbientOcclusion(FRDGBuilder& GraphBuilder, const FSceneTextures& SceneTextures);
 
 	FRDGTextureRef RenderMobileShadowProjections(FRDGBuilder& GraphBuilder, FIntPoint Extent, bool bIsMobileMultiView);
 
diff --git a/Engine/Source/Runtime/Renderer/Private/ShadowRendering.cpp b/Engine/Source/Runtime/Renderer/Private/ShadowRendering.cpp
--- a/Engine/Source/Runtime/Renderer/Private/ShadowRendering.cpp
+++ b/Engine/Source/Runtime/Renderer/Private/ShadowRendering.cpp
@@ -2432,8 +2432,8 @@ FRDGTextureRef FMobileSceneRenderer::RenderMobileShadowProjections(FRDGBuilder&
 	EPixelFormat ShadowMaskPixelFormat = PF_B8G8R8A8;
 	if (!IsStaticLightingAllowed())
 	{
-		// With pre-computed lighting disabled all lights use single channel. And second one for AO
-		ShadowMaskPixelFormat = PF_R8G8;
+		// With pre-computed lighting disabled all lights use single channel
+		ShadowMaskPixelFormat = PF_R8;
 	}
 	
 	FRDGTextureDesc TargetDesc = FRDGTextureDesc::CreateRenderTargetTextureDesc(Extent, ShadowMaskPixelFormat, FClearValueBinding::White, TexCreate_RenderTargetable | TexCreate_ShaderResource, bIsMobileMultiView);
