Moved ReleasePendingSpaces over to BeginUpdate due to avoid race condition with view uniform buffer parameters call to GetGlobalMipBias.

diff --git a/Engine/Source/Runtime/Renderer/Private/VT/VirtualTextureSystem.cpp b/Engine/Source/Runtime/Renderer/Private/VT/VirtualTextureSystem.cpp
--- a/Engine/Source/Runtime/Renderer/Private/VT/VirtualTextureSystem.cpp
+++ b/Engine/Source/Runtime/Renderer/Private/VT/VirtualTextureSystem.cpp
@@ -2958,12 +2958,12 @@ TUniquePtr<FVirtualTextureUpdater> FVirtualTextureSystem::BeginUpdate(FRDGBuilde
 	}
 
 	DestroyPendingVirtualTextures(false);
+	ReleasePendingSpaces();
 
 	// Early out when no allocated VTs
 	if (AllocatedVTs.Num() == 0)
 	{
 		MappedTilesToProduce.Reset();
-		ReleasePendingSpaces();
 		return {};
 	}
 
@@ -3058,8 +3058,6 @@ void FVirtualTextureSystem::EndUpdate(FRDGBuilder& GraphBuilder, TUniquePtr<FVir
 	UpdateCsvStats();
 #endif
 
-	ReleasePendingSpaces();
-
 	TilesToLock = MoveTemp(TilesToLockForNextFrame);
 }
 
