Mobile: Moved all reflections into a separate global shading pass instead of adding permutations to the main light PS.
Reduced main light permutations by ~10x.

diff --git a/Engine/Shaders/Private/MobileDeferredShading.usf b/Engine/Shaders/Private/MobileDeferredShading.usf
--- a/Engine/Shaders/Private/MobileDeferredShading.usf
+++ b/Engine/Shaders/Private/MobileDeferredShading.usf
@@ -202,7 +202,7 @@ void MobileDirectionalLightPS(
 	float4 ScreenPosition = SvPositionToScreenPosition(float4(SvPosition.xyz, GBuffer.Depth));
 
 	half GatheredAmbientOcclusion = 1;
-#if ENABLE_AMBIENT_OCCLUSION
+#if ENABLE_SKY_LIGHT && ENABLE_AMBIENT_OCCLUSION
 	GatheredAmbientOcclusion = Texture2DSample(MobileBasePass.AmbientOcclusionTexture, MobileBasePass.AmbientOcclusionSampler, UVAndScreenPos.xy).r;
 	#if ALLOW_STATIC_LIGHTING
 		GatheredAmbientOcclusion = lerp(1.0, GatheredAmbientOcclusion, MobileBasePass.AmbientOcclusionStaticFraction);
@@ -229,7 +229,7 @@ void MobileDirectionalLightPS(
 	}
 #endif
 
-#if ENABLE_CLUSTERED_REFLECTION || ENABLE_SKY_LIGHT || ENABLE_PLANAR_REFLECTION || MOBILE_SSR_ENABLED
+#if ENABLE_SKY_LIGHT
 	// If we have a single directional light, apply relfection and sky contrubution here
 	ReflectionEnvironmentSkyLighting(GBuffer, SvPosition, CameraVector, TranslatedWorldPosition, ReflectionVector, GridIndex, GatheredAmbientOcclusion, DirectLighting);
 #endif
@@ -436,7 +436,7 @@ void MobileReflectionEnvironmentSkyLightingPS(
 	uint GridIndex = ComputeLightGridCellIndex(uint2(LocalPosition.x, LocalPosition.y), GBuffer.Depth, EyeIndex);
 
 	half GatheredAmbientOcclusion = 1;
-#if ENABLE_AMBIENT_OCCLUSION
+#if ENABLE_SKY_LIGHT && ENABLE_AMBIENT_OCCLUSION
 	GatheredAmbientOcclusion = Texture2DSample(MobileBasePass.AmbientOcclusionTexture, MobileBasePass.AmbientOcclusionSampler, UVAndScreenPos.xy).r;
 	#if ALLOW_STATIC_LIGHTING
 		GatheredAmbientOcclusion = lerp(1.0, GatheredAmbientOcclusion, MobileBasePass.AmbientOcclusionStaticFraction);
diff --git a/Engine/Source/Runtime/Renderer/Private/MobileBasePass.cpp b/Engine/Source/Runtime/Renderer/Private/MobileBasePass.cpp
--- a/Engine/Source/Runtime/Renderer/Private/MobileBasePass.cpp
+++ b/Engine/Source/Runtime/Renderer/Private/MobileBasePass.cpp
@@ -66,23 +66,19 @@ EMobileLocalLightSetting GetMobileForwardLocalLightSetting(EShaderPlatform Shade
 	return EMobileLocalLightSetting::LOCAL_LIGHTS_DISABLED;
 }
 
-extern const uint8 MobileShadingModelSupportStencilValue = 0b01u;
-uint8 GetMobileShadingModelStencilValue(FMaterialShadingModelField ShadingModel, bool bFullyRough)
+uint8 GetMobileShadingModelStencilValue(FMaterialShadingModelField ShadingModel)
 {
-	// Bit 0 is set for materials that are receive SSR
-	// Bit 1 is set for DefaultLit materials (see MobileDeferredShadingPass.cpp)
-	const uint8 DefaultLitMask = bFullyRough ? 0b10u : 0b11u;
 	if (ShadingModel.HasOnlyShadingModel(MSM_DefaultLit))
 	{
-		return DefaultLitMask;
+		return GET_STENCIL_BIT_MASK(MOBILE_DEFAULTLIT, 1);
 	}
 	else if (ShadingModel.HasOnlyShadingModel(MSM_Unlit))
 	{
-		return 0b00u;
+		return 0u;
 	}
 
 	// mark everyhing as MSM_DefaultLit if GBuffer CustomData is not supported
-	return MobileUsesGBufferCustomData(GMaxRHIShaderPlatform) ? MobileShadingModelSupportStencilValue : DefaultLitMask;
+	return MobileUsesGBufferCustomData(GMaxRHIShaderPlatform) ? GET_STENCIL_BIT_MASK(MOBILE_SHADINGMODELS, 1) : GET_STENCIL_BIT_MASK(MOBILE_DEFAULTLIT, 1);
 }
 
 void SetMobileBasePassDepthState(
@@ -120,9 +116,7 @@ void SetMobileBasePassDepthState(
 
 	if (bUsesDeferredShading)
 	{
-		// store into [1-2] bits
-		uint8 ShadingModel = GetMobileShadingModelStencilValue(ShadingModels, Material.IsFullyRough());
-		StencilValue |= GET_STENCIL_MOBILE_SM_MASK(ShadingModel);
+		StencilValue |= GetMobileShadingModelStencilValue(ShadingModels);
 		StencilValue |= STENCIL_LIGHTING_CHANNELS_MASK(PrimitiveSceneProxy ? PrimitiveSceneProxy->GetLightingChannelStencilValue() : 0x00);
 	}
 	else
diff --git a/Engine/Source/Runtime/Renderer/Private/MobileDeferredShadingPass.cpp b/Engine/Source/Runtime/Renderer/Private/MobileDeferredShadingPass.cpp
--- a/Engine/Source/Runtime/Renderer/Private/MobileDeferredShadingPass.cpp
+++ b/Engine/Source/Runtime/Renderer/Private/MobileDeferredShadingPass.cpp
@@ -86,32 +86,25 @@ class FMobileDirectionalLightFunctionPS : public FMaterialShader
 
 	class FEnableShadingModelSupport	: SHADER_PERMUTATION_BOOL("ENABLE_SHADINGMODEL_SUPPORT_MOBILE_DEFERRED");
 	class FEnableClustredLights			: SHADER_PERMUTATION_BOOL("ENABLE_CLUSTERED_LIGHTS");
-	class FEnableClustredReflection		: SHADER_PERMUTATION_BOOL("ENABLE_CLUSTERED_REFLECTION");
-	class FEnablePlanarReflection		: SHADER_PERMUTATION_BOOL("ENABLE_PLANAR_REFLECTION");
 	class FEnableSkyLight				: SHADER_PERMUTATION_BOOL("ENABLE_SKY_LIGHT");
 	class FEnableScreenSpaceShadowMask  : SHADER_PERMUTATION_BOOL("ENABLE_SHADOWMASKTEXTURE");
 	class FEnableCSM					: SHADER_PERMUTATION_BOOL("ENABLE_MOBILE_CSM");
 	class FShadowQuality				: SHADER_PERMUTATION_RANGE_INT("MOBILE_SHADOW_QUALITY", 1, 3); // not using Quality=0
-	class FMobileSSRQuality 			: SHADER_PERMUTATION_ENUM_CLASS("MOBILE_SSR_QUALITY", EMobileSSRQuality);
 	class FSkyShadowing					: SHADER_PERMUTATION_BOOL("APPLY_SKY_SHADOWING");
 	class FInlineContactShadows			: SHADER_PERMUTATION_BOOL("SUPPORT_CONTACT_SHADOWS");
 
 	using FPermutationDomain = TShaderPermutationDomain<
 		FEnableShadingModelSupport,
 		FEnableClustredLights, 
-		FEnableClustredReflection, 
-		FEnablePlanarReflection,
 		FEnableSkyLight,
 		FEnableScreenSpaceShadowMask,
 		FEnableCSM, 
 		FShadowQuality,
-		FMobileSSRQuality,
 		FSkyShadowing,
 		FInlineContactShadows>;
 
 	BEGIN_SHADER_PARAMETER_STRUCT(FParameters, )
 		SHADER_PARAMETER_STRUCT_REF(FMobileDirectionalLightShaderParameters, MobileDirectionalLight)
-		SHADER_PARAMETER_STRUCT_REF(FMobileReflectionCaptureShaderData, MobileReflectionCaptureData)
 		SHADER_PARAMETER_STRUCT_INCLUDE(FMobileDeferredCommonParameters, MobileDeferredCommonParameters)
 		SHADER_PARAMETER_STRUCT_INCLUDE(FDFAOUpsampleParameters, DFAOUpsampleParameters)
 		SHADER_PARAMETER_STRUCT_INCLUDE(FSkyDiffuseLightingParameters, SkyDiffuseLighting)
@@ -129,12 +122,11 @@ class FMobileDirectionalLightFunctionPS : public FMaterialShader
 		const bool bMobileForceDepthRead = MobileUsesFullDepthPrepass(Parameters.Platform);
 		OutEnvironment.SetDefine(TEXT("FORCE_DEPTH_TEXTURE_READS"), bMobileForceDepthRead ? 1u : 0u);
 		OutEnvironment.SetDefine(TEXT("ENABLE_AMBIENT_OCCLUSION"), MobileDeferredEnableAmbientOcclusion(Parameters.Platform) ? 1u : 0u);
-		OutEnvironment.SetDefine(TEXT("MOBILE_SSR_ENABLED"), PermutationVector.Get<FMobileSSRQuality>() != EMobileSSRQuality::Disabled ? 1u : 0u);
 	}
 
 	static FPermutationDomain RemapPermutationVector(FPermutationDomain PermutationVector, EShaderPlatform Platform)
 	{
-		if (PermutationVector.Get<FEnableScreenSpaceShadowMask>() == true)
+		if (MobileUsesShadowMaskTexture(Platform) || PermutationVector.Get<FEnableScreenSpaceShadowMask>() == true)
 		{
 			PermutationVector.Set<FEnableCSM>(false);
 		}
@@ -149,11 +141,6 @@ class FMobileDirectionalLightFunctionPS : public FMaterialShader
 			PermutationVector.Set<FEnableShadingModelSupport>(false);
 		}
 
-		if (!AreMobileScreenSpaceReflectionsEnabled(Platform))
-		{
-			PermutationVector.Set<FMobileSSRQuality>(EMobileSSRQuality::Disabled);
-		}
-
 		if (!IsMobileDistanceFieldAOEnabled(Platform))
 		{
 			PermutationVector.Set<FSkyShadowing>(false);
@@ -192,27 +179,21 @@ class FMobileDirectionalLightFunctionPS : public FMaterialShader
 		bool bShadingModelSupport, 
 		bool bDynamicShadows, 
 		bool bSkyLight, 
-		bool bPlanarReflection,
 		bool bScreenSpaceShadowMask,
-		EMobileSSRQuality MobileSSRQuality,
 		bool bApplySkyShadowing)
 	{
 		EShaderPlatform ShaderPlatform = View.GetShaderPlatform();
 		bool bUseClusteredLights = UseClusteredDeferredShading(ShaderPlatform);
-		bool bClustredReflection = bInlineReflectionAndSky && (View.NumBoxReflectionCaptures + View.NumSphereReflectionCaptures) > 0;
 		bool bEnableSkyLight = bInlineReflectionAndSky && bSkyLight;
 		int32 ShadowQuality = bDynamicShadows && !bScreenSpaceShadowMask ? (int32)GetShadowQuality() : 0;
 
 		FPermutationDomain PermutationVector;
 		PermutationVector.Set<FEnableShadingModelSupport>(bShadingModelSupport);
 		PermutationVector.Set<FEnableClustredLights>(bUseClusteredLights);
-		PermutationVector.Set<FEnableClustredReflection>(bClustredReflection);
-		PermutationVector.Set<FEnablePlanarReflection>(bPlanarReflection);
 		PermutationVector.Set<FEnableSkyLight>(bEnableSkyLight);
 		PermutationVector.Set<FEnableScreenSpaceShadowMask>(bScreenSpaceShadowMask);
 		PermutationVector.Set<FEnableCSM>(ShadowQuality > 0);
 		PermutationVector.Set<FShadowQuality>(FMath::Clamp(ShadowQuality, 1, 3));
-		PermutationVector.Set<FMobileSSRQuality>(bInlineReflectionAndSky ? MobileSSRQuality : EMobileSSRQuality::Disabled);
 		PermutationVector.Set<FSkyShadowing>(bInlineReflectionAndSky && bApplySkyShadowing);
 		PermutationVector.Set<FInlineContactShadows>(IsMobileInlineContactShadowsEnabled(ShaderPlatform));
 		return PermutationVector;
@@ -470,16 +451,9 @@ static void GetLightMaterial(const FCachedLightMaterial& DefaultLightMaterial, c
 	OutShader = MaterialShaderMap->GetShader<ShaderType>(PermutationId);
 }
 
-extern const uint8 MobileShadingModelSupportStencilValue;
-
 uint8 PassShadingModelStencilValue(bool bEnableShadingModelSupport)
 {
-	return bEnableShadingModelSupport ? GET_STENCIL_MOBILE_SM_MASK(MobileShadingModelSupportStencilValue) : STENCIL_MOBILE_DEFAULTLIT_MASK;
-}
-
-constexpr uint8 PassShadingModelStencilMask(bool bEnableShadingModelSupport)
-{
-	return bEnableShadingModelSupport ? GET_STENCIL_MOBILE_SM_MASK(0xff) : STENCIL_MOBILE_DEFAULTLIT_MASK;
+	return bEnableShadingModelSupport ? GET_STENCIL_BIT_MASK(MOBILE_SHADINGMODELS, 1) : GET_STENCIL_BIT_MASK(MOBILE_DEFAULTLIT, 1);
 }
 
 void RenderReflectionEnvironmentSkyLighting(
@@ -515,22 +489,12 @@ void RenderReflectionEnvironmentSkyLighting(
 	}
 	GraphicsPSOInit.RasterizerState = TStaticRasterizerState<>::GetRHI();
 
-	int32 NumPasses = 1;
-	FDepthStencilStateRHIRef StencilState[3];
-	StencilState[0] = TStaticDepthStencilState<
-		false, CF_Always,
-		true, CF_Equal, SO_Keep, SO_Keep, SO_Keep,
-		false, CF_Always, SO_Keep, SO_Keep, SO_Keep,
-		PassShadingModelStencilMask(/*bEnableShadingModelSupport=*/false), 0x00>::GetRHI();
-	uint8 StencilRef[3] = {};
-	StencilRef[0] = PassShadingModelStencilValue(/*bEnableShadingModelSupport=*/false);
-
 	FMobileReflectionEnvironmentSkyLightingPS::FParameters PassParameters;
 	PassParameters.View = GetShaderBinding(View.ViewUniformBuffer);
 	PassParameters.MobileReflectionCaptureData = GetShaderBinding(View.MobileReflectionCaptureUniformBuffer);
 
 	// DFAO
-	if(DynamicBentNormalAOTexture != nullptr)
+	if (DynamicBentNormalAOTexture != nullptr)
 	{
 		PassParameters.DFAOUpsampleParameters = DistanceField::SetupAOUpsampleParameters(View, DynamicBentNormalAOTexture);
 		float DynamicBentNormalAO = 1.0f;
@@ -538,111 +502,82 @@ void RenderReflectionEnvironmentSkyLighting(
 	}
 
 	TShaderMapRef<FPostProcessVS> VertexShader(View.ShaderMap);
-	
-	uint32 PassEnableShadingModelSupport = 0;
-	uint32 PassEnableSSR = 0;
-	if (MobileSSRQuality != EMobileSSRQuality::Disabled)
-	{
-		// Separate pass for fully rough default lit materials
-		int PassIndex = NumPasses++;
-		StencilState[0] = StencilState[PassIndex] = TStaticDepthStencilState<
-			false, CF_Always,
-			true, CF_Equal, SO_Keep, SO_Keep, SO_Keep,
-			false, CF_Always, SO_Keep, SO_Keep, SO_Keep,
-			GET_STENCIL_MOBILE_SM_MASK(0xff), 0x00>::GetRHI();
 
-		StencilRef[PassIndex] = STENCIL_MOBILE_DEFAULTLIT_MASK | STENCIL_MOBILE_REFLECTIVE_MASK;
-		PassEnableSSR |= (1 << PassIndex);
-	}
+	GraphicsPSOInit.DepthStencilState = TStaticDepthStencilState<
+		false, CF_Always,
+		true, CF_Equal, SO_Keep, SO_Keep, SO_Keep,
+		false, CF_Always, SO_Keep, SO_Keep, SO_Keep,
+		GET_STENCIL_MOBILE_SM_MASK(0xff), 0x00>::GetRHI();
+	uint8 StencilRef = PassShadingModelStencilValue(/*bEnableShadingModelSupport=*/false);
+
 	if (!IsOnlyDefaultLitShadingModel(View.ShadingModelMaskInView) && MobileUsesGBufferCustomData(Scene.GetShaderPlatform()))
 	{
-		// Separate pass for all materials with custom shading models
-		int PassIndex = NumPasses++;
-		StencilState[PassIndex] = TStaticDepthStencilState<
+		GraphicsPSOInit.DepthStencilState = TStaticDepthStencilState<
 			false, CF_Always,
-			true, CF_Equal, SO_Keep, SO_Keep, SO_Keep,
+			true, CF_NotEqual, SO_Keep, SO_Keep, SO_Keep,
 			false, CF_Always, SO_Keep, SO_Keep, SO_Keep,
-			PassShadingModelStencilMask(/*bEnableShadingModelSupport=*/true), 0x00>::GetRHI();
-		StencilRef[PassIndex] = PassShadingModelStencilValue(/*bEnableShadingModelSupport=*/true);
-		PassEnableSSR |= (MobileSSRQuality != EMobileSSRQuality::Disabled) ? (1 << PassIndex) : 0;
-		PassEnableShadingModelSupport |= (1 << PassIndex);
+			GET_STENCIL_MOBILE_SM_MASK(0xff), 0x00>::GetRHI();
+		// Apply all shading models
+		StencilRef = 0;
 	}
 
-	for (int32 PassIndex = 0; PassIndex < NumPasses; PassIndex++)
-	{
-		GraphicsPSOInit.DepthStencilState = StencilState[PassIndex];
-		
-		FMobileReflectionEnvironmentSkyLightingPS::FPermutationDomain PermutationVector;
-		PermutationVector.Set<FMobileReflectionEnvironmentSkyLightingPS::FEnableShadingModelSupport>(PassEnableShadingModelSupport & (1 << PassIndex));
-		PermutationVector.Set<FMobileReflectionEnvironmentSkyLightingPS::FEnableClustredReflection>(bClustredReflection);
-		PermutationVector.Set<FMobileReflectionEnvironmentSkyLightingPS::FEnablePlanarReflection>(bPlanarReflection);
-		PermutationVector.Set<FMobileReflectionEnvironmentSkyLightingPS::FEnableSkyLight>(bEnableSkyLight);
-		PermutationVector.Set<FMobileReflectionEnvironmentSkyLightingPS::FMobileSSRQuality>((PassEnableSSR & (1 << PassIndex)) ? MobileSSRQuality : EMobileSSRQuality::Disabled);
-		extern bool UseDistanceFieldAO();
-		PermutationVector.Set<FMobileReflectionEnvironmentSkyLightingPS::FSkyShadowing>(DynamicBentNormalAOTexture && UseDistanceFieldAO() && IsMobileDistanceFieldAOEnabled(View.GetShaderPlatform()));
-		TShaderMapRef<FMobileReflectionEnvironmentSkyLightingPS> PixelShader(View.ShaderMap, PermutationVector);
+	FMobileReflectionEnvironmentSkyLightingPS::FPermutationDomain PermutationVector;
+	PermutationVector.Set<FMobileReflectionEnvironmentSkyLightingPS::FEnableShadingModelSupport>(!IsOnlyDefaultLitShadingModel(View.ShadingModelMaskInView));
+	PermutationVector.Set<FMobileReflectionEnvironmentSkyLightingPS::FEnableClustredReflection>(bClustredReflection);
+	PermutationVector.Set<FMobileReflectionEnvironmentSkyLightingPS::FEnablePlanarReflection>(bPlanarReflection);
+	PermutationVector.Set<FMobileReflectionEnvironmentSkyLightingPS::FEnableSkyLight>(bEnableSkyLight);
+	PermutationVector.Set<FMobileReflectionEnvironmentSkyLightingPS::FMobileSSRQuality>(MobileSSRQuality);
+	extern bool UseDistanceFieldAO();
+	PermutationVector.Set<FMobileReflectionEnvironmentSkyLightingPS::FSkyShadowing>(DynamicBentNormalAOTexture && UseDistanceFieldAO() && IsMobileDistanceFieldAOEnabled(View.GetShaderPlatform()));
+	TShaderMapRef<FMobileReflectionEnvironmentSkyLightingPS> PixelShader(View.ShaderMap, PermutationVector);
 		
-		GraphicsPSOInit.BoundShaderState.VertexDeclarationRHI = GFilterVertexDeclaration.VertexDeclarationRHI;
-		GraphicsPSOInit.BoundShaderState.VertexShaderRHI = VertexShader.GetVertexShader();
-		GraphicsPSOInit.BoundShaderState.PixelShaderRHI = PixelShader.GetPixelShader();
-		GraphicsPSOInit.PrimitiveType = PT_TriangleList;
+	GraphicsPSOInit.BoundShaderState.VertexDeclarationRHI = GFilterVertexDeclaration.VertexDeclarationRHI;
+	GraphicsPSOInit.BoundShaderState.VertexShaderRHI = VertexShader.GetVertexShader();
+	GraphicsPSOInit.BoundShaderState.PixelShaderRHI = PixelShader.GetPixelShader();
+	GraphicsPSOInit.PrimitiveType = PT_TriangleList;
 
-		SetGraphicsPipelineState(RHICmdList, GraphicsPSOInit, StencilRef[PassIndex]);
-		SetShaderParameters(RHICmdList, PixelShader, PixelShader.GetPixelShader(), PassParameters);
+	SetGraphicsPipelineState(RHICmdList, GraphicsPSOInit, StencilRef);
+	SetShaderParameters(RHICmdList, PixelShader, PixelShader.GetPixelShader(), PassParameters);
 
-		const FIntPoint TargetSize = View.GetSceneTexturesConfig().Extent;
+	const FIntPoint TargetSize = View.GetSceneTexturesConfig().Extent;
 
-		DrawRectangle(
-			RHICmdList,
-			0, 0,
-			View.ViewRect.Width(), View.ViewRect.Height(),
-			View.ViewRect.Min.X, View.ViewRect.Min.Y,
-			View.ViewRect.Width(), View.ViewRect.Height(),
-			FIntPoint(View.ViewRect.Width(), View.ViewRect.Height()),
-			TargetSize,
-			VertexShader);
-	}
+	DrawRectangle(
+		RHICmdList,
+		0, 0,
+		View.ViewRect.Width(), View.ViewRect.Height(),
+		View.ViewRect.Min.X, View.ViewRect.Min.Y,
+		View.ViewRect.Width(), View.ViewRect.Height(),
+		FIntPoint(View.ViewRect.Width(), View.ViewRect.Height()),
+		TargetSize,
+		VertexShader);
 }
 
-template<uint32 LightingChannelIdx, bool bEnableShadingModelSupport>
+template<uint32 LightingChannelIdx>
 static void SetDirectionalLightDepthStencilState(FGraphicsPipelineStateInitializer& GraphicsPSOInit)
 {
 	GraphicsPSOInit.DepthStencilState = TStaticDepthStencilState<
 		false, CF_Always,
 		true, CF_Equal, SO_Keep, SO_Keep, SO_Keep,
 		false, CF_Always, SO_Keep, SO_Keep, SO_Keep,
-		PassShadingModelStencilMask(bEnableShadingModelSupport) | STENCIL_LIGHTING_CHANNELS_MASK(1u << LightingChannelIdx), 0x00>::GetRHI();
+		GET_STENCIL_MOBILE_SM_MASK(0xff) | STENCIL_LIGHTING_CHANNELS_MASK(1u << LightingChannelIdx), 0x00>::GetRHI();
 }
 
-template<bool bEnableShadingModelSupport>
 static void SetDirectionalLightDepthStencilState(FGraphicsPipelineStateInitializer& GraphicsPSOInit, uint32 LightingChannelIdx)
 {
 	switch (LightingChannelIdx)
 	{
 	default:
-		SetDirectionalLightDepthStencilState<0, bEnableShadingModelSupport>(GraphicsPSOInit);
+		SetDirectionalLightDepthStencilState<0>(GraphicsPSOInit);
 		break;
 	case 1:
-		SetDirectionalLightDepthStencilState<1, bEnableShadingModelSupport>(GraphicsPSOInit);
+		SetDirectionalLightDepthStencilState<1>(GraphicsPSOInit);
 		break;
 	case 2:
-		SetDirectionalLightDepthStencilState<2, bEnableShadingModelSupport>(GraphicsPSOInit);
+		SetDirectionalLightDepthStencilState<2>(GraphicsPSOInit);
 		break;
 	}
 }
 
-static void SetDirectionalLightDepthStencilState(FGraphicsPipelineStateInitializer& GraphicsPSOInit, uint32 LightingChannelIdx, bool bEnableShadingModelSupport)
-{
-	if (bEnableShadingModelSupport)
-	{
-		SetDirectionalLightDepthStencilState<true>(GraphicsPSOInit, LightingChannelIdx);
-	}
-	else
-	{
-		SetDirectionalLightDepthStencilState<false>(GraphicsPSOInit, LightingChannelIdx);
-	}
-}
-
 extern void GetLightContactShadowParameters(const FLightSceneProxy* Proxy, float& OutLength, bool& bOutLengthInWS, float& OutCastingIntensity, float& OutNonCastingIntensity);
 void SetupContactShadowParameters(const FLightSceneProxy* Proxy, const FViewInfo& View, FVector3f& OutContactShadowsParameter)
 {
@@ -670,7 +605,6 @@ static void RenderDirectionalLight(
 	const FLightSceneInfo& DirectionalLight, 
 	uint32 LightingChannel, 
 	bool bInlineReflectionAndSky, 
-	EMobileSSRQuality MobileSSRQuality,
 	FRDGTextureRef DynamicBentNormalAOTexture)
 {
 	FString LightNameWithLevel;
@@ -691,10 +625,8 @@ static void RenderDirectionalLight(
 
 	FMobileDirectionalLightFunctionPS::FParameters PassParameters;
 	PassParameters.MobileDirectionalLight = Scene.UniformBuffers.MobileDirectionalLightUniformBuffers[LightingChannel + 1];
-	PassParameters.MobileReflectionCaptureData = GetShaderBinding(View.MobileReflectionCaptureUniformBuffer);
 	PassParameters.MobileDeferredCommonParameters.LightFunctionParameters = FVector4f(1.0f, 1.0f, 0.0f, 0.0f);
 	PassParameters.MobileDeferredCommonParameters.CameraRelativeLightPosition = GetCamRelativeLightPosition(View.ViewMatrices, DirectionalLight);
-
 	{
 		PassParameters.MobileDeferredCommonParameters.LightFunctionParameters2 = FVector2f(DirectionalLight.Proxy->GetLightFunctionFadeDistance(), DirectionalLight.Proxy->GetLightFunctionDisabledBrightness());
 		const FVector Scale = DirectionalLight.Proxy->GetLightFunctionScale();
@@ -716,7 +648,6 @@ static void RenderDirectionalLight(
 	const bool bDynamicSkyLight = Scene.SkyLight && (!Scene.SkyLight->bHasStaticLighting || !IsStaticLightingAllowed());
 	const bool bEnableSkyLight = bDynamicSkyLight && View.Family->EngineShowFlags.SkyLighting;
 	const bool bDynamicShadows = DirectionalLight.Proxy->CastsDynamicShadow() && View.Family->EngineShowFlags.DynamicShadows;
-	const bool bPlanarReflection = Scene.GetForwardPassGlobalPlanarReflection() != nullptr;
 
 	// Add to emissive in SceneColor
 	if (bInlineReflectionAndSky && !bDynamicSkyLight)
@@ -731,28 +662,17 @@ static void RenderDirectionalLight(
 	
 	// Do two passes, first masking DefautLit, second masking all other shading models
 	const bool bOnlyDefaultLitInView = IsOnlyDefaultLitShadingModel(View.ShadingModelMaskInView);
-	const bool bUseSSR = bInlineReflectionAndSky && MobileSSRQuality != EMobileSSRQuality::Disabled;
 	int32 NumPasses = 1;
-	uint32 PassEnableSSR = 0;
 	uint32 PassEnableShadingModelSupport = 0;
-	uint32 PassShadingModelStencilMaskAnyDefaultLit = 1;
-	uint32 ShadingModelStencilRef[3] = {};
+	uint32 ShadingModelStencilRef[2] = {};
 	ShadingModelStencilRef[0] = PassShadingModelStencilValue(/*bEnableShadingModelSupport=*/false);
 
 	if (!bOnlyDefaultLitInView && MobileUsesGBufferCustomData(Scene.GetShaderPlatform()))
 	{
 		const int32 PassIndex = NumPasses++;
 		PassEnableShadingModelSupport |= (1 << PassIndex);
-		PassEnableSSR |= bUseSSR ? (1 << PassIndex) : 0;
 		ShadingModelStencilRef[PassIndex] = PassShadingModelStencilValue(/*bEnableShadingModelSupport=*/true);
 	}
-	if (bUseSSR)
-	{
-		PassShadingModelStencilMaskAnyDefaultLit = 0; // Pass 0 only does default lit, non-reflective.
-		const int32 PassIndex = NumPasses++;
-		PassEnableSSR |= (1 << PassIndex);
-		ShadingModelStencilRef[PassIndex] = STENCIL_MOBILE_DEFAULTLIT_MASK | STENCIL_MOBILE_REFLECTIVE_MASK;
-	}
 
 	static auto CVarMobileSupportInsetShadows = IConsoleManager::Get().FindConsoleVariable(TEXT("r.Mobile.SupportInsetShadows"));
 	const bool bInsetShadows = (CVarMobileSupportInsetShadows != nullptr && CVarMobileSupportInsetShadows->GetInt() != 0);
@@ -760,20 +680,17 @@ static void RenderDirectionalLight(
 
 	SetupContactShadowParameters(DirectionalLight.Proxy, View, PassParameters.MobileDeferredCommonParameters.ContactShadows);
 	const uint8 LightingChannelStencilValue = GetLightingChannelStencilValue(LightingChannel);
+	SetDirectionalLightDepthStencilState(GraphicsPSOInit, LightingChannel);
 	
 	for (int32 PassIndex = 0; PassIndex < NumPasses; ++PassIndex)
 	{
-		SetDirectionalLightDepthStencilState(GraphicsPSOInit, LightingChannel, (PassShadingModelStencilMaskAnyDefaultLit & (1 << PassIndex)) == 0);
-		
 		FMobileDirectionalLightFunctionPS::FPermutationDomain PermutationVector = FMobileDirectionalLightFunctionPS::BuildPermutationVector(
 			View,
 			bInlineReflectionAndSky,
 			PassEnableShadingModelSupport & (1 << PassIndex),
 			bDynamicShadows,
 			bEnableSkyLight,
-			bPlanarReflection, 
 			bMobileUsesShadowMaskTexture,
-			(PassEnableSSR & (1<<PassIndex)) ? MobileSSRQuality : EMobileSSRQuality::Disabled,
 			DynamicBentNormalAOTexture && UseDistanceFieldAO() && IsMobileDistanceFieldAOEnabled(View.GetShaderPlatform())
 		);
 		FCachedLightMaterial LightMaterial;
@@ -828,8 +745,11 @@ static int RenderDirectionalLights(
 	}
 	// We can merge reflection and skylight pass with a sole directional light pass and if all primitives and the directional light use the default lighting channel
 	bool bPrimitivesUseLightingChannels = (View.bUsesLightingChannels && GMobileIgnoreDeferredShadingSkyLightChannels == 0);
-	const bool bInlineReflectionAndSky = (NumLights == 1) && !bPrimitivesUseLightingChannels && (Scene.MobileDirectionalLights[0] != nullptr);
+	bool bPlanarReflection = Scene.GetForwardPassGlobalPlanarReflection() != nullptr;
+	bool bClusteredReflection = (View.NumBoxReflectionCaptures + View.NumSphereReflectionCaptures) > 0;
+	bool bSSR = MobileSSRQuality != EMobileSSRQuality::Disabled;
 
+	bool bInlineReflectionAndSky = (NumLights == 1) && !bPrimitivesUseLightingChannels && (Scene.MobileDirectionalLights[0] != nullptr) && !(bPlanarReflection || bClusteredReflection || bSSR);
 	if (!bInlineReflectionAndSky)
 	{
 		RenderReflectionEnvironmentSkyLighting(RHICmdList, Scene, View, MobileSSRQuality, DynamicBentNormalAOTexture);
@@ -840,7 +760,7 @@ static int RenderDirectionalLights(
 		FLightSceneInfo* DirectionalLight = Scene.MobileDirectionalLights[ChannelIdx];
 		if (DirectionalLight)
 		{
-			RenderDirectionalLight(RHICmdList, Scene, View, DefaultLightMaterial, *DirectionalLight, ChannelIdx, bInlineReflectionAndSky, MobileSSRQuality, DynamicBentNormalAOTexture);
+			RenderDirectionalLight(RHICmdList, Scene, View, DefaultLightMaterial, *DirectionalLight, ChannelIdx, bInlineReflectionAndSky, DynamicBentNormalAOTexture);
 		}
 	}
 	return NumLights;
@@ -858,7 +778,7 @@ static void SetLocalLightRasterizerAndDepthState(FGraphicsPipelineStateInitializ
 			false, CF_LessEqual,
 			false, CF_Equal, SO_Keep, SO_Keep, SO_Keep,
 			true, CF_Equal, SO_Zero, SO_Keep, SO_Zero,
-			PassShadingModelStencilMask(bEnableShadingModelSupport) | STENCIL_LIGHTING_CHANNELS_MASK(1u << LightingChannel) | STENCIL_SANDBOX_MASK,
+			GET_STENCIL_MOBILE_SM_MASK(0xff) | STENCIL_LIGHTING_CHANNELS_MASK(1u << LightingChannel) | STENCIL_SANDBOX_MASK,
 			STENCIL_SANDBOX_MASK
 		>::GetRHI();
 	}
@@ -872,7 +792,7 @@ static void SetLocalLightRasterizerAndDepthState(FGraphicsPipelineStateInitializ
 				false, CF_Always,
 				true, CF_Equal, SO_Keep, SO_Keep, SO_Keep,
 				false, CF_Always, SO_Keep, SO_Keep, SO_Keep,
-				PassShadingModelStencilMask(bEnableShadingModelSupport) | STENCIL_LIGHTING_CHANNELS_MASK(1u << LightingChannel), 0x00>::GetRHI();
+				GET_STENCIL_MOBILE_SM_MASK(0xff) | STENCIL_LIGHTING_CHANNELS_MASK(1u << LightingChannel), 0x00>::GetRHI();
 		}
 		else
 		{
@@ -882,7 +802,7 @@ static void SetLocalLightRasterizerAndDepthState(FGraphicsPipelineStateInitializ
 				false, CF_DepthNearOrEqual,
 				true, CF_Equal, SO_Keep, SO_Keep, SO_Keep,
 				false, CF_Always, SO_Keep, SO_Keep, SO_Keep,
-				PassShadingModelStencilMask(bEnableShadingModelSupport) | STENCIL_LIGHTING_CHANNELS_MASK(1u << LightingChannel), 0x00>::GetRHI();
+				GET_STENCIL_MOBILE_SM_MASK(0xff) | STENCIL_LIGHTING_CHANNELS_MASK(1u << LightingChannel), 0x00>::GetRHI();
 		}
 	}
 }
@@ -1285,17 +1205,13 @@ void FDeferredMobileLightMaterialPSOCollector::CollectPSOInitializersDirectional
 			const bool bMobileUsesShadowMaskTexture = MobileUsesShadowMaskTexture(GMaxRHIShaderPlatform);
 			const bool bApplySkyShadowing = false;
 			const int32 LightingChannel = 0;
-			uint32 PassShadingModelStencilMaskAnyDefaultLit = 1;
-
+			
 			FMobileDirectionalLightFunctionPS::FPermutationDomain PermutationVector;
 			PermutationVector.Set<FMobileDirectionalLightFunctionPS::FEnableShadingModelSupport>(bShadingModelSupport);
 			PermutationVector.Set<FMobileDirectionalLightFunctionPS::FEnableClustredLights>(bUseClusteredLights);
-			PermutationVector.Set<FMobileDirectionalLightFunctionPS::FEnableClustredReflection>(bClustredReflection);
-			PermutationVector.Set<FMobileDirectionalLightFunctionPS::FEnablePlanarReflection>(bPlanarReflection);
 			PermutationVector.Set<FMobileDirectionalLightFunctionPS::FEnableSkyLight>(bEnableSkyLight);
 			PermutationVector.Set<FMobileDirectionalLightFunctionPS::FEnableCSM>(ShadowQuality > 0);
 			PermutationVector.Set<FMobileDirectionalLightFunctionPS::FShadowQuality>(FMath::Clamp(ShadowQuality, 1, 3));
-			PermutationVector.Set<FMobileDirectionalLightFunctionPS::FMobileSSRQuality>(EMobileSSRQuality::Disabled);
 			PermutationVector.Set<FMobileDirectionalLightFunctionPS::FSkyShadowing>(bInlineReflectionAndSky && bApplySkyShadowing);
 			FMaterialShaderTypes ShaderTypesToGet;
 			ShaderTypesToGet.AddShaderType<FMobileDirectionalLightFunctionPS>(PermutationVector.ToDimensionValueId());
@@ -1340,7 +1256,7 @@ void FDeferredMobileLightMaterialPSOCollector::CollectPSOInitializersDirectional
 			{
 				GraphicsPSOInit.BlendState = TStaticBlendState<CW_RGB, BO_Add, BF_One, BF_One>::GetRHI();
 			}
-			SetDirectionalLightDepthStencilState(GraphicsPSOInit, LightingChannel, (PassShadingModelStencilMaskAnyDefaultLit & (1 << PassIndex)) == 0);
+			SetDirectionalLightDepthStencilState(GraphicsPSOInit, LightingChannel);
 
 			GraphicsPSOInit.BoundShaderState.VertexDeclarationRHI = GFilterVertexDeclaration.VertexDeclarationRHI;
 			GraphicsPSOInit.BoundShaderState.VertexShaderRHI = RHIVertexShader;
diff --git a/Engine/Source/Runtime/Renderer/Private/PostProcess/SceneRenderTargets.h b/Engine/Source/Runtime/Renderer/Private/PostProcess/SceneRenderTargets.h
--- a/Engine/Source/Runtime/Renderer/Private/PostProcess/SceneRenderTargets.h
+++ b/Engine/Source/Runtime/Renderer/Private/PostProcess/SceneRenderTargets.h
@@ -55,8 +55,8 @@
 // [Mobile Deferred only]
 // Store shading model into stencil [1-2] bits
 #define GET_STENCIL_MOBILE_SM_MASK(Value)			uint8(((Value) & 0x3) << 1)
-#define STENCIL_MOBILE_DEFAULTLIT_MASK				uint8(1 << 2)
-#define STENCIL_MOBILE_REFLECTIVE_MASK				uint8(1 << 1)
+#define STENCIL_MOBILE_SHADINGMODELS_BIT_ID			1
+#define STENCIL_MOBILE_DEFAULTLIT_BIT_ID			2
 
 // [Mobile Forward only]
 // Cast contact shadow mask - bit 4 / Must match shader (ScreenSpaceShadows.usf)
