Explicitly use lower case for virtual texture name hash.
Adaptive virtual texture specializes name hash for grid location.
#rnx

diff --git a/Engine/Source/Runtime/Engine/Private/Texture2D.cpp b/Engine/Source/Runtime/Engine/Private/Texture2D.cpp
--- a/Engine/Source/Runtime/Engine/Private/Texture2D.cpp
+++ b/Engine/Source/Runtime/Engine/Private/Texture2D.cpp
@@ -1552,7 +1552,7 @@ FVirtualTexture2DResource::FVirtualTexture2DResource(const UTexture2D* InOwner,
 
 	TextureName = InOwner->GetFName();
 	PackageName = InOwner->GetOutermost()->GetFName();
-	FullNameHash = GetTypeHash(InOwner->GetName());
+	FullNameHash = GetTypeHash(InOwner->GetName().ToLower());
 
 	UTextureLODSettings* TextureLODSettings = UDeviceProfileManager::Get().GetActiveProfile()->GetTextureLODSettings();
 	Filter = (ESamplerFilter)TextureLODSettings->GetSamplerFilter(InOwner);
diff --git a/Engine/Source/Runtime/Engine/Private/VT/RuntimeVirtualTexture.cpp b/Engine/Source/Runtime/Engine/Private/VT/RuntimeVirtualTexture.cpp
--- a/Engine/Source/Runtime/Engine/Private/VT/RuntimeVirtualTexture.cpp
+++ b/Engine/Source/Runtime/Engine/Private/VT/RuntimeVirtualTexture.cpp
@@ -316,7 +316,7 @@ int32 URuntimeVirtualTexture::GetPageTableSize() const
 void URuntimeVirtualTexture::GetProducerDescription(FVTProducerDescription& OutDesc, FInitSettings const& InitSettings, FTransform const& VolumeToWorld) const
 {
 	OutDesc.Name = GetFName();
-	OutDesc.FullNameHash = GetTypeHash(GetName());
+	OutDesc.FullNameHash = GetTypeHash(GetName().ToLower());
 	OutDesc.Dimensions = 2;
 	OutDesc.DepthInTiles = 1;
 	OutDesc.WidthInBlocks = 1;
diff --git a/Engine/Source/Runtime/Renderer/Private/VT/AdaptiveVirtualTexture.cpp b/Engine/Source/Runtime/Renderer/Private/VT/AdaptiveVirtualTexture.cpp
--- a/Engine/Source/Runtime/Renderer/Private/VT/AdaptiveVirtualTexture.cpp
+++ b/Engine/Source/Runtime/Renderer/Private/VT/AdaptiveVirtualTexture.cpp
@@ -190,6 +190,9 @@ IAllocatedVirtualTexture* FAdaptiveVirtualTexture::AllocateVirtualTexture(
 		NewProducerDesc.BlockWidthInTiles = InWidthInTiles;
 		NewProducerDesc.BlockHeightInTiles = InHeightInTiles;
 		NewProducerDesc.MaxLevel = FMath::CeilLogTwo(FMath::Max(InWidthInTiles, InHeightInTiles));
+		
+		// Specialize the producer hash by position in the adaptive grid.
+		NewProducerDesc.FullNameHash = HashCombine(Producer->GetDescription().FullNameHash, GetTypeHash(InAddressOffset));
 
 		IVirtualTexture* VirtualTextureProducer = Producer->GetVirtualTexture();
 		IVirtualTexture* NewVirtualTextureProducer = new FVirtualTextureAddressRedirect(VirtualTextureProducer, InAddressOffset, InLevelOffset);
@@ -308,7 +311,7 @@ int32 FAdaptiveVirtualTexture::GetSpaceID() const
 	return AllocatedVirtualTextureLowMips->GetSpaceID();
 }
 
-void FAdaptiveVirtualTexture::GetProducers(FIntRect const& InTextureRegion, uint32 InMaxLevel, TArray<FProducerInfo>& OutProducerInfos)
+void FAdaptiveVirtualTexture::GetProducers(FIntRect const& InTextureRegion, uint32 InMaxLevel, TArray<FProducerInfo>& OutProducerInfos) const
 {
 	const uint32 NumProducers = AllocatedVirtualTextureLowMips->GetNumUniqueProducers();
 
diff --git a/Engine/Source/Runtime/Renderer/Private/VT/AdaptiveVirtualTexture.h b/Engine/Source/Runtime/Renderer/Private/VT/AdaptiveVirtualTexture.h
--- a/Engine/Source/Runtime/Renderer/Private/VT/AdaptiveVirtualTexture.h
+++ b/Engine/Source/Runtime/Renderer/Private/VT/AdaptiveVirtualTexture.h
@@ -53,7 +53,7 @@ public:
 		uint32 RemappedMaxLevel;
 	};
 	/** Get internal producers that touch the texture region. */
-	void GetProducers(FIntRect const& InTextureRegion, uint32 InMaxLevel, TArray<FProducerInfo>& OutProducerInfos);
+	void GetProducers(FIntRect const& InTextureRegion, uint32 InMaxLevel, TArray<FProducerInfo>& OutProducerInfos) const;
 
 protected:
 	//~ Begin IAdaptiveVirtualTexture Interface.
