Support runtime virtual texture preload requests for the adaptive page table case.
Adaptive page tables need to handle multiple allocated virtual textures overlapping the requested preload bounds.
Also the preload needs to generate adaptive reallocation requests to reach the requested mip level.

diff --git a/Engine/Source/Runtime/Renderer/Private/VT/AdaptiveVirtualTexture.cpp b/Engine/Source/Runtime/Renderer/Private/VT/AdaptiveVirtualTexture.cpp
--- a/Engine/Source/Runtime/Renderer/Private/VT/AdaptiveVirtualTexture.cpp
+++ b/Engine/Source/Runtime/Renderer/Private/VT/AdaptiveVirtualTexture.cpp
@@ -371,6 +371,57 @@ void FAdaptiveVirtualTexture::GetProducers(FIntRect const& InTextureRegion, uint
 	}
 }
 
+void FAdaptiveVirtualTexture::GetAllocatedVirtualTextures(FBox2D const& InUVRegion, uint32 InLevel, TArray<FAllocatedInfo>& OutInfos, TArray<uint32>& OutAllocationRequests)
+{
+	// The persistent allocated virtual texture.
+	const int32 RemappedLevel = (int32)AllocatedVirtualTextureLowMips->GetMaxLevel() - (int32)MaxLevel + (int32)InLevel;
+	const int32 ClampedRemappedLevel = FMath::Max(RemappedLevel, 0);
+	OutInfos.Add(FAllocatedInfo{ AllocatedVirtualTextureLowMips, InUVRegion, (uint32)ClampedRemappedLevel });
+
+	// If the mip level is outside of the persistent allocated virtual texture then find the transient allocated virtual textures.
+	if (RemappedLevel < 0)
+	{
+		const int32 MinGridX = FMath::FloorToInt(FMath::Max(InUVRegion.Min.X, 0.0) * GridSize.X);
+		const int32 MaxGridX = FMath::CeilToInt(FMath::Min(InUVRegion.Max.X, 1.0) * GridSize.X);
+		const int32 MinGridY = FMath::FloorToInt(FMath::Max(InUVRegion.Min.Y, 0.0) * GridSize.Y);
+		const int32 MaxGridY = FMath::CeilToInt(FMath::Min(InUVRegion.Max.Y, 1.0) * GridSize.Y);
+
+		for (int32 Y = MinGridY; Y < MaxGridY; Y++)
+		{
+			for (int32 X = MinGridX; X < MaxGridX; X++)
+			{
+				const uint32 GridIndex = X + Y * GridSize.X;
+				const uint32 AllocationIndex = GetAllocationIndex(GridIndex);
+				if (AllocationIndex != INDEX_NONE)
+				{
+					FAllocatedVirtualTexture* AllocatedVT = AllocationSlots[AllocationIndex].AllocatedVT;
+					const int32 LocalRemappedLevel = RemappedLevel + (int32)AllocatedVT->GetMaxLevel();
+
+					if (LocalRemappedLevel >= 0)
+					{
+						FVector2D LocalUV0 = InUVRegion.Min * FVector2D(GridSize.X, GridSize.Y) - FVector2D(X, Y);
+						FVector2D LocalUV1 = InUVRegion.Max * FVector2D(GridSize.X, GridSize.Y) - FVector2D(X, Y);
+						const FBox2D RemappedUVRegion = FBox2D(LocalUV0, LocalUV1).Overlap(FBox2D(FVector2D::Zero(), FVector2D::One()));
+
+						OutInfos.Add(FAllocatedInfo{ AllocatedVT, RemappedUVRegion, (uint32)LocalRemappedLevel });
+
+						continue;
+					}
+				}
+
+				// The grid location doesn't cover the requested mip level so we create an allocation request.
+				FPackedAdaptiveAllocationRequest Request;
+				Request.bIsValid = 1;
+				Request.bIsAllocated = AllocationIndex == INDEX_NONE ? 0 : 1;
+				Request.bIsRequest = 1;
+				Request.AllocationOrGridIndex = AllocationIndex == INDEX_NONE ? GridIndex : AllocationIndex;
+				Request.Space = GetSpaceID();
+				OutAllocationRequests.Add(Request.PackedValue);
+			}
+		}
+	}
+}
+
 /** Get hash key for the GridIndexMap. */
 static uint16 GetGridIndexHash(int32 InGridIndex)
 {
@@ -465,46 +516,48 @@ uint32 FAdaptiveVirtualTexture::GetPackedAllocationRequest(uint32 vAddress, uint
 	return Request.PackedValue;
 }
 
-void FAdaptiveVirtualTexture::QueuePackedAllocationRequests(FVirtualTextureSystem* InSystem, uint32 const* InRequests, uint32 InNumRequests, uint32 InFrame)
+void FAdaptiveVirtualTexture::QueuePackedAllocationRequests(FVirtualTextureSystem* InSystem, TConstArrayView<uint32> InRequests, uint32 InFrame)
 {
-	if (InNumRequests > 0)
+	if (InRequests.Num() > 0)
 	{
-		// Sort for batching by SpaceID.
-		// We also sort here to help can skip duplicate requests. It might be better to remove duplicates before this call (when gathering requests) so that the sort here is cheaper.
-		TArray<uint32> SortRequests;
-		SortRequests.Insert(InRequests, InNumRequests, 0);
-		SortRequests.Sort();
+		// Sort in place for batching by SpaceID.
+		// We also sort here to help can skip duplicate requests. 
+		// It would be better either sort before this call, or to remove duplicates before this call (when gathering requests) so that the sort is cheaper.
+		TArray<uint32> SortedRequests;
+		SortedRequests = InRequests;
+		Algo::Sort(SortedRequests);
+		TConstArrayView<uint32> SortedRequestsView(SortedRequests);
 
-		uint32 StartRequestIndex = 0;
+		int32 StartRequestIndex = 0;
 		FPackedAdaptiveAllocationRequest StartRequest;
-		StartRequest.PackedValue = SortRequests[0];
+		StartRequest.PackedValue = SortedRequestsView[0];
 
-		for (uint32 RequestIndex = 0; RequestIndex < InNumRequests; ++RequestIndex)
+		for (int32 RequestIndex = 0; RequestIndex < SortedRequestsView.Num(); ++RequestIndex)
 		{
 			FPackedAdaptiveAllocationRequest Request;
-			Request.PackedValue = SortRequests[RequestIndex];
+			Request.PackedValue = SortedRequestsView[RequestIndex];
 
 			if (Request.Space != StartRequest.Space)
 			{
 				FAdaptiveVirtualTexture* AdaptiveVT = InSystem->GetAdaptiveVirtualTexture(StartRequest.Space);
-				AdaptiveVT->QueuePackedAllocationRequests(SortRequests.GetData() + StartRequestIndex, RequestIndex - StartRequestIndex, InFrame);
+				AdaptiveVT->QueuePackedAllocationRequests(SortedRequestsView.Slice(StartRequestIndex, RequestIndex - StartRequestIndex), InFrame);
 
 				StartRequestIndex = RequestIndex;
 				StartRequest = Request;
 			}
 		}
 
-		if (StartRequestIndex < InNumRequests)
+		if (StartRequestIndex < SortedRequestsView.Num())
 		{
 			FAdaptiveVirtualTexture* AdaptiveVT = InSystem->GetAdaptiveVirtualTexture(StartRequest.Space);
-			AdaptiveVT->QueuePackedAllocationRequests(SortRequests.GetData() + StartRequestIndex, InNumRequests - StartRequestIndex, InFrame);
+			AdaptiveVT->QueuePackedAllocationRequests(SortedRequestsView.Right(SortedRequestsView.Num() - StartRequestIndex), InFrame);
 		}
 	}
 }
 
-void FAdaptiveVirtualTexture::QueuePackedAllocationRequests(uint32 const* InRequests, uint32 InNumRequests, uint32 InFrame)
+void FAdaptiveVirtualTexture::QueuePackedAllocationRequests(TConstArrayView<uint32> InRequests, uint32 InFrame)
 {
-	for (uint32 RequestIndex = 0; RequestIndex < InNumRequests; ++RequestIndex)
+	for (int32 RequestIndex = 0; RequestIndex < InRequests.Num(); ++RequestIndex)
 	{
 		// Skip duplicates.
 		if (RequestIndex == 0 || InRequests[RequestIndex] != InRequests[RequestIndex - 1])
diff --git a/Engine/Source/Runtime/Renderer/Private/VT/AdaptiveVirtualTexture.h b/Engine/Source/Runtime/Renderer/Private/VT/AdaptiveVirtualTexture.h
--- a/Engine/Source/Runtime/Renderer/Private/VT/AdaptiveVirtualTexture.h
+++ b/Engine/Source/Runtime/Renderer/Private/VT/AdaptiveVirtualTexture.h
@@ -32,9 +32,9 @@ public:
 	/** Get a packed allocation key based on a virtual texture feedback request. The virtual texture system collects these opaque keys before queuing them for processing. */
 	uint32 GetPackedAllocationRequest(uint32 vAddress, uint32 vLevelPlusOne, uint32 Frame) const;
 	/** Queue a batch of allocation requests. These will be used to reallocate any virtual textures during the next call to UpdateAllocations(). */
-	void QueuePackedAllocationRequests(uint32 const* InRequests, uint32 InNumRequests, uint32 InFrame);
-	/** Queue a batch of allocation requests. This static function relays the global requests to the individual object queues. */
-	static void QueuePackedAllocationRequests(FVirtualTextureSystem* InSystem, uint32 const* InRequests, uint32 InNumRequests, uint32 InFrame);
+	void QueuePackedAllocationRequests(TConstArrayView<uint32> InRequests, uint32 InFrame);
+	/** Queue a batch of allocation requests. This static function relays the global requests to the individual object queues. Internally it sorts the requests. */
+	static void QueuePackedAllocationRequests(FVirtualTextureSystem* InSystem, TConstArrayView<uint32> InRequests, uint32 InFrame);
 	/** Update any allocations based on recent requests. */
 	void UpdateAllocations(FVirtualTextureSystem* InSystem, FRHICommandListImmediate& RHICmdList, uint32 InFrame);
 	/** Apply any page table updates. */
@@ -52,9 +52,19 @@ public:
 		FIntRect RemappedTextureRegion;
 		uint32 RemappedMaxLevel;
 	};
-	/** Get internal producers that touch the texture region. */
+	/** Get internal producers that touch the texture texel region. */
 	void GetProducers(FIntRect const& InTextureRegion, uint32 InMaxLevel, TArray<FProducerInfo>& OutProducerInfos) const;
 
+	/** Information needed by GetAllocatedVirtualTextures() for the internally allocated virtual textures. */
+	struct FAllocatedInfo
+	{
+		IAllocatedVirtualTexture* AllocatedVirtualTexture;
+		FBox2D RemappedUV;
+		uint32 RemappedLevel;
+	};
+	/** Get internal allocated virtual textures that touch the texture UV region. */
+	void GetAllocatedVirtualTextures(FBox2D const& InUVRegion, uint32 InLevel, TArray<FAllocatedInfo>& OutInfos, TArray<uint32>& OutAllocationRequests);
+
 protected:
 	//~ Begin IAdaptiveVirtualTexture Interface.
 	virtual void Destroy(class FVirtualTextureSystem* InSystem) override;
diff --git a/Engine/Source/Runtime/Renderer/Private/VT/VirtualTextureSystem.cpp b/Engine/Source/Runtime/Renderer/Private/VT/VirtualTextureSystem.cpp
--- a/Engine/Source/Runtime/Renderer/Private/VT/VirtualTextureSystem.cpp
+++ b/Engine/Source/Runtime/Renderer/Private/VT/VirtualTextureSystem.cpp
@@ -1199,7 +1199,7 @@ void FVirtualTextureSystem::RequestTiles(const FVector2D& InScreenSpaceSize, int
 	UE::TScopeLock Lock(Mutex);
 
 	for (const auto& Pair : AllocatedVTs)
-	{
+{
 		RequestTilesInternal(Pair.Value, InScreenSpaceSize, InMipLevel);
 	}
 }
@@ -1244,7 +1244,25 @@ void FVirtualTextureSystem::RequestTiles(IAllocatedVirtualTexture* AllocatedVT,
 	UE::TScopeLock Lock(Mutex);
 	if (InMipLevel >= 0)
 	{
-		RequestTilesForRegionInternal(AllocatedVT, InScreenSpaceSize, InViewportPosition, InViewportSize, InUV0, InUV1, InMipLevel);
+		if (FAdaptiveVirtualTexture* AdaptiveVT = AdaptiveVTs[AllocatedVT->GetSpaceID()])
+		{
+			// Adaptive virtual textures are a special case where the "root" AllocatedVT is one of many AllocatedVTs.
+			// We get the full set of AllocatedVTs before requesting tiles.
+			TArray<FAdaptiveVirtualTexture::FAllocatedInfo> Infos;
+			TArray<uint32> AdaptiveAllocationRequests;
+			AdaptiveVT->GetAllocatedVirtualTextures(FBox2D(InUV0, InUV1), InMipLevel, Infos, AdaptiveAllocationRequests);
+			for (FAdaptiveVirtualTexture::FAllocatedInfo const& Info : Infos)
+			{
+				RequestTilesForRegionInternal(Info.AllocatedVirtualTexture, InScreenSpaceSize, InViewportPosition, InViewportSize, Info.RemappedUV.Min, Info.RemappedUV.Max, Info.RemappedLevel);
+			}
+			
+			// If we found any adaptive AllocatedVTs that need reallocation to reach the requested mip level then queue those requests immediately.
+			AdaptiveVT->QueuePackedAllocationRequests(AdaptiveAllocationRequests, Frame);
+		}
+		else
+		{
+			RequestTilesForRegionInternal(AllocatedVT, InScreenSpaceSize, InViewportPosition, InViewportSize, InUV0, InUV1, InMipLevel);
+		}
 	}
 	else
 	{
@@ -2232,7 +2250,7 @@ void FVirtualTextureSystem::SubmitThrottledRequests(FRHICommandList& RHICmdList,
 	if (MergedRequestList->GetNumAdaptiveAllocationRequests() > 0)
 	{
 		SCOPE_CYCLE_COUNTER(STAT_ProcessRequests_QueueAdaptiveRequests);
-		FAdaptiveVirtualTexture::QueuePackedAllocationRequests(this, &MergedRequestList->GetAdaptiveAllocationRequest(0), MergedRequestList->GetNumAdaptiveAllocationRequests(), Frame);
+		FAdaptiveVirtualTexture::QueuePackedAllocationRequests(this, MakeConstArrayView(&MergedRequestList->GetAdaptiveAllocationRequest(0), MergedRequestList->GetNumAdaptiveAllocationRequests()), Frame);
 	}
 
 	// Deal with dirty MappedTilesToProduce pages first.
