Recycle unreferenced virtual texture spaces when creating a new virtual texture space.
Only recycle if we have no spare slots.
Reduces overall pressure on the virtual texture space counts.
#rnx

diff --git a/Engine/Source/Runtime/Renderer/Private/VT/VirtualTextureSystem.cpp b/Engine/Source/Runtime/Renderer/Private/VT/VirtualTextureSystem.cpp
--- a/Engine/Source/Runtime/Renderer/Private/VT/VirtualTextureSystem.cpp
+++ b/Engine/Source/Runtime/Renderer/Private/VT/VirtualTextureSystem.cpp
@@ -910,27 +910,47 @@ FVirtualTextureSpace* FVirtualTextureSystem::AcquireSpace(FRHICommandListBase& R
 		}
 	}
 
-	// Try to allocate a new space
 	if (InForceSpaceID == 0xff)
 	{
-		for (uint32 SpaceIndex = 0u; SpaceIndex < MaxSpaces; ++SpaceIndex)
+		// Try to allocate a new space.
+		int32 FreeSpaceIndex = INDEX_NONE;
+		for (uint32 SpaceIndex = 0u; FreeSpaceIndex == INDEX_NONE && SpaceIndex < MaxSpaces; ++SpaceIndex)
 		{
-			if (!Spaces[SpaceIndex])
+			if (Spaces[SpaceIndex] == nullptr)
 			{
-				const uint32 InitialPageTableSize = InDesc.bPrivateSpace ? InDesc.MaxSpaceSize : FMath::Max(AllocatedVT->GetWidthInTiles(), AllocatedVT->GetHeightInTiles());
-				FVirtualTextureSpace* Space = new FVirtualTextureSpace(this, SpaceIndex, InDesc, InitialPageTableSize);
-				Spaces[SpaceIndex].Reset(Space);
-				NumAllocatedSpaces++;
-				INC_MEMORY_STAT_BY(STAT_TotalPagetableMemory, Space->GetSizeInBytes());
-				Space->InitResource(RHICmdList);
-
-				const uint32 vAddress = Space->AllocateVirtualTexture(AllocatedVT);
-				AllocatedVT->AssignVirtualAddress(vAddress);
-
-				Space->AddRef();
-				return Space;
+				FreeSpaceIndex = SpaceIndex;
+			}
+		}
+		
+		// Try to recycle a space. We don't release spaces by default to avoid unnecessary page table recreation.
+		for (uint32 SpaceIndex = 0u; FreeSpaceIndex == INDEX_NONE && SpaceIndex < MaxSpaces; ++SpaceIndex)
+		{
+			if (Spaces[SpaceIndex] != nullptr && Spaces[SpaceIndex]->GetRefCount() == 0)
+			{
+				DEC_MEMORY_STAT_BY(STAT_TotalPagetableMemory, Spaces[SpaceIndex]->GetSizeInBytes());
+				Spaces[SpaceIndex]->ReleaseResource();
+				NumAllocatedSpaces--;
+
+				FreeSpaceIndex = SpaceIndex;
 			}
 		}
+
+		if (FreeSpaceIndex != INDEX_NONE)
+		{
+			const uint32 InitialPageTableSize = InDesc.bPrivateSpace ? InDesc.MaxSpaceSize : FMath::Max(AllocatedVT->GetWidthInTiles(), AllocatedVT->GetHeightInTiles());
+			FVirtualTextureSpace* Space = new FVirtualTextureSpace(this, FreeSpaceIndex, InDesc, InitialPageTableSize);
+			Spaces[FreeSpaceIndex].Reset(Space);
+
+			NumAllocatedSpaces++;
+			INC_MEMORY_STAT_BY(STAT_TotalPagetableMemory, Space->GetSizeInBytes());
+			Space->InitResource(RHICmdList);
+
+			const uint32 vAddress = Space->AllocateVirtualTexture(AllocatedVT);
+			AllocatedVT->AssignVirtualAddress(vAddress);
+
+			Space->AddRef();
+			return Space;
+		}
 	}
 
 	// Out of space slots
