[HWRT] Fix crash in FRayTracingDynamicGeometryUpdateManager due to missing VertexFactory when skeletal mesh section is disabled.
- improved handling of disabled sections.
- removed restrictions requiring 1:1 match of mesh batches and ray tracing geometry segments in dynamic path.
- skip updating geometries if all sections are disabled.
- added extra checks to better track invalid usage.

#[robomerge][partner] megalights

diff --git a/Engine/Plugins/Experimental/GPULightmass/Source/GPULightmass/Private/LightmapRenderer.cpp b/Engine/Plugins/Experimental/GPULightmass/Source/GPULightmass/Private/LightmapRenderer.cpp
--- a/Engine/Plugins/Experimental/GPULightmass/Source/GPULightmass/Private/LightmapRenderer.cpp
+++ b/Engine/Plugins/Experimental/GPULightmass/Source/GPULightmass/Private/LightmapRenderer.cpp
@@ -1127,7 +1127,7 @@ bool FSceneRenderState::SetupRayTracingScene(FRDGBuilder& GraphBuilder, FSceneUn
 						TArray<FMeshBatch> MeshBatches = Landscape.GetMeshBatchesForGBufferRendering(0);
 						ensure(MeshBatches.Num() == 1);
 
-						uint32 SegmentCount = MeshBatches.Num();
+						const uint32 SegmentCount = Landscape.SectionRayTracingStates[SubSectionIdx]->Geometry.Initializer.Segments.Num();
 						FRayTracingSBTAllocation* SBTAllocation = RayTracingSBT.AllocateDynamicRange(ERayTracingShaderBindingLayerMask::Base, SegmentCount);
 						RayTracingNumSegments += SegmentCount;
 
diff --git a/Engine/Source/Runtime/Apple/MetalRHI/Private/MetalRayTracing.cpp b/Engine/Source/Runtime/Apple/MetalRHI/Private/MetalRayTracing.cpp
--- a/Engine/Source/Runtime/Apple/MetalRHI/Private/MetalRayTracing.cpp
+++ b/Engine/Source/Runtime/Apple/MetalRHI/Private/MetalRayTracing.cpp
@@ -1003,6 +1003,8 @@ void FMetalRHICommandContext::RHIBuildAccelerationStructures(TConstArrayView<FRa
 
 			for (int32 i = 0; i < P.Segments.Num(); ++i)
 			{
+				checkf(P.Segments[i].VertexBuffer != nullptr, TEXT("Segments used to build/update ray tracing geometry must have a valid VertexBuffer."));
+
 				checkf(P.Segments[i].MaxVertices <= Geometry->Initializer.Segments[i].MaxVertices,
 					TEXT("Maximum number of vertices in a segment (%u) must not be smaller than what was declared during FRHIRayTracingGeometry creation (%u), as this controls BLAS memory allocation."),
 					P.Segments[i].MaxVertices, Geometry->Initializer.Segments[i].MaxVertices
diff --git a/Engine/Source/Runtime/D3D12RHI/Private/D3D12RayTracing.cpp b/Engine/Source/Runtime/D3D12RHI/Private/D3D12RayTracing.cpp
--- a/Engine/Source/Runtime/D3D12RHI/Private/D3D12RayTracing.cpp
+++ b/Engine/Source/Runtime/D3D12RHI/Private/D3D12RayTracing.cpp
@@ -4658,6 +4658,8 @@ void FD3D12CommandContext::RHIBuildAccelerationStructures(TConstArrayView<FRayTr
 
 			for (int32 i = 0; i < P.Segments.Num(); ++i)
 			{
+				checkf(P.Segments[i].VertexBuffer != nullptr, TEXT("Segments used to build/update ray tracing geometry must have a valid VertexBuffer."));
+
 				checkf(P.Segments[i].MaxVertices <= Geometry->Initializer.Segments[i].MaxVertices,
 					TEXT("Maximum number of vertices in a segment (%u) must not be larger than what was declared during FRHIRayTracingGeometry creation (%u), as this controls BLAS memory allocation."),
 					P.Segments[i].MaxVertices, Geometry->Initializer.Segments[i].MaxVertices
diff --git a/Engine/Source/Runtime/Engine/Private/SkeletalMeshSceneProxy.cpp b/Engine/Source/Runtime/Engine/Private/SkeletalMeshSceneProxy.cpp
--- a/Engine/Source/Runtime/Engine/Private/SkeletalMeshSceneProxy.cpp
+++ b/Engine/Source/Runtime/Engine/Private/SkeletalMeshSceneProxy.cpp
@@ -1033,6 +1033,12 @@ void FSkeletalMeshSceneProxy::GetDynamicRayTracingInstances(FRayTracingInstanceC
 	for (FSkeletalMeshSectionIter Iter(LODIndex, *MeshObject, LODData, LODSection, bIgnorePreviewFilter); Iter; ++Iter)
 	{
 		const FSkelMeshRenderSection& Section = Iter.GetSection();
+
+		if (Section.bDisabled)
+		{
+			continue;
+		}
+
 		const int32 SectionIndex = Iter.GetSectionElementIndex();
 		const FSectionElementInfo& SectionElementInfo = Iter.GetSectionElementInfo();
 
@@ -1054,6 +1060,11 @@ void FSkeletalMeshSceneProxy::GetDynamicRayTracingInstances(FRayTracingInstanceC
 		}
 	}
 
+	if (RayTracingInstance.Materials.IsEmpty())
+	{
+		return;
+	}
+
 	RayTracingInstance.InstanceTransforms.Add(GetLocalToWorld());
 
 	const uint32 VertexBufferNumVertices = LODData.StaticVertexBuffers.PositionVertexBuffer.GetNumVertices();
diff --git a/Engine/Source/Runtime/Renderer/Private/RayTracing/RayTracing.cpp b/Engine/Source/Runtime/Renderer/Private/RayTracing/RayTracing.cpp
--- a/Engine/Source/Runtime/Renderer/Private/RayTracing/RayTracing.cpp
+++ b/Engine/Source/Runtime/Renderer/Private/RayTracing/RayTracing.cpp
@@ -621,12 +621,11 @@ namespace RayTracing
 					{
 						const FRayTracingMeshBatchWorkItem& WorkItem = Page->WorkItems[ItemIndex];
 						TArrayView<const FMeshBatch> MeshBatches = WorkItem.GetMeshBatches();
-						for (int32 SegmentIndex = 0; SegmentIndex < MeshBatches.Num(); SegmentIndex++)
+						for (const FMeshBatch& MeshBatch : MeshBatches)
 						{
-							const FMeshBatch& MeshBatch = MeshBatches[SegmentIndex];
 							FDynamicRayTracingMeshCommandContext CommandContext(
 								*TaskDynamicCommandStorage, *TaskVisibleShaderBindings,
-								WorkItem.RayTracingGeometry, SegmentIndex, WorkItem.SBTAllocation);
+								WorkItem.RayTracingGeometry, MeshBatch.SegmentIndex, WorkItem.SBTAllocation);
 							FRayTracingMeshProcessor RayTracingMeshProcessor(&CommandContext, &Scene, &View, Scene.CachedRayTracingMeshCommandsType);
 							RayTracingMeshProcessor.AddMeshBatch(MeshBatch, 1, WorkItem.SceneProxy);
 						}
@@ -862,6 +861,7 @@ namespace RayTracing
 					const int32 ViewInstanceSceneDataOffset = View.RayTracingDynamicPrimitiveCollector.GetInstanceSceneDataOffset();
 
 					const FRayTracingGeometry* Geometry = Instance.Geometry;
+					const int32 NumSegments = Geometry->Initializer.Segments.Num();
 
 					if (!ensureMsgf(Geometry->DynamicGeometrySharedBufferGenerationID == FRayTracingGeometry::NonSharedVertexBuffers
 						|| Geometry->DynamicGeometrySharedBufferGenerationID == SharedBufferGenerationID,
@@ -888,31 +888,28 @@ namespace RayTracing
 					}
 
 					// Validate the material/segment counts
-					if (!ensureMsgf(Instance.GetMaterials().Num() == Geometry->Initializer.Segments.Num() ||
-						(Geometry->Initializer.Segments.Num() == 0 && Instance.GetMaterials().Num() == 1),
+					if (!ensureMsgf(Instance.GetMaterials().Num() <= NumSegments,
 						TEXT("Ray tracing material assignment validation failed for geometry '%s'. "
 							"Instance.GetMaterials().Num() = %d, Geometry->Initializer.Segments.Num() = %d."),
 						*Geometry->Initializer.DebugName.ToString(), Instance.GetMaterials().Num(),
-						Geometry->Initializer.Segments.Num()))
+						NumSegments))
 					{
 						continue;
 					}
 
+					if (Instance.GetMaterials().IsEmpty())
+					{
+						// If the material list is empty, skip this instance altogether
+						continue;
+					}
+
 					if (Instance.bInstanceMaskAndFlagsDirty || SceneInfo->bDynamicRayTracingInstanceCachedDataDirty)
 					{
 						// Build InstanceMaskAndFlags since the data in SceneInfo is not up to date
 
 						FRayTracingMaskAndFlags InstanceMaskAndFlags;
 
-						if (Instance.GetMaterials().IsEmpty())
-						{
-							// If the material list is empty, explicitly set the mask to 0 so it will not be added in the raytracing scene
-							InstanceMaskAndFlags.Mask = 0;
-						}
-						else
-						{
-							InstanceMaskAndFlags = BuildRayTracingInstanceMaskAndFlags(Instance, *SceneProxy);
-						}
+						InstanceMaskAndFlags = BuildRayTracingInstanceMaskAndFlags(Instance, *SceneProxy);
 
 						SceneInfo->DynamicRayTracingInstanceCachedData.Mask = InstanceMaskAndFlags.Mask;
 						SceneInfo->DynamicRayTracingInstanceCachedData.bForceOpaque = InstanceMaskAndFlags.bForceOpaque;
@@ -1029,7 +1026,7 @@ namespace RayTracing
 						EnumAddFlags(ActiveLayers, ERayTracingShaderBindingLayerMask::Decals);
 					}
 
-					FRayTracingSBTAllocation* SBTAllocation = RayTracingSBT.AllocateDynamicRange(ActiveLayers, Instance.GetMaterials().Num());
+					FRayTracingSBTAllocation* SBTAllocation = RayTracingSBT.AllocateDynamicRange(ActiveLayers, NumSegments);
 					if (bNeedMainInstance)
 					{
 						RayTracingInstance.InstanceContributionToHitGroupIndex = SBTAllocation->GetInstanceContributionToHitGroupIndex(ERayTracingShaderBindingLayer::Base);
@@ -1108,10 +1105,9 @@ namespace RayTracing
 					else
 					{
 						TArrayView<const FMeshBatch> InstanceMaterials = Instance.GetMaterials();
-						for (int32 SegmentIndex = 0; SegmentIndex < InstanceMaterials.Num(); SegmentIndex++)
+						for (const FMeshBatch& MeshBatch : InstanceMaterials)
 						{
-							const FMeshBatch& MeshBatch = InstanceMaterials[SegmentIndex];
-							FDynamicRayTracingMeshCommandContext CommandContext(View.DynamicRayTracingMeshCommandStorage, View.VisibleRayTracingShaderBindings, Geometry->GetRHI(), SegmentIndex, SBTAllocation);
+							FDynamicRayTracingMeshCommandContext CommandContext(View.DynamicRayTracingMeshCommandStorage, View.VisibleRayTracingShaderBindings, Geometry->GetRHI(), MeshBatch.SegmentIndex, SBTAllocation);
 							FRayTracingMeshProcessor RayTracingMeshProcessor(&CommandContext, &Scene, &View, Scene.CachedRayTracingMeshCommandsType);
 							RayTracingMeshProcessor.AddMeshBatch(MeshBatch, 1, SceneProxy);
 						}
diff --git a/Engine/Source/Runtime/Renderer/Private/RayTracing/RayTracingDynamicGeometryUpdateManager.cpp b/Engine/Source/Runtime/Renderer/Private/RayTracing/RayTracingDynamicGeometryUpdateManager.cpp
--- a/Engine/Source/Runtime/Renderer/Private/RayTracing/RayTracingDynamicGeometryUpdateManager.cpp
+++ b/Engine/Source/Runtime/Renderer/Private/RayTracing/RayTracingDynamicGeometryUpdateManager.cpp
@@ -388,10 +388,19 @@ void FRayTracingDynamicGeometryUpdateManager::AddDynamicGeometryToUpdate(
 			Segment.VertexBufferOffset = VertexBufferOffset;
 		}
 	}
+#if DO_CHECK
+	else
+	{
+		for (FRayTracingGeometrySegment& Segment : Geometry.Initializer.Segments)
+		{
+			checkf(Segment.VertexBuffer != nullptr, TEXT("Dynamic ray tracing geometry '%s' has a segment without a valid VertexBuffer."), *Geometry.Initializer.DebugName.ToString());
+		}
+	}
+#endif
 
 	if (!bRefit)
 	{
-		checkf(Geometry.RawData.IsEmpty() && Geometry.Initializer.OfflineData == nullptr, TEXT("Dynamic geometry is not expected to have offline acceleration structure data"));
+		checkf(Geometry.RawData.IsEmpty() && Geometry.Initializer.OfflineData == nullptr, TEXT("Dynamic ray tracing geometry '%s' is not expected to have offline acceleration structure data."), *Geometry.Initializer.DebugName.ToString());
 		Geometry.CreateRayTracingGeometry(RHICmdList, ERTAccelerationStructureBuildPriority::Skip);
 	}
 
diff --git a/Engine/Source/Runtime/VulkanRHI/Private/VulkanRayTracing.cpp b/Engine/Source/Runtime/VulkanRHI/Private/VulkanRayTracing.cpp
--- a/Engine/Source/Runtime/VulkanRHI/Private/VulkanRayTracing.cpp
+++ b/Engine/Source/Runtime/VulkanRHI/Private/VulkanRayTracing.cpp
@@ -1177,6 +1177,8 @@ void FVulkanCommandListContext::RHIBuildAccelerationStructures(TConstArrayView<F
 
 			for (int32 i = 0; i < P.Segments.Num(); ++i)
 			{
+				checkf(P.Segments[i].VertexBuffer != nullptr, TEXT("Segments used to build/update ray tracing geometry must have a valid VertexBuffer."));
+
 				checkf(P.Segments[i].MaxVertices <= Geometry->Initializer.Segments[i].MaxVertices,
 					TEXT("Maximum number of vertices in a segment (%u) must not be smaller than what was declared during FRHIRayTracingGeometry creation (%u), as this controls BLAS memory allocation."),
 					P.Segments[i].MaxVertices, Geometry->Initializer.Segments[i].MaxVertices
