Add r.VT.SpaceReleaseFrames to allow release/recycle of virtual texture spaces after they are unused for some time.
We don't release immediately to avoid cases during loading where this might trigger unnecessary recreate/resize.
A value of -1 will default to old behavior where we never release/recycle unless we run out of space slots.
We default to 150 frames.

diff --git a/Engine/Source/Runtime/Renderer/Private/VT/VirtualTextureSpace.cpp b/Engine/Source/Runtime/Renderer/Private/VT/VirtualTextureSpace.cpp
--- a/Engine/Source/Runtime/Renderer/Private/VT/VirtualTextureSpace.cpp
+++ b/Engine/Source/Runtime/Renderer/Private/VT/VirtualTextureSpace.cpp
@@ -50,13 +50,9 @@ static EPixelFormat GetFormatForNumLayers(uint32 NumLayers, EVTPageTableFormat F
 FVirtualTextureSpace::FVirtualTextureSpace(FVirtualTextureSystem* InSystem, uint8 InID, const FVTSpaceDescription& InDesc, uint32 InSizeNeeded)
 	: Description(InDesc)
 	, Allocator(InDesc.Dimensions)
-	, NumRefs(0u)
 	, ID(InID)
-	, CachedPageTableWidth(0u)
-	, CachedPageTableHeight(0u)
-	, CachedNumPageTableLevels(0u)
 	, bNeedToAllocatePageTable(true)
-	, bForceEntireUpdate(false)
+	, bNeedToAllocatePageTableIndirection(InDesc.IndirectionTextureSize > 0)
 {
 	// Initialize page map with large enough capacity to handle largest possible physical texture
 	const uint32 PhysicalTileSize = InDesc.TileSize + InDesc.TileBorderSize * 2u;
@@ -97,7 +93,6 @@ FVirtualTextureSpace::FVirtualTextureSpace(FVirtualTextureSystem* InSystem, uint
 
 	Allocator.Initialize(Description.MaxSpaceSize);
 
-	bNeedToAllocatePageTableIndirection = InDesc.IndirectionTextureSize > 0;
 }
 
 FVirtualTextureSpace::~FVirtualTextureSpace()
diff --git a/Engine/Source/Runtime/Renderer/Private/VT/VirtualTextureSpace.h b/Engine/Source/Runtime/Renderer/Private/VT/VirtualTextureSpace.h
--- a/Engine/Source/Runtime/Renderer/Private/VT/VirtualTextureSpace.h
+++ b/Engine/Source/Runtime/Renderer/Private/VT/VirtualTextureSpace.h
@@ -75,6 +75,9 @@ public:
 	inline uint32 Release() { check(NumRefs > 0u); return --NumRefs; }
 	inline uint32 GetRefCount() const { return NumRefs; }
 
+	inline void SetReleasedFrame(uint32 InReleasedFrame) { ReleasedFrame = InReleasedFrame; }
+	inline uint32 GetReleasedFrame() const { return ReleasedFrame; }
+
 	uint32 GetSizeInBytes() const;
 
 	uint32 AllocateVirtualTexture(FAllocatedVirtualTexture* VirtualTexture);
@@ -136,13 +139,14 @@ private:
 	FBufferRHIRef UpdateBuffer;
 	FShaderResourceViewRHIRef UpdateBufferSRV;
 
-	uint32 NumRefs;
+	uint32 NumRefs = 0;
+	uint32 ReleasedFrame = 0;
 
-	uint8 ID;
-	uint32 CachedPageTableWidth;
-	uint32 CachedPageTableHeight;
-	uint32 CachedNumPageTableLevels;
-	bool bNeedToAllocatePageTable;
-	bool bNeedToAllocatePageTableIndirection;
-	bool bForceEntireUpdate;
+	uint8 ID = 0;
+	uint32 CachedPageTableWidth = 0;
+	uint32 CachedPageTableHeight = 0;
+	uint32 CachedNumPageTableLevels = 0;
+	bool bNeedToAllocatePageTable = false;
+	bool bNeedToAllocatePageTableIndirection = false;
+	bool bForceEntireUpdate = false;
 };
diff --git a/Engine/Source/Runtime/Renderer/Private/VT/VirtualTextureSystem.cpp b/Engine/Source/Runtime/Renderer/Private/VT/VirtualTextureSystem.cpp
--- a/Engine/Source/Runtime/Renderer/Private/VT/VirtualTextureSystem.cpp
+++ b/Engine/Source/Runtime/Renderer/Private/VT/VirtualTextureSystem.cpp
@@ -105,6 +105,12 @@ static TAutoConsoleVariable<int32> CVarVTNumFeedbackTasks(
 	TEXT("Number of tasks to create to read virtual texture feedback."),
 	ECVF_RenderThreadSafe
 );
+static TAutoConsoleVariable<int32> CVarVTNumGatherTasks(
+	TEXT("r.VT.NumGatherTasks"),
+	1,
+	TEXT("Number of tasks to create to combine virtual texture feedback."),
+	ECVF_RenderThreadSafe
+);
 static TAutoConsoleVariable<int32> CVarVTEnablePlayback(
 	TEXT("r.VT.EnablePlayback"),
 	1,
@@ -117,12 +123,6 @@ static TAutoConsoleVariable<float> CVarVTPlaybackMipBias(
 	TEXT("Mip bias to apply during playback of recorded feedback requests."),
 	ECVF_RenderThreadSafe
 );
-static TAutoConsoleVariable<int32> CVarVTNumGatherTasks(
-	TEXT("r.VT.NumGatherTasks"),
-	1,
-	TEXT("Number of tasks to create to combine virtual texture feedback."),
-	ECVF_RenderThreadSafe
-);
 static TAutoConsoleVariable<int32> CVarVTPageUpdateFlushCount(
 	TEXT("r.VT.PageUpdateFlushCount"),
 	8,
@@ -147,6 +147,12 @@ static TAutoConsoleVariable<int32> CVarVTProduceLockedTilesOnFlush(
 	TEXT("Should locked tiles be (re)produced when flushing the cache"),
 	ECVF_RenderThreadSafe
 );
+static TAutoConsoleVariable<int32> CVarVTSpaceReleaseFrames(
+	TEXT("r.VT.SpaceReleaseFrames"),
+	150,
+	TEXT("Number of frames to wait before releasing a virtual texture space. < 0 for no release."),
+	ECVF_RenderThreadSafe
+);
 static TAutoConsoleVariable<int32> CVarVTResidencyShow(
 	TEXT("r.VT.Residency.Show"),
 	0,
@@ -916,19 +922,20 @@ FVirtualTextureSpace* FVirtualTextureSystem::AcquireSpace(FRHICommandListBase& R
 		int32 FreeSpaceIndex = INDEX_NONE;
 		for (uint32 SpaceIndex = 0u; FreeSpaceIndex == INDEX_NONE && SpaceIndex < MaxSpaces; ++SpaceIndex)
 		{
-			if (Spaces[SpaceIndex] == nullptr)
+			if (!Spaces[SpaceIndex].IsValid())
 			{
 				FreeSpaceIndex = SpaceIndex;
 			}
 		}
-		
-		// Try to recycle a space. We don't release spaces by default to avoid unnecessary page table recreation.
+
+		// Try to recycle a space. We don't release spaces immediately to avoid unnecessary page table recreation.
 		for (uint32 SpaceIndex = 0u; FreeSpaceIndex == INDEX_NONE && SpaceIndex < MaxSpaces; ++SpaceIndex)
 		{
-			if (Spaces[SpaceIndex] != nullptr && Spaces[SpaceIndex]->GetRefCount() == 0)
+			FVirtualTextureSpace* Space = Spaces[SpaceIndex].Get();
+			if (Space != nullptr && Space->GetRefCount() == 0)
 			{
-				DEC_MEMORY_STAT_BY(STAT_TotalPagetableMemory, Spaces[SpaceIndex]->GetSizeInBytes());
-				Spaces[SpaceIndex]->ReleaseResource();
+				DEC_MEMORY_STAT_BY(STAT_TotalPagetableMemory, Space->GetSizeInBytes());
+				Space->ReleaseResource();
 				NumAllocatedSpaces--;
 
 				FreeSpaceIndex = SpaceIndex;
@@ -963,14 +970,24 @@ void FVirtualTextureSystem::ReleaseSpace(FVirtualTextureSpace* Space)
 {
 	check(!bUpdating);
 	const uint32 NumRefs = Space->Release();
-	if (NumRefs == 0u && Space->GetDescription().bPrivateSpace)
+	if (NumRefs == 0u)
 	{
-		// Private spaces are destroyed when ref count reaches 0
-		// This can only happen on render thread, so we can call ReleaseResource() directly and then delete the pointer immediately
-		DEC_MEMORY_STAT_BY(STAT_TotalPagetableMemory, Space->GetSizeInBytes());
-		Space->ReleaseResource();
-		Spaces[Space->GetID()].Reset();
-		NumAllocatedSpaces--;
+		// Private spaces can't be recycled and so are destroyed immediately when ref count reaches 0.
+		// Other spaces wait on a timer to avoid unnecessary recreation during load phases etc.
+		const int32 ReleaseFrames = CVarVTSpaceReleaseFrames.GetValueOnRenderThread();
+
+		if (Space->GetDescription().bPrivateSpace || ReleaseFrames == 0)
+		{
+			// We only get here on the render thread, so we can call ReleaseResource() directly and then delete the pointer immediately
+			DEC_MEMORY_STAT_BY(STAT_TotalPagetableMemory, Space->GetSizeInBytes());
+			Space->ReleaseResource();
+			Spaces[Space->GetID()].Reset();
+			NumAllocatedSpaces--;
+		}
+		else
+		{
+			Space->SetReleasedFrame(Frame);
+		}
 	}
 }
 
@@ -1122,14 +1139,13 @@ FVirtualTexturePhysicalSpace* FVirtualTextureSystem::AcquirePhysicalSpace(FRHICo
 	return PhysicalSpace;
 }
 
-void FVirtualTextureSystem::ReleasePendingSpaces()
+void FVirtualTextureSystem::ReleasePendingSpaces(bool bForceReleaseAll)
 {
 	for (int32 Id = 0; Id < PhysicalSpaces.Num(); ++Id)
 	{
-		// Physical space is released when ref count hits 0
-		// Might need to have some mechanism to hold an extra reference if we know we will be recycling very soon (such when doing level reload)
+		// Physical space is immediately released when ref count hits 0.
 		FVirtualTexturePhysicalSpace* PhysicalSpace = PhysicalSpaces[Id];
-		if ((bool)PhysicalSpace && PhysicalSpace->GetRefCount() == 0u)
+		if (PhysicalSpace != nullptr && PhysicalSpace->GetRefCount() == 0u)
 		{
 			const FTexturePagePool& PagePool = PhysicalSpace->GetPagePool();
 			check(PagePool.GetNumMappedPages() == 0u);
@@ -1140,6 +1156,23 @@ void FVirtualTextureSystem::ReleasePendingSpaces()
 			PhysicalSpaces[Id] = nullptr;
 		}
 	}
+
+	for (int32 Id = 0; Id < MaxSpaces; ++Id)
+	{
+		// Virtual space is released when ref count hits 0 and after some delay to avoid unnecessary recreation.
+		FVirtualTextureSpace* Space = Spaces[Id].Get();
+		if (Space != nullptr && Space->GetRefCount() == 0u)
+		{
+			const int32 ReleaseFrames = CVarVTSpaceReleaseFrames.GetValueOnRenderThread();
+			if (bForceReleaseAll || (ReleaseFrames >= 0 && Space->GetReleasedFrame() + ReleaseFrames <= Frame))
+			{
+				DEC_MEMORY_STAT_BY(STAT_TotalPagetableMemory, Spaces[Id]->GetSizeInBytes());
+				Space->ReleaseResource();
+				Spaces[Id].Reset();
+				NumAllocatedSpaces--;
+			}
+		}
+	}
 }
 
 void FVirtualTextureSystem::LockTile(const FVirtualTextureLocalTile& Tile)
@@ -2996,7 +3029,7 @@ TUniquePtr<FVirtualTextureUpdater> FVirtualTextureSystem::BeginUpdate(FRDGBuilde
 	}
 
 	DestroyPendingVirtualTextures(false);
-	ReleasePendingSpaces();
+	ReleasePendingSpaces(false);
 
 	// Early out when no allocated VTs
 	if (AllocatedVTs.Num() == 0)
@@ -3117,7 +3150,7 @@ void FVirtualTextureSystem::ReleasePendingResources()
 	check(!bUpdating);
 	UE::TScopeLock Lock(Mutex);
 	DestroyPendingVirtualTextures(true);
-	ReleasePendingSpaces();
+	ReleasePendingSpaces(true);
 }
 
 FVector4f FVirtualTextureSystem::GetGlobalMipBias() const
diff --git a/Engine/Source/Runtime/Renderer/Private/VT/VirtualTextureSystem.h b/Engine/Source/Runtime/Renderer/Private/VT/VirtualTextureSystem.h
--- a/Engine/Source/Runtime/Renderer/Private/VT/VirtualTextureSystem.h
+++ b/Engine/Source/Runtime/Renderer/Private/VT/VirtualTextureSystem.h
@@ -202,7 +202,7 @@ private:
 
 	void AllocateResources(FRDGBuilder& GraphBuilder);
 	void DestroyPendingVirtualTextures(bool bForceDestroyAll);
-	void ReleasePendingSpaces();
+	void ReleasePendingSpaces(bool bForceDestroyAll);
 
 	void RequestTilesForRegionInternal(const IAllocatedVirtualTexture* AllocatedVT, const FVector2D& InScreenSpaceSize, const FVector2D& InViewportPosition, const FVector2D& InViewportSize, const FVector2D& InUV0, const FVector2D& InUV1, int32 InMipLevel);
 	void RequestTilesInternal(const IAllocatedVirtualTexture* AllocatedVT, int32 InMipLevel);
