Don't allow splitting of physical pools for RVT.
Always allow 16bit entries for private page tables even when r.VT.PageTableMode=0.
#rnx

diff --git a/Engine/Source/Runtime/Engine/Private/VT/RuntimeVirtualTexture.cpp b/Engine/Source/Runtime/Engine/Private/VT/RuntimeVirtualTexture.cpp
--- a/Engine/Source/Runtime/Engine/Private/VT/RuntimeVirtualTexture.cpp
+++ b/Engine/Source/Runtime/Engine/Private/VT/RuntimeVirtualTexture.cpp
@@ -355,6 +355,10 @@ void URuntimeVirtualTexture::GetProducerDescription(FVTProducerDescription& OutD
 	OutDesc.MaxLevel = FMath::Max((int32)FMath::CeilLogTwo(FMath::Max(OutDesc.BlockWidthInTiles, OutDesc.BlockHeightInTiles)) - GetRemoveLowMips(), 0);
 
 	OutDesc.bContinuousUpdate = bContinuousUpdate;
+	
+	// Disable splitting physical pools to give 16bit page tables when r.VT.PageTableMode=2.
+	// This is only useful for virtual texture spaces with many small allocated page tables, and that isn't typically the case for RVT.
+	OutDesc.bRequiresSinglePhysicalPool = true;
 
 	OutDesc.NumTextureLayers = GetLayerCount();
 	OutDesc.NumPhysicalGroups = bSinglePhysicalSpace ? 1 : GetLayerCount();
diff --git a/Engine/Source/Runtime/Renderer/Private/VT/AllocatedVirtualTexture.cpp b/Engine/Source/Runtime/Renderer/Private/VT/AllocatedVirtualTexture.cpp
--- a/Engine/Source/Runtime/Renderer/Private/VT/AllocatedVirtualTexture.cpp
+++ b/Engine/Source/Runtime/Renderer/Private/VT/AllocatedVirtualTexture.cpp
@@ -77,15 +77,16 @@ FAllocatedVirtualTexture::FAllocatedVirtualTexture(
 	// Lock tiles
 	LockOrUnlockTiles(InSystem, true);
 
-	// Use 16bit page table entries if all physical spaces are small enough
-	bool bSupport16BitPageTable = VirtualTexturePool::Support16BitPageTables();
-	for (int32 Index = 0; Index < UniquePageTableLayers.Num(); ++Index)
+	// Use 16bit page table entries if all physical spaces are small enough.
+	// We gate this on the config support because using both 16bit and 32bit spaces can increase the total number of virtual texture spaces used and there is a hard limit.
+	// But private spaces always use their own space, so we ignore the config for that case.
+	bool bSupport16BitPageTable = VirtualTexturePool::Support16BitPageTables() || InDesc.bPrivateSpace;
+	for (int32 Index = 0; bSupport16BitPageTable && Index < UniquePageTableLayers.Num(); ++Index)
 	{
 		const FVirtualTexturePhysicalSpace* PhysicalSpace = UniquePageTableLayers[Index].PhysicalSpace;
 		if (PhysicalSpace && !PhysicalSpace->DoesSupport16BitPageTable())
 		{
 			bSupport16BitPageTable = false;
-			break;
 		}
 	}
 
