Mobile: Revert ShadowMask pixel format back to RGBA8 in all cases.
Using single channel texture is a bit slower on some platforms

diff --git a/Engine/Shaders/Private/MobileLightingCommon.ush b/Engine/Shaders/Private/MobileLightingCommon.ush
--- a/Engine/Shaders/Private/MobileLightingCommon.ush
+++ b/Engine/Shaders/Private/MobileLightingCommon.ush
@@ -276,15 +276,9 @@ half4 SampleScreenSpaceShadowMask(float2 BufferUV, uint EyeIndex)
 {
 #if ENABLE_SHADOWMASKTEXTURE 
 	#if MOBILE_MULTI_VIEW
-		half4 Sample = Texture2DArraySample(MobileBasePass.ScreenSpaceShadowMaskTextureArray, MobileBasePass.ScreenSpaceShadowMaskSampler, float3(BufferUV, EyeIndex));
+		return Texture2DArraySample(MobileBasePass.ScreenSpaceShadowMaskTextureArray, MobileBasePass.ScreenSpaceShadowMaskSampler, float3(BufferUV, EyeIndex));
 	#else
-		half4 Sample = Texture2DSample(MobileBasePass.ScreenSpaceShadowMaskTexture, MobileBasePass.ScreenSpaceShadowMaskSampler, BufferUV);
-	#endif
-
-	#if ALLOW_STATIC_LIGHTING
-		return Sample;
-	#else	
-		return half4(Sample.r, 1.0, 1.0, 1.0); // PG_R8
+		return Texture2DSample(MobileBasePass.ScreenSpaceShadowMaskTexture, MobileBasePass.ScreenSpaceShadowMaskSampler, BufferUV);
 	#endif
 #endif
 	return 1.0;
diff --git a/Engine/Source/Runtime/Renderer/Private/ShadowRendering.cpp b/Engine/Source/Runtime/Renderer/Private/ShadowRendering.cpp
--- a/Engine/Source/Runtime/Renderer/Private/ShadowRendering.cpp
+++ b/Engine/Source/Runtime/Renderer/Private/ShadowRendering.cpp
@@ -2429,14 +2429,7 @@ FRDGTextureRef FMobileSceneRenderer::RenderMobileShadowProjections(FRDGBuilder&
 {
 	RDG_RHI_EVENT_SCOPE(GraphBuilder, RenderMobileShadowProjections);
 	
-	EPixelFormat ShadowMaskPixelFormat = PF_B8G8R8A8;
-	if (!IsStaticLightingAllowed())
-	{
-		// With pre-computed lighting disabled all lights use single channel
-		ShadowMaskPixelFormat = PF_R8;
-	}
-	
-	FRDGTextureDesc TargetDesc = FRDGTextureDesc::CreateRenderTargetTextureDesc(Extent, ShadowMaskPixelFormat, FClearValueBinding::White, TexCreate_RenderTargetable | TexCreate_ShaderResource, bIsMobileMultiView);
+	FRDGTextureDesc TargetDesc = FRDGTextureDesc::CreateRenderTargetTextureDesc(Extent, PF_B8G8R8A8, FClearValueBinding::White, TexCreate_RenderTargetable | TexCreate_ShaderResource, bIsMobileMultiView);
 	FRDGTextureRef ScreenShadowMaskTexture = GraphBuilder.CreateTexture(TargetDesc, TEXT("ScreenSpaceShadowMaskTextureMobile"));
 	
 	FRDGTextureClearInfo TextureClearInfo;
