Fixed: Crash on a decal PSO precaching on platforms with a VulkanRHI. Incorrect sub-pass setup

diff --git a/Engine/Source/Runtime/Renderer/Private/CompositionLighting/PostProcessDeferredDecals.cpp b/Engine/Source/Runtime/Renderer/Private/CompositionLighting/PostProcessDeferredDecals.cpp
--- a/Engine/Source/Runtime/Renderer/Private/CompositionLighting/PostProcessDeferredDecals.cpp
+++ b/Engine/Source/Runtime/Renderer/Private/CompositionLighting/PostProcessDeferredDecals.cpp
@@ -588,8 +588,10 @@ void CollectDeferredDecalPassPSOInitializers(
 	if (FeatureLevel == ERHIFeatureLevel::ES3_1)
 	{
 		// subpass info set during the submission of the draws in a mobile renderer
-		GraphicsPSOInit.SubpassIndex = 1; // all decals use second sub-pass on mobile
 		GraphicsPSOInit.SubpassHint = GetSubpassHint(SceneTexturesConfig.ShaderPlatform, SceneTexturesConfig.bIsUsingGBuffers, SceneTexturesConfig.bRequireMultiView, SceneTexturesConfig.NumSamples);
+		// all decals use second sub-pass on mobile
+		GraphicsPSOInit.SubpassIndex = (GraphicsPSOInit.SubpassHint != ESubpassHint::None) ? 1 : 0;
+		
 	}
 		
 	const auto AddDeferredDecalPSO = [&](bool bInsideDecal,	bool bReverseHanded, bool bReverseCulling, bool bDecalUsesStencil)
diff --git a/Engine/Source/Runtime/Renderer/Private/CompositionLighting/PostProcessMeshDecals.cpp b/Engine/Source/Runtime/Renderer/Private/CompositionLighting/PostProcessMeshDecals.cpp
--- a/Engine/Source/Runtime/Renderer/Private/CompositionLighting/PostProcessMeshDecals.cpp
+++ b/Engine/Source/Runtime/Renderer/Private/CompositionLighting/PostProcessMeshDecals.cpp
@@ -399,8 +399,9 @@ void FMeshDecalMeshProcessor::CollectDeferredDecalMeshPSOInitializers(
 	if (FeatureLevel == ERHIFeatureLevel::ES3_1)
 	{
 		// subpass info set during the submission of the draws in a mobile renderer
-		SubpassIndex = 1; // all decals use second sub-pass on mobile
 		SubpassHint = GetSubpassHint(SceneTexturesConfig.ShaderPlatform, SceneTexturesConfig.bIsUsingGBuffers, SceneTexturesConfig.bRequireMultiView, SceneTexturesConfig.NumSamples);
+		// all decals use second sub-pass on mobile
+		SubpassIndex = (SubpassHint != ESubpassHint::None) ? 1 : 0;
 	}
 
 	AddGraphicsPipelineStateInitializer(
