Lumen HWRT - skip creating RGS when Lumen is disabled

diff --git a/Engine/Source/Runtime/Renderer/Private/DeferredShadingRenderer.cpp b/Engine/Source/Runtime/Renderer/Private/DeferredShadingRenderer.cpp
--- a/Engine/Source/Runtime/Renderer/Private/DeferredShadingRenderer.cpp
+++ b/Engine/Source/Runtime/Renderer/Private/DeferredShadingRenderer.cpp
@@ -702,7 +702,9 @@ bool FDeferredShadingSceneRenderer::SetupRayTracingPipelineStatesAndSBT(FRDGBuil
 					PrepareRayTracingTranslucency(View, RayGenShaders);
 					PrepareRayTracingVolumetricFogShadows(View, *Scene, RayGenShaders);
 
-					if (DoesPlatformSupportLumenGI(ShaderPlatform) && Lumen::UseHardwareRayTracing(ViewFamily))
+					if (DoesPlatformSupportLumenGI(ShaderPlatform) 
+						&& Lumen::UseHardwareRayTracing(ViewFamily)
+						&& IsLumenEnabled(View))
 					{
 						PrepareLumenHardwareRayTracingScreenProbeGather(View, RayGenShaders);
 						PrepareLumenHardwareRayTracingShortRangeAO(View, RayGenShaders);
@@ -761,18 +763,21 @@ bool FDeferredShadingSceneRenderer::SetupRayTracingPipelineStatesAndSBT(FRDGBuil
 
 		if (!bIsPathTracing)
 		{
-			if (DoesPlatformSupportLumenGI(ShaderPlatform))
+			if (DoesPlatformSupportLumenGI(ShaderPlatform) && Lumen::UseHardwareRayTracing(ViewFamily))
 			{
 				for (const FViewInfo& View : Views)
 				{
-					PrepareLumenHardwareRayTracingVisualizeLumenMaterial(View, LumenHardwareRayTracingRayGenShaders);
-					PrepareLumenHardwareRayTracingRadianceCacheLumenMaterial(View, LumenHardwareRayTracingRayGenShaders);
-					PrepareLumenHardwareRayTracingTranslucencyVolumeLumenMaterial(View, LumenHardwareRayTracingRayGenShaders);
-					PrepareLumenHardwareRayTracingRadiosityLumenMaterial(View, LumenHardwareRayTracingRayGenShaders);
-					PrepareLumenHardwareRayTracingReflectionsLumenMaterial(View, LumenHardwareRayTracingRayGenShaders);
-					PrepareLumenHardwareRayTracingReSTIRLumenMaterial(View, LumenHardwareRayTracingRayGenShaders);
-					PrepareLumenHardwareRayTracingScreenProbeGatherLumenMaterial(View, LumenHardwareRayTracingRayGenShaders);
-					PrepareLumenHardwareRayTracingDirectLightingLumenMaterial(View, LumenHardwareRayTracingRayGenShaders);
+					if (IsLumenEnabled(View))
+					{
+						PrepareLumenHardwareRayTracingVisualizeLumenMaterial(View, LumenHardwareRayTracingRayGenShaders);
+						PrepareLumenHardwareRayTracingRadianceCacheLumenMaterial(View, LumenHardwareRayTracingRayGenShaders);
+						PrepareLumenHardwareRayTracingTranslucencyVolumeLumenMaterial(View, LumenHardwareRayTracingRayGenShaders);
+						PrepareLumenHardwareRayTracingRadiosityLumenMaterial(View, LumenHardwareRayTracingRayGenShaders);
+						PrepareLumenHardwareRayTracingReflectionsLumenMaterial(View, LumenHardwareRayTracingRayGenShaders);
+						PrepareLumenHardwareRayTracingReSTIRLumenMaterial(View, LumenHardwareRayTracingRayGenShaders);
+						PrepareLumenHardwareRayTracingScreenProbeGatherLumenMaterial(View, LumenHardwareRayTracingRayGenShaders);
+						PrepareLumenHardwareRayTracingDirectLightingLumenMaterial(View, LumenHardwareRayTracingRayGenShaders);
+					}
 				}
 			}
 
