Only call translucency clipped depth pass when temporal responsiveness node is used in the material.

#tests on test map

diff --git a/Engine/Source/Runtime/Renderer/Private/SceneVisibility.cpp b/Engine/Source/Runtime/Renderer/Private/SceneVisibility.cpp
--- a/Engine/Source/Runtime/Renderer/Private/SceneVisibility.cpp
+++ b/Engine/Source/Runtime/Renderer/Private/SceneVisibility.cpp
@@ -1648,7 +1648,8 @@ void FRelevancePacket::ComputeRelevance(FDynamicPrimitiveIndexList& DynamicPrimi
 											DrawCommandPacket.AddCommandsForMesh(PrimitiveIndex, PrimitiveSceneInfo, StaticMeshRelevance, StaticMesh, CullingPayloadFlags, Scene, bCanCache, EMeshPass::TranslucentVelocity);
 											
 											// Only Before DoF Translucency needs this pass
-											if (ViewRelevance.bNormalTranslucency && bIsTranslucentClippedDepthEnabled)
+											const bool bAnyMaterialUsesTemporalResponsiveness = PrimitiveSceneProxy->AnyMaterialUsesTemporalResponsiveness();
+											if (ViewRelevance.bNormalTranslucency && bIsTranslucentClippedDepthEnabled && bAnyMaterialUsesTemporalResponsiveness)
 											{
 												DrawCommandPacket.AddCommandsForMesh(PrimitiveIndex, PrimitiveSceneInfo, StaticMeshRelevance, StaticMesh, CullingPayloadFlags, Scene, bCanCache, EMeshPass::TranslucentVelocityClippedDepth);
 											}
@@ -2466,8 +2467,8 @@ static void ComputeDynamicMeshRelevance(
 			{
 				PassMask.Set(EMeshPass::TranslucentVelocity);
 				View.NumVisibleDynamicMeshElements[EMeshPass::TranslucentVelocity] += NumElements;
-
-				if (ViewRelevance.bNormalTranslucency && bIsTranslucentClippedDepthEnabled)
+				const bool bAnyMaterialUsesTemporalResponsiveness = MeshBatch.PrimitiveSceneProxy->AnyMaterialUsesTemporalResponsiveness();
+				if (ViewRelevance.bNormalTranslucency && bIsTranslucentClippedDepthEnabled && bAnyMaterialUsesTemporalResponsiveness)
 				{
 					PassMask.Set(EMeshPass::TranslucentVelocityClippedDepth);
 					View.NumVisibleDynamicMeshElements[EMeshPass::TranslucentVelocityClippedDepth] += NumElements;
diff --git a/Engine/Source/Runtime/Renderer/Private/VelocityRendering.cpp b/Engine/Source/Runtime/Renderer/Private/VelocityRendering.cpp
--- a/Engine/Source/Runtime/Renderer/Private/VelocityRendering.cpp
+++ b/Engine/Source/Runtime/Renderer/Private/VelocityRendering.cpp
@@ -66,7 +66,7 @@ static TAutoConsoleVariable<int32> CVarVelocityOutputTranslucentClippedDepthEnab
 	1,
 	TEXT("Enable/Disable the translucent velocity clipped depth pass on the fly.\n")\
 	TEXT("0: Skip this pass.\n")\
-	TEXT("1: Provide functions e.g., mark before DoF translucency pixel animation for opacity below clip value.\n"),
+	TEXT("1: Provide functions e.g., mark before DoF translucency Temporal Responsiveness for opacity below clip value.\n"),
 	ECVF_RenderThreadSafe);
 
 static TAutoConsoleVariable<int32> CVarVelocityTemporalResponsivenessSupported(
