Remove check() that fires when we allocate more than 16 virtual texture spaces.
We now allow allocating more spaces but only support virtual texture feedback for the first 16.
We warn (instead of crashing) if the 16 limit is reached, and there are several suggested configuration changes that can be made if that happens.
Other recent changes for recycling spaces should mean this condition is now rarely if ever hit anyway.
Also modified some logging code for consistant use of logging channels, and using the log verbosity system instead of rolling our own with r.vt.verbose.
#rnx

diff --git a/Engine/Source/Runtime/Engine/Private/VT/VirtualTextureChunkManager.cpp b/Engine/Source/Runtime/Engine/Private/VT/VirtualTextureChunkManager.cpp
--- a/Engine/Source/Runtime/Engine/Private/VT/VirtualTextureChunkManager.cpp
+++ b/Engine/Source/Runtime/Engine/Private/VT/VirtualTextureChunkManager.cpp
@@ -69,12 +69,8 @@ FVTRequestPageResult FVirtualTextureChunkStreamingManager::RequestTile(FRHIComma
 	const uint32 TileOffset = VTData->GetTileOffset(vLevel, vAddress, 0);
 	if (TileOffset == ~0u)
 	{
-		static const auto CVar = IConsoleManager::Get().FindTConsoleVariableDataInt(TEXT("r.VT.Verbose"));
-		if (CVar->GetValueOnRenderThread())
-		{
-			UE_LOG(LogConsoleResponse, Display, TEXT("vAddr %i@%i has an invalid tile (-1)."), vAddress, vLevel);
-		}
-
+		UE_LOG(LogConsoleResponse, Verbose, TEXT("vAddr %i@%i has an invalid tile (-1)."), vAddress, vLevel);
+	
 		return EVTRequestPageStatus::Invalid;
 	}
 
diff --git a/Engine/Source/Runtime/RenderCore/Public/VirtualTexturing.h b/Engine/Source/Runtime/RenderCore/Public/VirtualTexturing.h
--- a/Engine/Source/Runtime/RenderCore/Public/VirtualTexturing.h
+++ b/Engine/Source/Runtime/RenderCore/Public/VirtualTexturing.h
@@ -50,6 +50,9 @@ union FVirtualTextureProducerHandle
 };
 static_assert(sizeof(FVirtualTextureProducerHandle) == sizeof(uint32), "Bad packing");
 
+/** Maximum number of virtual texture spaces that can support feedback */
+#define VIRTUALTEXTURE_MAX_FEEDBACK_SPACES 16
+
 /** Maximum number of layers that can be allocated in a single VT page table */
 #define VIRTUALTEXTURE_SPACE_MAXLAYERS 8
 
diff --git a/Engine/Source/Runtime/Renderer/Private/VT/AllocatedVirtualTexture.cpp b/Engine/Source/Runtime/Renderer/Private/VT/AllocatedVirtualTexture.cpp
--- a/Engine/Source/Runtime/Renderer/Private/VT/AllocatedVirtualTexture.cpp
+++ b/Engine/Source/Runtime/Renderer/Private/VT/AllocatedVirtualTexture.cpp
@@ -295,15 +295,11 @@ bool FAllocatedVirtualTexture::TryMapLockedTiles(FVirtualTextureSystem* InSystem
 		}
 	}
 
-	// Display a warning message if we've failed to map pages for this after a set number of frames
+	// Display a message if we've failed to map pages for this after a set number of frames
 	// Generally there should be no delay, but if the system is saturated, it's possible that locked pages may not be loaded immediately
 	if (bHasMissingTiles && InSystem->GetFrame() > FrameAllocated + 30u)
 	{
-		static const auto CVar = IConsoleManager::Get().FindTConsoleVariableDataInt(TEXT("r.VT.Verbose"));
-		if (CVar->GetValueOnRenderThread())
-		{ 
-			UE_LOG(LogVirtualTexturing, Warning, TEXT("Failed to map lowest resolution mip for AllocatedVT %s (%d frames)"), *Description.Name.ToString(), InSystem->GetFrame() - FrameAllocated);
-		}
+		UE_LOG(LogVirtualTexturing, Verbose, TEXT("Failed to map lowest resolution mip for AllocatedVT %s (%d frames)"), *Description.Name.ToString(), InSystem->GetFrame() - FrameAllocated);
 	}
 
 	return !bHasMissingTiles;
@@ -477,6 +473,14 @@ static inline uint32 BitcastFloatToUInt32(float v)
 
 void FAllocatedVirtualTexture::GetPackedPageTableUniform(FUintVector4* OutUniform) const
 {
+	// Force invalidate the page table when the space is no longer in the supported feedback range.
+	// This will cause us to sample the fallback color only.
+	if (UNLIKELY(SpaceID >= VIRTUALTEXTURE_MAX_FEEDBACK_SPACES))
+	{
+		OutUniform[0] = OutUniform[1] = FUintVector4(0, 0, 0, 0);
+		return;
+	}
+
 	const uint32 vPageX = VirtualPageX;
 	const uint32 vPageY = VirtualPageY;
 	const uint32 vPageSize = GetVirtualTileSize();
@@ -499,7 +503,6 @@ void FAllocatedVirtualTexture::GetPackedPageTableUniform(FUintVector4* OutUnifor
 	checkSlow(MipBiasGroup < 4u);
 	checkSlow(MaxLevel < 16u);
 	checkSlow(AdaptiveLevelBias < 16u);
-	checkSlow(SpaceID < 16u);
 
 	OutUniform[0].X = BitcastFloatToUInt32(1.0f / (float)WidthInBlocks);
 	OutUniform[0].Y = BitcastFloatToUInt32(1.0f / (float)HeightInBlocks);
diff --git a/Engine/Source/Runtime/Renderer/Private/VT/AllocatedVirtualTexture.h b/Engine/Source/Runtime/Renderer/Private/VT/AllocatedVirtualTexture.h
--- a/Engine/Source/Runtime/Renderer/Private/VT/AllocatedVirtualTexture.h
+++ b/Engine/Source/Runtime/Renderer/Private/VT/AllocatedVirtualTexture.h
@@ -51,8 +51,6 @@ public:
 	virtual bool TryMapLockedTiles(FVirtualTextureSystem* InSystem) const override;
 	// end IAllocatedVirtualTexture
 
-	inline FVirtualTextureSpace* GetSpace() const { return Space; }
-
 	inline uint32 GetNumUniqueProducers() const { return UniqueProducers.Num(); }
 	inline const FVirtualTextureProducerHandle& GetUniqueProducerHandle(uint32 InProducerIndex) const { return UniqueProducers[InProducerIndex].Handle; }
 	inline uint8 GetUniqueProducerMipBias(uint32 InProducerIndex) const { return UniqueProducers[InProducerIndex].MipBias; }
diff --git a/Engine/Source/Runtime/Renderer/Private/VT/VirtualTextureSystem.cpp b/Engine/Source/Runtime/Renderer/Private/VT/VirtualTextureSystem.cpp
--- a/Engine/Source/Runtime/Renderer/Private/VT/VirtualTextureSystem.cpp
+++ b/Engine/Source/Runtime/Renderer/Private/VT/VirtualTextureSystem.cpp
@@ -74,13 +74,6 @@ DECLARE_MEMORY_STAT_POOL(TEXT("Total Pagetable Memory"), STAT_TotalPagetableMemo
 DECLARE_GPU_STAT(VirtualTexture);
 DECLARE_GPU_STAT(VirtualTextureAllocate);
 
-static TAutoConsoleVariable<int32> CVarVTVerbose(
-	TEXT("r.VT.Verbose"),
-	0,
-	TEXT("Be pedantic about certain things that shouln't occur unless something is wrong. This may cause a lot of logspam 100's of lines per frame."),
-	ECVF_RenderThreadSafe
-);
-
 static TAutoConsoleVariable<int32> CVarVTEnableFeedback(
 	TEXT("r.VT.EnableFeedback"),
 	1,
@@ -407,7 +400,7 @@ FVirtualTextureSystem::~FVirtualTextureSystem()
 	check(AllocatedVTs.Num() == 0);
 	check(PersistentVTMap.Num() == 0);
 
-	for (uint32 SpaceID = 0u; SpaceID < MaxSpaces; ++SpaceID)
+	for (int32 SpaceID = 0; SpaceID < Spaces.Num(); ++SpaceID)
 	{
 		FVirtualTextureSpace* Space = Spaces[SpaceID].Get();
 		if (Space)
@@ -464,9 +457,9 @@ void FVirtualTextureSystem::FlushCache(FVirtualTextureProducerHandle const& Prod
 
 		// If this is an Adaptive VT we need to collect all of the associated Producers to flush.
 		TArray<FAdaptiveVirtualTexture::FProducerInfo> ProducerInfos;
-		if (AdaptiveVTs[SpaceID] != nullptr)
+		if (FAdaptiveVirtualTexture* AdaptiveVT = GetAdaptiveVirtualTexture(SpaceID))
 		{
-			AdaptiveVTs[SpaceID]->GetProducers(TextureRegion, MaxLevelToEvict, ProducerInfos);
+			AdaptiveVT->GetProducers(TextureRegion, MaxLevelToEvict, ProducerInfos);
 		}
 
 		for (int32 i = 0; i < PhysicalSpacesForProducer.Num(); ++i)
@@ -501,7 +494,7 @@ void FVirtualTextureSystem::FlushCache(FVirtualTextureProducerHandle const& Prod
 void FVirtualTextureSystem::DumpFromConsole()
 {
 	bool verbose = false;
-	for (int ID = 0; ID < 16; ID++)
+	for (int ID = 0; ID < Spaces.Num(); ID++)
 	{
 		FVirtualTextureSpace* Space = Spaces[ID].Get();
 		if (Space)
@@ -525,7 +518,7 @@ void FVirtualTextureSystem::ListPhysicalPoolsFromConsole()
 			const uint32 TotalSizeInBytes = PhysicalSpace.GetSizeInBytes();
 			const uint32 NumTiles = PhysicalSpace.GetNumTiles();
 
-			UE_LOG(LogConsoleResponse, Display, TEXT("PhysicalPool: [%i] %s (%ix%i):"), i, *PhysicalSpace.GetFormatString(), Desc.TileSize, Desc.TileSize);
+			UE_LOG(LogVirtualTexturing, Display, TEXT("PhysicalPool: [%i] %s (%ix%i):"), i, *PhysicalSpace.GetFormatString(), Desc.TileSize, Desc.TileSize);
 			
 			const uint32 AllocatedTiles = PagePool.GetNumAllocatedPages();
 			const uint32 AllocatedMemoryInBytes = AllocatedTiles * (TotalSizeInBytes / NumTiles);
@@ -535,12 +528,12 @@ void FVirtualTextureSystem::ListPhysicalPoolsFromConsole()
 			const uint32 LockedMemoryInBytes = LockedTiles * (TotalSizeInBytes / NumTiles);
 			const float LockedMemory = (float)LockedMemoryInBytes / 1024.0f / 1024.0f;
 
-			UE_LOG(LogConsoleResponse, Display, TEXT("  SizeInMegabyte= %f"), (float)TotalSizeInBytes / 1024.0f / 1024.0f);
-			UE_LOG(LogConsoleResponse, Display, TEXT("  Dimensions= %ix%i"), PhysicalSpace.GetTextureSize(), PhysicalSpace.GetTextureSize());
-			UE_LOG(LogConsoleResponse, Display, TEXT("  Tiles= %i"), NumTiles);
-			UE_LOG(LogConsoleResponse, Display, TEXT("  Tiles Allocated= %i (%fMB)"), AllocatedTiles, AllocatedMemory);
-			UE_LOG(LogConsoleResponse, Display, TEXT("  Tiles Locked= %i (%fMB)"), LockedTiles, LockedMemory);
-			UE_LOG(LogConsoleResponse, Display, TEXT("  Tiles Mapped= %i"), PagePool.GetNumMappedPages());
+			UE_LOG(LogVirtualTexturing, Display, TEXT("  SizeInMegabyte= %f"), (float)TotalSizeInBytes / 1024.0f / 1024.0f);
+			UE_LOG(LogVirtualTexturing, Display, TEXT("  Dimensions= %ix%i"), PhysicalSpace.GetTextureSize(), PhysicalSpace.GetTextureSize());
+			UE_LOG(LogVirtualTexturing, Display, TEXT("  Tiles= %i"), NumTiles);
+			UE_LOG(LogVirtualTexturing, Display, TEXT("  Tiles Allocated= %i (%fMB)"), AllocatedTiles, AllocatedMemory);
+			UE_LOG(LogVirtualTexturing, Display, TEXT("  Tiles Locked= %i (%fMB)"), LockedTiles, LockedMemory);
+			UE_LOG(LogVirtualTexturing, Display, TEXT("  Tiles Mapped= %i"), PagePool.GetNumMappedPages());
 
 			TotalPhysicalMemory += TotalSizeInBytes;
 			TotalLockedMemory += LockedMemoryInBytes;
@@ -548,7 +541,7 @@ void FVirtualTextureSystem::ListPhysicalPoolsFromConsole()
 	}
 
 	uint64 TotalPageTableMemory = 0u;
-	for (int ID = 0; ID < 16; ID++)
+	for (int32 ID = 0; ID < Spaces.Num(); ID++)
 	{
 		const FVirtualTextureSpace* Space = Spaces[ID].Get();
 		if (Space == nullptr)
@@ -574,9 +567,9 @@ void FVirtualTextureSystem::ListPhysicalPoolsFromConsole()
 		default: checkNoEntry(); break;
 		}
 
-		UE_LOG(LogConsoleResponse, Display, TEXT("Pool: [%i] %s (%ix%i) x %i:"), ID, FormatName, PhysicalTileSize, PhysicalTileSize, Desc.NumPageTableLayers);
-		UE_LOG(LogConsoleResponse, Display, TEXT("  PageTableSize= %ix%i"), PageTableWidth, PageTableHeight);
-		UE_LOG(LogConsoleResponse, Display, TEXT("  Allocations= %i, %i%% (%fMB)"),
+		UE_LOG(LogVirtualTexturing, Display, TEXT("Pool: [%i] %s (%ix%i) x %i:"), ID, FormatName, PhysicalTileSize, PhysicalTileSize, Desc.NumPageTableLayers);
+		UE_LOG(LogVirtualTexturing, Display, TEXT("  PageTableSize= %ix%i"), PageTableWidth, PageTableHeight);
+		UE_LOG(LogVirtualTexturing, Display, TEXT("  Allocations= %i, %i%% (%fMB)"),
 			Allocator.GetNumAllocations(),
 			(int)(AllocatedRatio * 100.0),
 			(float)(AllocatedRatio * TotalSizeInBytes / 1024.0 / 1024.0));
@@ -584,9 +577,9 @@ void FVirtualTextureSystem::ListPhysicalPoolsFromConsole()
 		TotalPageTableMemory += TotalSizeInBytes;
 	}
 
-	UE_LOG(LogConsoleResponse, Display, TEXT("TotalPageTableMemory: %fMB"), (double)TotalPageTableMemory / 1024.0 / 1024.0);
-	UE_LOG(LogConsoleResponse, Display, TEXT("TotalPhysicalMemory: %fMB"), (double)TotalPhysicalMemory / 1024.0 / 1024.0);
-	UE_LOG(LogConsoleResponse, Display, TEXT("TotalLockedMemory: %fMB"), (double)TotalLockedMemory / 1024.0 / 1024.0);
+	UE_LOG(LogVirtualTexturing, Display, TEXT("TotalPageTableMemory: %fMB"), (double)TotalPageTableMemory / 1024.0 / 1024.0);
+	UE_LOG(LogVirtualTexturing, Display, TEXT("TotalPhysicalMemory: %fMB"), (double)TotalPhysicalMemory / 1024.0 / 1024.0);
+	UE_LOG(LogVirtualTexturing, Display, TEXT("TotalLockedMemory: %fMB"), (double)TotalLockedMemory / 1024.0 / 1024.0);
 }
 
 void FVirtualTextureSystem::DumpPoolUsageFromConsole()
@@ -599,7 +592,7 @@ void FVirtualTextureSystem::DumpPoolUsageFromConsole()
 			const FVTPhysicalSpaceDescription& Desc = PhysicalSpace.GetDescription();
 			const FTexturePagePool& PagePool = PhysicalSpace.GetPagePool();
 
-			UE_LOG(LogConsoleResponse, Display, TEXT("PhysicalPool: [%i] %s (%ix%i):"), i, *PhysicalSpace.GetFormatString(), Desc.TileSize, Desc.TileSize);
+			UE_LOG(LogVirtualTexturing, Display, TEXT("PhysicalPool: [%i] %s (%ix%i):"), i, *PhysicalSpace.GetFormatString(), Desc.TileSize, Desc.TileSize);
 
 			TMap<uint32, FUintVector2> ProducerCountMap;
 			PagePool.CollectProducerCounts(ProducerCountMap);
@@ -626,7 +619,7 @@ void FVirtualTextureSystem::DumpPoolUsageFromConsole()
 				FVirtualTextureProducer* Producer = Producers.FindProducer(FVirtualTextureProducerHandle(PackedProducerHandle));
 				if (Producer != nullptr)
 				{
-					UE_LOG(LogConsoleResponse, Display, TEXT("   %s %d (%d)"), *Producer->GetName().ToString(), Count, LockedCount);
+					UE_LOG(LogVirtualTexturing, Display, TEXT("   %s %d (%d)"), *Producer->GetName().ToString(), Count, LockedCount);
 				}
 			}
 		}
@@ -636,7 +629,7 @@ void FVirtualTextureSystem::DumpPoolUsageFromConsole()
 #if WITH_EDITOR
 void FVirtualTextureSystem::SaveAllocatorImagesFromConsole()
 {
-	for (int ID = 0; ID < 16; ID++)
+	for (int32 ID = 0; ID < Spaces.Num(); ID++)
 	{
 		const FVirtualTextureSpace* Space = Spaces[ID].Get();
 		if (Space)
@@ -724,9 +717,10 @@ IAllocatedVirtualTexture* FVirtualTextureSystem::AllocateVirtualTexture(FRHIComm
 		NumPhysicalGroups = Desc.NumTextureLayers;
 	}
 
-	if (NumAllocatedSpaces == MaxSpaces)
+	if (UNLIKELY(NumAllocatedSpaces == VIRTUALTEXTURE_MAX_FEEDBACK_SPACES))
 	{
-		// We have already allocated the maximum number of spaces. If this allocation needs a new space then it will fail.
+		// We have already allocated the maximum number of fully supported spaces. 
+		// If this allocation needs a new space then it may no longer support feedback.
 		// To give us the best chance of success, flush all pending destroys now. This may free up a space.
 		// This is potentially an expensive operation, which is why we only do it if necessary.
 		DestroyPendingVirtualTextures(true);
@@ -825,6 +819,7 @@ IAdaptiveVirtualTexture* FVirtualTextureSystem::AllocateAdaptiveVirtualTexture(F
 	UE::TScopeLock Lock(Mutex);
 	FAdaptiveVirtualTexture* AdaptiveVT = new FAdaptiveVirtualTexture(AdaptiveVTDesc, AllocatedVTDesc);
 	AdaptiveVT->Init(RHICmdList, this);
+	check(AdaptiveVTs.IsValidIndex(AdaptiveVT->GetSpaceID()));
 	check(AdaptiveVTs[AdaptiveVT->GetSpaceID()] == nullptr);
 	AdaptiveVTs[AdaptiveVT->GetSpaceID()] = AdaptiveVT;
 	return AdaptiveVT;
@@ -834,6 +829,7 @@ void FVirtualTextureSystem::DestroyAdaptiveVirtualTexture(IAdaptiveVirtualTextur
 {
 	check(!bUpdating);
 	UE::TScopeLock Lock(Mutex);
+	check(AdaptiveVTs.IsValidIndex(AdaptiveVT->GetSpaceID()));
 	check(AdaptiveVTs[AdaptiveVT->GetSpaceID()] == AdaptiveVT);
 	AdaptiveVTs[AdaptiveVT->GetSpaceID()] = nullptr;
 	AdaptiveVT->Destroy(this);
@@ -884,17 +880,17 @@ FVirtualTextureSpace* FVirtualTextureSystem::AcquireSpace(FRHICommandListBase& R
 	check(!bUpdating);
 	LLM_SCOPE(ELLMTag::VirtualTextureSystem);
 
-	uint32 NumFailedAllocations = 0u;
+	uint32 NumFailedAllocations = 0;
 
 	// If InDesc requests a private space, don't reuse any existing spaces (unless it is a forced space)
 	if (!InDesc.bPrivateSpace || InForceSpaceID != 0xff)
 	{
-		for (uint32 SpaceIndex = 0u; SpaceIndex < MaxSpaces; ++SpaceIndex)
+		for (int32 SpaceIndex = 0; SpaceIndex < Spaces.Num(); ++SpaceIndex)
 		{
 			if (SpaceIndex == InForceSpaceID || InForceSpaceID == 0xff)
 			{
 				FVirtualTextureSpace* Space = Spaces[SpaceIndex].Get();
-				if (Space && Space->GetDescription() == InDesc)
+				if (Space != nullptr && Space->GetDescription() == InDesc)
 				{
 					const int32 PagetableMemory = Space->GetSizeInBytes();
 					const uint32 vAddress = Space->AllocateVirtualTexture(AllocatedVT);
@@ -915,55 +911,74 @@ FVirtualTextureSpace* FVirtualTextureSystem::AcquireSpace(FRHICommandListBase& R
 			}
 		}
 	}
+	
+	if (InForceSpaceID != 0xff)
+	{
+		// ForceSpaceID is used by adaptive virtual textures. 
+		// If we get here we failed to allocate in an adaptive page table. 
+		// This should be unlikely/impossible since adaptive virtual textures test with a TryAlloc() before really allocating.
+		// If you do hit this somehow, try reducing r.VT.AVT.MaxPageResidency to fix.
+		// We allocate a new space to avoid crashing, but warn the user that things won't work well from this point on.
+		UE_LOG(LogVirtualTexturing, Warning, TEXT("Failed to allocate in requested Virtual Texture Space %d. Expect to see visual corruption."), InForceSpaceID);
+	}
 
-	if (InForceSpaceID == 0xff)
+	// Find an empty space slot.
+	int32 FreeSpaceIndex = INDEX_NONE;
+	for (int32 SpaceIndex = 0; FreeSpaceIndex == INDEX_NONE && SpaceIndex < Spaces.Num(); ++SpaceIndex)
 	{
-		// Try to allocate a new space.
-		int32 FreeSpaceIndex = INDEX_NONE;
-		for (uint32 SpaceIndex = 0u; FreeSpaceIndex == INDEX_NONE && SpaceIndex < MaxSpaces; ++SpaceIndex)
+		if (!Spaces[SpaceIndex].IsValid())
 		{
-			if (!Spaces[SpaceIndex].IsValid())
-			{
-				FreeSpaceIndex = SpaceIndex;
-			}
+			FreeSpaceIndex = SpaceIndex;
 		}
+	}
 
-		// Try to recycle a space. We don't release spaces immediately to avoid unnecessary page table recreation.
-		for (uint32 SpaceIndex = 0u; FreeSpaceIndex == INDEX_NONE && SpaceIndex < MaxSpaces; ++SpaceIndex)
+	// If no empty space was found then try to recycle a space. 
+	// A space may be available because we don't always release spaces immediately so as to avoid unnecessary page table recreation.
+	for (int32 SpaceIndex = 0; FreeSpaceIndex == INDEX_NONE && SpaceIndex < Spaces.Num(); ++SpaceIndex)
+	{
+		FVirtualTextureSpace* Space = Spaces[SpaceIndex].Get();
+		if (Space != nullptr && Space->GetRefCount() == 0)
 		{
-			FVirtualTextureSpace* Space = Spaces[SpaceIndex].Get();
-			if (Space != nullptr && Space->GetRefCount() == 0)
-			{
-				DEC_MEMORY_STAT_BY(STAT_TotalPagetableMemory, Space->GetSizeInBytes());
-				Space->ReleaseResource();
-				NumAllocatedSpaces--;
+			DEC_MEMORY_STAT_BY(STAT_TotalPagetableMemory, Space->GetSizeInBytes());
+			Space->ReleaseResource();
+			NumAllocatedSpaces--;
 
-				FreeSpaceIndex = SpaceIndex;
-			}
+			FreeSpaceIndex = SpaceIndex;
 		}
+	}
 
-		if (FreeSpaceIndex != INDEX_NONE)
-		{
-			const uint32 InitialPageTableSize = InDesc.bPrivateSpace ? InDesc.MaxSpaceSize : FMath::Max(AllocatedVT->GetWidthInTiles(), AllocatedVT->GetHeightInTiles());
-			FVirtualTextureSpace* Space = new FVirtualTextureSpace(this, FreeSpaceIndex, InDesc, InitialPageTableSize);
-			Spaces[FreeSpaceIndex].Reset(Space);
+	// Nothing found so we need to create a new slot.
+	if (FreeSpaceIndex == INDEX_NONE)
+	{
+		// Allocate a new Space slot.
+		FreeSpaceIndex = Spaces.AddDefaulted();
+		// Expect an AdaptiveVT slot for each space.
+		AdaptiveVTs.Add(nullptr);
+		check (Spaces.Num() == AdaptiveVTs.Num());
+	}
 
-			NumAllocatedSpaces++;
-			INC_MEMORY_STAT_BY(STAT_TotalPagetableMemory, Space->GetSizeInBytes());
-			Space->InitResource(RHICmdList);
+	if (FreeSpaceIndex >= VIRTUALTEXTURE_MAX_FEEDBACK_SPACES)
+	{
+		// We don't support feedback on more than 16 spaces, because we only have 4 bits available in the 32bit feedback requests.
+		// We still create a new space to avoid crashing, but warn the user that things won't work well from this point on.
+		// To reduce the number of spaces created, try using r.VT.PageTableMode=0 or modifying RVTs to not use private spaces.
+		UE_LOG(LogVirtualTexturing, Warning, TEXT("Allocating a Virtual Texture space with ID %d. Feedback will not work for this space. A dump of the current pool state follows."), FreeSpaceIndex);
+		ListPhysicalPoolsFromConsole();
+	}
 
-			const uint32 vAddress = Space->AllocateVirtualTexture(AllocatedVT);
-			AllocatedVT->AssignVirtualAddress(vAddress);
+	const uint32 InitialPageTableSize = InDesc.bPrivateSpace ? InDesc.MaxSpaceSize : FMath::Max(AllocatedVT->GetWidthInTiles(), AllocatedVT->GetHeightInTiles());
+	FVirtualTextureSpace* Space = new FVirtualTextureSpace(this, FreeSpaceIndex, InDesc, InitialPageTableSize);
+	Spaces[FreeSpaceIndex].Reset(Space);
 
-			Space->AddRef();
-			return Space;
-		}
-	}
+	NumAllocatedSpaces++;
+	INC_MEMORY_STAT_BY(STAT_TotalPagetableMemory, Space->GetSizeInBytes());
+	Space->InitResource(RHICmdList);
 
-	// Out of space slots
-	checkf(false, TEXT("Failed to acquire space for VT (%d x %d), failed to allocate from %d existing matching spaces"),
-		AllocatedVT->GetWidthInTiles(), AllocatedVT->GetHeightInTiles(), NumFailedAllocations);
-	return nullptr;
+	const uint32 vAddress = Space->AllocateVirtualTexture(AllocatedVT);
+	AllocatedVT->AssignVirtualAddress(vAddress);
+
+	Space->AddRef();
+	return Space;
 }
 
 void FVirtualTextureSystem::ReleaseSpace(FVirtualTextureSpace* Space)
@@ -1157,7 +1172,7 @@ void FVirtualTextureSystem::ReleasePendingSpaces(bool bForceReleaseAll)
 		}
 	}
 
-	for (int32 Id = 0; Id < MaxSpaces; ++Id)
+	for (int32 Id = 0; Id < Spaces.Num(); ++Id)
 	{
 		// Virtual space is released when ref count hits 0 and after some delay to avoid unnecessary recreation.
 		FVirtualTextureSpace* Space = Spaces[Id].Get();
@@ -1297,7 +1312,7 @@ void FVirtualTextureSystem::RequestTiles(IAllocatedVirtualTexture* AllocatedVT,
 	UE::TScopeLock Lock(Mutex);
 	if (InMipLevel >= 0)
 	{
-		if (FAdaptiveVirtualTexture* AdaptiveVT = AdaptiveVTs[AllocatedVT->GetSpaceID()])
+		if (FAdaptiveVirtualTexture* AdaptiveVT = GetAdaptiveVirtualTexture(AllocatedVT->GetSpaceID()))
 		{
 			// Adaptive virtual textures are a special case where the "root" AllocatedVT is one of many AllocatedVTs.
 			// We get the full set of AllocatedVTs before requesting tiles.
@@ -1372,6 +1387,12 @@ void FVirtualTextureSystem::LoadPendingTiles(FRDGBuilder& GraphBuilder, ERHIFeat
 
 void FVirtualTextureSystem::RequestTilesInternal(const IAllocatedVirtualTexture* AllocatedVT, int32 InMipLevel)
 {
+	// Don't generate tile requests for spaces that don't support it.
+	if (AllocatedVT->GetSpaceID() >= VIRTUALTEXTURE_MAX_FEEDBACK_SPACES)
+	{
+		return;
+	}
+
 	const int32 MipWidthInTiles = FMath::Max<int32>(AllocatedVT->GetWidthInTiles() >> InMipLevel, 1);
 	const int32 MipHeightInTiles = FMath::Max<int32>(AllocatedVT->GetHeightInTiles() >> InMipLevel, 1);
 	const uint32 vBaseTileX = AllocatedVT->GetVirtualPageX() >> InMipLevel;
@@ -1441,6 +1462,12 @@ static int32 WrapTilePosition(int32 Position, int32 Size)
 
 void FVirtualTextureSystem::RequestTilesForRegionInternal(const IAllocatedVirtualTexture* AllocatedVT, const FVector2D& InScreenSpaceSize, const FVector2D& InViewportPosition, const FVector2D& InViewportSize, const FVector2D& InUV0, const FVector2D& InUV1, int32 InMipLevel)
 {
+	// Don't generate tile requests for spaces that don't support it.
+	if (AllocatedVT->GetSpaceID() >= VIRTUALTEXTURE_MAX_FEEDBACK_SPACES)
+	{
+		return;
+	}
+
 	// Screen size must be a least a pixel
 	FVector2D ScreenSize = FVector2D::Max(InScreenSpaceSize, FVector2D::One());
 
@@ -1693,7 +1720,7 @@ void FVirtualTextureSystem::GatherRequestsTask(const FGatherRequestsParameters&
 		
 		const uint32 vAddress = FMath::MortonCode2(vPageX) | (FMath::MortonCode2(vPageY) << 1);
 	
-		const FAdaptiveVirtualTexture* RESTRICT AdaptiveVT = AdaptiveVTs[ID];
+		const FAdaptiveVirtualTexture* RESTRICT AdaptiveVT = GetAdaptiveVirtualTexture(ID);
 		if (AdaptiveVT != nullptr && vLevelPlusOne <= 1)
 		{
 			// If vLevelPlusOne == 0 this is a request for higher resolution adaptive VT.
@@ -1766,10 +1793,7 @@ void FVirtualTextureSystem::GatherRequestsTask(const FGatherRequestsParameters&
 		const FAllocatedVirtualTexture* RESTRICT AllocatedVT = Space->GetAllocator().Find(vAddress);
 		if (!AllocatedVT)
 		{
-			if (CVarVTVerbose.GetValueOnAnyThread())
-			{
-				UE_LOG(LogConsoleResponse, Display, TEXT("Space %i, vAddr %i@%i is not allocated to any AllocatedVT but was still requested."), ID, vAddress, vLevel);
-			}
+			UE_LOG(LogVirtualTexturing, Verbose, TEXT("Space %i, vAddr %i@%i is not allocated to any AllocatedVT but was still requested."), ID, vAddress, vLevel);
 			continue;
 		}
 
@@ -2432,10 +2456,7 @@ void FVirtualTextureSystem::SubmitRequests(FRHICommandList& RHICmdList, ERHIFeat
 				//checkf(!bLockTile, TEXT("Tried to lock an invalid VT tile"));
 
 				bTileInvalid = true;
-				if (CVarVTVerbose.GetValueOnAnyThread())
-				{
-					UE_LOG(LogConsoleResponse, Display, TEXT("vAddr %i@%i is not a valid request for AllocatedVT but is still requested."), TileToLoad.Local_vAddress, TileToLoad.Local_vLevel);
-				}
+				UE_LOG(LogVirtualTexturing, Verbose, TEXT("vAddr %i@%i is not a valid request for AllocatedVT but is still requested."), TileToLoad.Local_vAddress, TileToLoad.Local_vLevel);
 			}
 			else if (RequestPageResult.Status == EVTRequestPageStatus::Available)
 			{
@@ -2698,13 +2719,13 @@ void FVirtualTextureSystem::FinalizeRequests(FRDGBuilder& GraphBuilder, ISceneRe
 		RDG_GPU_STAT_SCOPE(GraphBuilder, VirtualTexture);
 		TRACE_CPUPROFILER_EVENT_SCOPE(FVirtualTextureSystem::ApplyUpdates);
 
-		for (uint32 ID = 0; ID < MaxSpaces; ID++)
+		for (int32 ID = 0; ID < Spaces.Num(); ID++)
 		{
 			if (Spaces[ID])
 			{
 				Spaces[ID]->ApplyUpdates(this, GraphBuilder, ExternalAccessQueue);
 
-				if (FAdaptiveVirtualTexture* AdaptiveVT = AdaptiveVTs[ID])
+				if (FAdaptiveVirtualTexture* AdaptiveVT = GetAdaptiveVirtualTexture(ID))
 				{
 					AdaptiveVT->ApplyPageTableUpdates(this, GraphBuilder, ExternalAccessQueue);
 				}
@@ -2726,7 +2747,7 @@ void FVirtualTextureSystem::AllocateResources(FRDGBuilder& GraphBuilder)
 	RDG_GPU_STAT_SCOPE(GraphBuilder, VirtualTextureAllocate);
 	RDG_CSV_STAT_EXCLUSIVE_SCOPE(GraphBuilder, VirtualTextureSystem_Update);
 
-	for (uint32 ID = 0; ID < MaxSpaces; ID++)
+	for (int32 ID = 0; ID < Spaces.Num(); ID++)
 	{
 		if (Spaces[ID])
 		{
@@ -2989,7 +3010,7 @@ TUniquePtr<FVirtualTextureUpdater> FVirtualTextureSystem::BeginUpdate(FRDGBuilde
 	// Update Adaptive VTs. This can trigger allocation/destruction of VTs and must happen before the flush below.
 	{
 		SCOPE_CYCLE_COUNTER(STAT_ProcessRequests_UpdateAdaptiveAllocations);
-		for (uint32 ID = 0; ID < MaxSpaces; ID++)
+		for (int32 ID = 0; ID < AdaptiveVTs.Num(); ID++)
 		{
 			if (AdaptiveVTs[ID])
 			{
@@ -3180,12 +3201,20 @@ bool FVirtualTextureSystem::IsPendingRootPageMap(IAllocatedVirtualTexture* Alloc
 
 void FVirtualTextureSystem::GetOnScreenMessages(FCoreDelegates::FSeverityMessageMap& OutMessages)
 {
+	if (NumAllocatedSpaces > VIRTUALTEXTURE_MAX_FEEDBACK_SPACES)
+	{
+		// To reduce the number of spaces created, try using r.VT.PageTableMode=0 or modifying RVTs to not use private spaces.
+		OutMessages.Add(
+			FCoreDelegates::EOnScreenMessageSeverity::Warning,
+			FText::Format(LOCTEXT("VTSpacesOversubscribed", "Using {0} virtual texture spaces. Feedback is disabled for some virtual textures. Consider using r.VT.PageTableMode=0."), FText::AsNumber(NumAllocatedSpaces)));
+	}
+
 	if (CVarVTResidencyNotify.GetValueOnRenderThread() == 0)
 	{
 		return;
 	}
 
-	for (int32 SpaceIndex = 0u; SpaceIndex < PhysicalSpaces.Num(); ++SpaceIndex)
+	for (int32 SpaceIndex = 0; SpaceIndex < PhysicalSpaces.Num(); ++SpaceIndex)
 	{
 		const FVirtualTexturePhysicalSpace* PhysicalSpace = PhysicalSpaces[SpaceIndex];
 		if (PhysicalSpace != nullptr && Frame <= PhysicalSpace->GetLastFrameOversubscribed() + 60u)
@@ -3210,7 +3239,7 @@ void FVirtualTextureSystem::UpdateCsvStats()
 	}
 
 	float MaxResidencyMipMapBias = 0.f;
-	for (int32 SpaceIndex = 0u; SpaceIndex < PhysicalSpaces.Num(); ++SpaceIndex)
+	for (int32 SpaceIndex = 0; SpaceIndex < PhysicalSpaces.Num(); ++SpaceIndex)
 	{
 		const FVirtualTexturePhysicalSpace* PhysicalSpace = PhysicalSpaces[SpaceIndex];
 		if (PhysicalSpace != nullptr)
@@ -3392,7 +3421,7 @@ void FVirtualTextureSystem::AddRequestedTilesTask(const FAddRequestedTilesParame
 			CurrentPersistentHash = PersistentHash;
 		}
 
-		if (AllocatedVT != nullptr)
+		if (AllocatedVT != nullptr && AllocatedVT->GetSpaceID() < VIRTUALTEXTURE_MAX_FEEDBACK_SPACES)
 		{
 			const uint32 LevelPlusOne = (uint32)(PageRequest >> 28) & 0xf;
 			const uint32 LocalAddress = (uint32)(PageRequest & 0xffffff);
diff --git a/Engine/Source/Runtime/Renderer/Private/VT/VirtualTextureSystem.h b/Engine/Source/Runtime/Renderer/Private/VT/VirtualTextureSystem.h
--- a/Engine/Source/Runtime/Renderer/Private/VT/VirtualTextureSystem.h
+++ b/Engine/Source/Runtime/Renderer/Private/VT/VirtualTextureSystem.h
@@ -190,9 +190,9 @@ private:
 	FVirtualTextureProducer* FindProducer(const FVirtualTextureProducerHandle& Handle);
 	FVirtualTexturePhysicalSpace* AcquirePhysicalSpace(FRHICommandListBase& RHICmdList, const FVTPhysicalSpaceDescription& InDesc);
 
-	FVirtualTextureSpace* GetSpace(uint8 ID) const { check(ID < MaxSpaces); return Spaces[ID].Get(); }
-	FAdaptiveVirtualTexture* GetAdaptiveVirtualTexture(uint8 ID) const { check(ID < MaxSpaces); return AdaptiveVTs[ID]; }
-	FVirtualTexturePhysicalSpace* GetPhysicalSpace(uint16 ID) const { check(PhysicalSpaces[ID]);  return PhysicalSpaces[ID]; }
+	FVirtualTextureSpace* GetSpace(int32 ID) const { return Spaces.IsValidIndex(ID) ? Spaces[ID].Get() : nullptr; }
+	FAdaptiveVirtualTexture* GetAdaptiveVirtualTexture(int32 ID) const { return AdaptiveVTs.IsValidIndex(ID) ? AdaptiveVTs[ID] : nullptr; }
+	FVirtualTexturePhysicalSpace* GetPhysicalSpace(uint32 ID) const { return PhysicalSpaces.IsValidIndex(ID) ? PhysicalSpaces[ID] : nullptr; }
 
 	void LockTile(const FVirtualTextureLocalTile& Tile);
 	void UnlockTile(const FVirtualTextureLocalTile& Tile, const FVirtualTextureProducer* Producer);
@@ -246,9 +246,12 @@ private:
 	mutable UE::FRecursiveMutex Mutex;
 
 	static const uint32 MaxNumTasks = 16;
-	static const uint32 MaxSpaces = 16;
+
+	TArray<TUniquePtr<FVirtualTextureSpace>> Spaces;
 	uint32 NumAllocatedSpaces = 0;
-	TUniquePtr<FVirtualTextureSpace> Spaces[MaxSpaces];
+
+	TArray<FAdaptiveVirtualTexture*> AdaptiveVTs;
+
 	TArray<FVirtualTexturePhysicalSpace*> PhysicalSpaces;
 	FVirtualTextureProducerCollection Producers;
 
@@ -258,8 +261,6 @@ private:
 	TArray<IAllocatedVirtualTexture*> AllocatedVTsToMap;
 	TMap<uint32, IAllocatedVirtualTexture*> PersistentVTMap;
 
-	FAdaptiveVirtualTexture* AdaptiveVTs[MaxSpaces] = { nullptr };
-
 	bool bUpdating = false;
 	bool bFlushCaches;
 	void FlushCachesFromConsole();
