Add all shading model to TSR thin geometry detection shading range with r.TSR.ThinGeometryDetection.Coverage.ShadingRange.

#tests Done on test map.

diff --git a/Engine/Shaders/Private/TemporalSuperResolution/TSRMeasureCoverage.usf b/Engine/Shaders/Private/TemporalSuperResolution/TSRMeasureCoverage.usf
--- a/Engine/Shaders/Private/TemporalSuperResolution/TSRMeasureCoverage.usf
+++ b/Engine/Shaders/Private/TemporalSuperResolution/TSRMeasureCoverage.usf
@@ -10,7 +10,12 @@
 
 #define TILE_ROWS 2
 
-#define CONFIG_HAIR_AS_THIN_GEOMETRY (DIM_HAIR_AS_THIN_GEOMETRY)
+#define THIN_GEOMETRY_SHADING_RANGE_FOLIAGE		 0
+#define THIN_GEOMETRY_SHADING_RANGE_FOLIAGE_HAIR 1
+#define THIN_GEOMETRY_SHADING_RANGE_ALL			 2
+
+#define CONFIG_THIN_GEOMETRY_SHADING_RANGE (DIM_THIN_GEOMETRY_SHADING_RANGE)
+
 
 //------------------------------------------------------- PARAMETERS
 
@@ -19,11 +24,15 @@ RWTexture2DArray<float> ThinGeometryCoverageOutput;
 
 uint GetThinGeometry(const FGBufferData GBufferData)
 {
+#if CONFIG_THIN_GEOMETRY_SHADING_RANGE == THIN_GEOMETRY_SHADING_RANGE_ALL
+	uint ThinGeometryFlag = 2u;
+#else
 	uint ThinGeometryFlag = ((GBufferData.ShadingModelID == SHADINGMODELID_TWOSIDED_FOLIAGE) 
-#if CONFIG_HAIR_AS_THIN_GEOMETRY
+	#if CONFIG_THIN_GEOMETRY_SHADING_RANGE == THIN_GEOMETRY_SHADING_RANGE_FOLIAGE_HAIR
 		    ||(GBufferData.ShadingModelID == SHADINGMODELID_HAIR)
-#endif
+	#endif
 		    ) ? 2u : 1u;
+#endif
 	return select(GBufferData.ShadingModelID == SHADINGMODELID_UNLIT, 0u, ThinGeometryFlag);
 }
 
diff --git a/Engine/Source/Runtime/Renderer/Private/PostProcess/TemporalSuperResolution.cpp b/Engine/Source/Runtime/Renderer/Private/PostProcess/TemporalSuperResolution.cpp
--- a/Engine/Source/Runtime/Renderer/Private/PostProcess/TemporalSuperResolution.cpp
+++ b/Engine/Source/Runtime/Renderer/Private/PostProcess/TemporalSuperResolution.cpp
@@ -321,9 +321,12 @@ TAutoConsoleVariable<bool> CVarTSRThinGeometryCoverageEdgeReprojection(
 	TEXT("Whether thin geometry edge against not thin geometry should be considered for coverage. Better stability for sparse thin geometry clusters."),
 	ECVF_RenderThreadSafe);
 
-TAutoConsoleVariable<bool> CVarTSRThinGeometryCoverageHair(
-	TEXT("r.TSR.ThinGeometryDetection.Coverage.Hair"), true,
-	TEXT("Include hair as thin geometry and accumulate the coverage."),
+TAutoConsoleVariable<int32> CVarTSRThinGeometryCoverageShadingRange(
+	TEXT("r.TSR.ThinGeometryDetection.Coverage.ShadingRange"), 1,
+	TEXT("The shading model range regarded as thin geometry and accumulate the coverage.\n")
+	TEXT("0: Two-sided foliage only.\n")
+	TEXT("1: Two-sided foliage and hair.(Default)\n")
+	TEXT("2: All shading model except unlit.\n"),
 	ECVF_RenderThreadSafe);
 
 TAutoConsoleVariable<float> CVarTSRThinGeometryCoverageMaxRelaxationWeight(
@@ -714,8 +717,36 @@ class FTSRMeasureThinGeometryCoverageCS : public FTSRShader
 	DECLARE_GLOBAL_SHADER(FTSRMeasureThinGeometryCoverageCS);
 	SHADER_USE_PARAMETER_STRUCT(FTSRMeasureThinGeometryCoverageCS, FTSRShader);
 
-	class FHairAsThinGeometryDim : SHADER_PERMUTATION_BOOL("DIM_HAIR_AS_THIN_GEOMETRY");
-	using FPermutationDomain = TShaderPermutationDomain<FHairAsThinGeometryDim>;
+	enum class EThinGeometryShadingRange: uint32
+	{
+		Foliage,
+		FoliageAndHair,
+		All,
+		MAX
+	};
+
+	static EThinGeometryShadingRange GetShadingRange()
+	{
+		int32 ShadingRange = CVarTSRThinGeometryCoverageShadingRange.GetValueOnRenderThread();
+		
+		return static_cast<EThinGeometryShadingRange>(FMath::Clamp(ShadingRange,
+			static_cast<int32>(EThinGeometryShadingRange::Foliage), 
+			static_cast<int32>(EThinGeometryShadingRange::MAX)-1));
+	}
+
+	static const TCHAR* GetShadingRangeName(EThinGeometryShadingRange ShadingRange)
+	{
+		static const TCHAR* const kShadingRangeNames[] = {
+			TEXT("Foliage"),
+			TEXT("Foliage&Hair"),
+			TEXT("All")
+		};
+		static_assert(UE_ARRAY_COUNT(kShadingRangeNames) == int32(EThinGeometryShadingRange::MAX), "Fix me");
+		return kShadingRangeNames[int32(ShadingRange)];
+	}
+
+	class FThinGeometryShadingRangeDim : SHADER_PERMUTATION_ENUM_CLASS("DIM_THIN_GEOMETRY_SHADING_RANGE",EThinGeometryShadingRange);
+	using FPermutationDomain = TShaderPermutationDomain<FThinGeometryShadingRangeDim>;
 
 	BEGIN_SHADER_PARAMETER_STRUCT(FParameters, )
 		SHADER_PARAMETER_STRUCT(FScreenPassTextureViewportParameters, InputInfo)
@@ -1692,9 +1723,9 @@ void AddTSRMeasureThinGeometryCoverage(FRDGBuilder& GraphBuilder, FGlobalShaderM
 	
 	PassParameters->ThinGeometryCoverageOutput = GraphBuilder.CreateUAV(ThinGeometryCoverage.Texture);
 
-	bool HairAsThinGeometry = CVarTSRThinGeometryCoverageHair.GetValueOnRenderThread();
+	FTSRMeasureThinGeometryCoverageCS::EThinGeometryShadingRange ShadingRange = FTSRMeasureThinGeometryCoverageCS::GetShadingRange();
 	FTSRMeasureThinGeometryCoverageCS::FPermutationDomain PermutationVector;
-	PermutationVector.Set<FTSRMeasureThinGeometryCoverageCS::FHairAsThinGeometryDim>(HairAsThinGeometry);
+	PermutationVector.Set<FTSRMeasureThinGeometryCoverageCS::FThinGeometryShadingRangeDim>(ShadingRange);
 
 	// whether TSR passes can run on async compute.
 	int32 AsyncComputePasses = GSupportsEfficientAsyncCompute ? CVarTSRAsyncCompute.GetValueOnRenderThread() : 0;
@@ -1704,7 +1735,7 @@ void AddTSRMeasureThinGeometryCoverage(FRDGBuilder& GraphBuilder, FGlobalShaderM
 		GraphBuilder,
 		RDG_EVENT_NAME("TSR MeasureThinGeometryCoverage(#%d%s) %dx%d", 
 			PermutationVector.ToDimensionValueId(),
-			HairAsThinGeometry ? TEXT(" Hair"): TEXT(""),
+			FTSRMeasureThinGeometryCoverageCS::GetShadingRangeName(ShadingRange),
 			ViewRect.Width(), 
 			ViewRect.Height()),
 		AsyncComputePasses >= 2 ? ERDGPassFlags::AsyncCompute : ERDGPassFlags::Compute,
@@ -3581,7 +3612,7 @@ FDefaultTemporalUpscaler::FOutputs AddTemporalSuperResolutionPasses(
 				Tiles[4 * 1 + 3] = Visualize(EVisualizeId::AntiFlickering, TEXT("Flickering Temporal Analysis"));
 				if (PassConfig.ThinGeometryDetectionEnable)
 				{
-					Tiles[4 * 2 + 3] = Visualize(EVisualizeId::ThinGeometry, TEXT("ThinGeometry"));
+					Tiles[4 * 2 + 3] = Visualize(EVisualizeId::ThinGeometry, TEXT("Thin Geometry Detection"));
 				}
 			}
 			else if (Visualization == EVisualizeId::ReprojectionFieldOverview)
