Trivial fix: change mismatched END_GLOBAL_SHADER_PARAMETER_STRUCT to END_SHADER_PARAMETER_STRUCT

diff --git a/Engine/Source/Runtime/Renderer/Private/CompositionLighting/PostProcessAmbientOcclusion.cpp b/Engine/Source/Runtime/Renderer/Private/CompositionLighting/PostProcessAmbientOcclusion.cpp
--- a/Engine/Source/Runtime/Renderer/Private/CompositionLighting/PostProcessAmbientOcclusion.cpp
+++ b/Engine/Source/Runtime/Renderer/Private/CompositionLighting/PostProcessAmbientOcclusion.cpp
@@ -719,7 +719,7 @@ public:
 		SHADER_PARAMETER_STRUCT_INCLUDE(FAmbientOcclusionParameters, SharedParameters)
 		SHADER_PARAMETER_RDG_UNIFORM_BUFFER(FSubstrateGlobalUniformParameters, Substrate)
 		RENDER_TARGET_BINDING_SLOTS()
-	END_GLOBAL_SHADER_PARAMETER_STRUCT();
+	END_SHADER_PARAMETER_STRUCT();
 };
 IMPLEMENT_GLOBAL_SHADER(FAmbientOcclusionPS, "/Engine/Private/PostProcessAmbientOcclusion.usf", "MainPS", SF_Pixel);
 
@@ -752,7 +752,7 @@ public:
 		SHADER_PARAMETER_STRUCT_INCLUDE(FAmbientOcclusionParameters, SharedParameters)
 		SHADER_PARAMETER_RDG_UNIFORM_BUFFER(FSubstrateGlobalUniformParameters, Substrate)
 		SHADER_PARAMETER_RDG_TEXTURE_UAV(RWTexture2D, OutTexture)
-	END_GLOBAL_SHADER_PARAMETER_STRUCT();
+	END_SHADER_PARAMETER_STRUCT();
 };
 IMPLEMENT_GLOBAL_SHADER(FAmbientOcclusionCS, "/Engine/Private/PostProcessAmbientOcclusion.usf", "MainCS", SF_Compute);
 
@@ -1115,7 +1115,7 @@ public:
 		SHADER_PARAMETER_STRUCT_INCLUDE(FGTAOShaderParameters, GTAOParameters)
 
 		SHADER_PARAMETER_RDG_TEXTURE_UAV(RWTexture2D, OutTexture)
-	END_GLOBAL_SHADER_PARAMETER_STRUCT();
+	END_SHADER_PARAMETER_STRUCT();
 };
 IMPLEMENT_GLOBAL_SHADER(FGTAOHorizonSearchAndIntegrateCS, "/Engine/Private/PostProcessAmbientOcclusion.usf", "GTAOCombinedCS", SF_Compute);
 
@@ -1209,7 +1209,7 @@ public:
 		SHADER_PARAMETER_SAMPLER(SamplerState, HorizonsTextureSampler)
 
 		RENDER_TARGET_BINDING_SLOTS()
-	END_GLOBAL_SHADER_PARAMETER_STRUCT();
+	END_SHADER_PARAMETER_STRUCT();
 };
 IMPLEMENT_GLOBAL_SHADER(FGTAOInnerIntegratePS, "/Engine/Private/PostProcessAmbientOcclusion.usf", "GTAOInnerIntegratePS", SF_Pixel);
 
@@ -1306,7 +1306,7 @@ public:
 
 		SHADER_PARAMETER_RDG_TEXTURE_UAV(RWTexture2D, HorizonOutTexture)
 		SHADER_PARAMETER_RDG_TEXTURE_UAV(RWTexture2D, DepthOutTexture)
-	END_GLOBAL_SHADER_PARAMETER_STRUCT();
+	END_SHADER_PARAMETER_STRUCT();
 };
 IMPLEMENT_GLOBAL_SHADER(FGTAOHorizonSearchCS, "/Engine/Private/PostProcessAmbientOcclusion.usf", "HorizonSearchCS", SF_Compute);
 
@@ -1398,7 +1398,7 @@ public:
 		SHADER_PARAMETER_SAMPLER(SamplerState, SceneVelocityTextureSampler)
 
 		SHADER_PARAMETER_RDG_TEXTURE_UAV(RWTexture2D, OutTexture)
-	END_GLOBAL_SHADER_PARAMETER_STRUCT();
+	END_SHADER_PARAMETER_STRUCT();
 };
 IMPLEMENT_GLOBAL_SHADER(FGTAOTemporalFilterCS, "/Engine/Private/PostProcessAmbientOcclusion.usf", "GTAOTemporalFilterCS", SF_Compute);
 
@@ -1526,7 +1526,7 @@ public:
 		SHADER_PARAMETER_RDG_TEXTURE(Texture2D, GTAOSpatialFilterDepthTexture)
 
 		SHADER_PARAMETER_RDG_TEXTURE_UAV(RWTexture2D, OutTexture)
-	END_GLOBAL_SHADER_PARAMETER_STRUCT();
+	END_SHADER_PARAMETER_STRUCT();
 };
 IMPLEMENT_GLOBAL_SHADER(FGTAOSpatialFilterCS, "/Engine/Private/PostProcessAmbientOcclusion.usf", "GTAOSpatialFilterCS", SF_Compute);
 
@@ -1631,7 +1631,7 @@ public:
 		SHADER_PARAMETER(FVector2f, GTAOUpsamplePixelSize)
 
 		RENDER_TARGET_BINDING_SLOTS()
-	END_GLOBAL_SHADER_PARAMETER_STRUCT();
+	END_SHADER_PARAMETER_STRUCT();
 };
 IMPLEMENT_GLOBAL_SHADER(FGTAOUpsamplePS, "/Engine/Private/PostProcessAmbientOcclusion.usf", "GTAOUpsamplePS", SF_Pixel);
 
