Fix cluster shading compilation issue on Dx11/FXC.

Remove cluster shading shaders from SM5 as cluster shading requires VSM which is only supported on SM6

diff --git a/Engine/Source/Runtime/Renderer/Private/ClusteredDeferredShadingPass.cpp b/Engine/Source/Runtime/Renderer/Private/ClusteredDeferredShadingPass.cpp
--- a/Engine/Source/Runtime/Renderer/Private/ClusteredDeferredShadingPass.cpp
+++ b/Engine/Source/Runtime/Renderer/Private/ClusteredDeferredShadingPass.cpp
@@ -51,12 +51,12 @@ DECLARE_GPU_STAT_NAMED(ClusteredShading, TEXT("Clustered Shading"));
 
 using namespace LightFunctionAtlas;
 
-bool FDeferredShadingSceneRenderer::ShouldUseClusteredDeferredShading() const
+bool FDeferredShadingSceneRenderer::ShouldUseClusteredDeferredShading(EShaderPlatform InPlatform) const
 {
 	// The feature level is the same as in the shader compile conditions below, maybe we don't need SM5?
 	// NOTE: should also take into account the conditions for building the light grid, since these 
 	//       shaders might have another feature level.
-	return CVarClusteredDeferredShadingEnableForProject.GetValueOnAnyThread() > 0 && GUseClusteredDeferredShading != 0 && Scene->GetFeatureLevel() >= ERHIFeatureLevel::SM5;
+	return CVarClusteredDeferredShadingEnableForProject.GetValueOnAnyThread() > 0 && GUseClusteredDeferredShading != 0 && Scene->GetFeatureLevel() >= ERHIFeatureLevel::SM6 && DoesPlatformSupportVirtualShadowMaps(InPlatform);
 }
 
 bool FDeferredShadingSceneRenderer::AreLightsInLightGrid() const
@@ -137,7 +137,13 @@ class FClusteredShadingPS : public FGlobalShader
 			return false;
 		}
 
-		return IsFeatureLevelSupported(Parameters.Platform, ERHIFeatureLevel::SM5) && CVarClusteredDeferredShadingEnableForProject.GetValueOnAnyThread() > 0;
+		// Cluster shading requires VSM support
+		if (!DoesPlatformSupportVirtualShadowMaps(Parameters.Platform))
+		{
+			return false;
+		}
+
+		return IsFeatureLevelSupported(Parameters.Platform, ERHIFeatureLevel::SM6) && CVarClusteredDeferredShadingEnableForProject.GetValueOnAnyThread() > 0;
 	}
 
 	static void ModifyCompilationEnvironment(const FGlobalShaderPermutationParameters& Parameters, FShaderCompilerEnvironment& OutEnvironment)
@@ -315,7 +321,7 @@ void FDeferredShadingSceneRenderer::AddClusteredDeferredShadingPass(
 	FRDGTextureRef HairStrandsShadowMaskBits,
 	FRDGTextureRef LightingChannelsTexture)
 {
-	check(ShouldUseClusteredDeferredShading());
+	check(ShouldUseClusteredDeferredShading(ViewFamily.GetShaderPlatform()));
 
 	const int32 NumLightsToRender = SortedLightsSet.ClusteredSupportedEnd;
 
diff --git a/Engine/Source/Runtime/Renderer/Private/DeferredShadingRenderer.cpp b/Engine/Source/Runtime/Renderer/Private/DeferredShadingRenderer.cpp
--- a/Engine/Source/Runtime/Renderer/Private/DeferredShadingRenderer.cpp
+++ b/Engine/Source/Runtime/Renderer/Private/DeferredShadingRenderer.cpp
@@ -2186,7 +2186,7 @@ void FDeferredShadingSceneRenderer::Render(FRDGBuilder& GraphBuilder, const FSce
 			extern UE::Tasks::FTask GetGatherAndSortLightsPrerequisiteTask(const FDynamicShadowsTaskData* TaskData);
 
 			auto* SortedLightSet = GraphBuilder.AllocObject<FSortedLightSetSceneInfo>();
-			const bool bShadowedLightsInClustered = ShouldUseClusteredDeferredShading()
+			const bool bShadowedLightsInClustered = ShouldUseClusteredDeferredShading(ViewFamily.GetShaderPlatform())
 				&& IsVSMOnePassProjectionEnabled(ViewFamily.EngineShowFlags)
 				&& VirtualShadowMapArray.IsEnabled();
 
diff --git a/Engine/Source/Runtime/Renderer/Private/DeferredShadingRenderer.h b/Engine/Source/Runtime/Renderer/Private/DeferredShadingRenderer.h
--- a/Engine/Source/Runtime/Renderer/Private/DeferredShadingRenderer.h
+++ b/Engine/Source/Runtime/Renderer/Private/DeferredShadingRenderer.h
@@ -825,7 +825,7 @@ private:
 	/** 
 	 * True if the 'r.UseClusteredDeferredShading' flag is 1 and sufficient feature level. 
 	 */
-	bool ShouldUseClusteredDeferredShading() const;
+	bool ShouldUseClusteredDeferredShading(EShaderPlatform InPlatform) const;
 
 	/**
 	 * Have the lights been injected into the light grid?
diff --git a/Engine/Source/Runtime/Renderer/Private/LightRendering.cpp b/Engine/Source/Runtime/Renderer/Private/LightRendering.cpp
--- a/Engine/Source/Runtime/Renderer/Private/LightRendering.cpp
+++ b/Engine/Source/Runtime/Renderer/Private/LightRendering.cpp
@@ -1575,11 +1575,11 @@ void FDeferredShadingSceneRenderer::RenderLights(
 
 			bool bRenderSimpleLightsStandardDeferred = SortedLightSet.SimpleLights.InstanceData.Num() > 0;
 
-			UE_CLOG(ShouldUseClusteredDeferredShading() && !AreLightsInLightGrid(), LogRenderer, Warning,
+			UE_CLOG(ShouldUseClusteredDeferredShading(ViewFamily.GetShaderPlatform()) && !AreLightsInLightGrid(), LogRenderer, Warning,
 				TEXT("Clustered deferred shading is enabled, but lights were not injected in grid, falling back to other methods (hint 'r.LightCulling.Quality' may cause this)."));
 
 			// True if the clustered shading is enabled and the feature level is there, and that the light grid had lights injected.
-			if (ShouldUseClusteredDeferredShading() && AreLightsInLightGrid())
+			if (ShouldUseClusteredDeferredShading(ViewFamily.GetShaderPlatform()) && AreLightsInLightGrid())
 			{
 				// Tell the trad. deferred that the clustered deferred capable lights are taken care of.
 				// This includes the simple lights
