Use RDG to create and fill virtual texture page table update buffers.
Old behavior used Lock/Unlock which adds an RHI flush.
#rnx

diff --git a/Engine/Source/Runtime/Renderer/Private/VT/VirtualTextureSpace.cpp b/Engine/Source/Runtime/Renderer/Private/VT/VirtualTextureSpace.cpp
--- a/Engine/Source/Runtime/Renderer/Private/VT/VirtualTextureSpace.cpp
+++ b/Engine/Source/Runtime/Renderer/Private/VT/VirtualTextureSpace.cpp
@@ -139,9 +139,6 @@ void FVirtualTextureSpace::ReleaseRHI()
 
 	PageTableIndirection.TextureReferenceRHI.SafeRelease();
 	GRenderTargetPool.FreeUnusedResource(PageTableIndirection.RenderTarget);
-
-	UpdateBuffer.SafeRelease();
-	UpdateBufferSRV.SafeRelease();
 }
 
 FUintPoint FVirtualTextureSpace::GetRequiredPageTableAllocationSize() const
@@ -188,67 +185,21 @@ TGlobalResource< FSpriteIndexBuffer<8> > GQuadIndexBuffer;
 
 class FPageTableUpdateVS : public FGlobalShader
 {
-	DECLARE_INLINE_TYPE_LAYOUT(FPageTableUpdateVS, NonVirtual);
-protected:
-	FPageTableUpdateVS() {}
-	
 public:
-	FPageTableUpdateVS(const ShaderMetaType::CompiledShaderInitializerType& Initializer)
-		: FGlobalShader(Initializer)
-	{
-		PageTableSize.Bind( Initializer.ParameterMap, TEXT("PageTableSize") );
-		FirstUpdate.Bind( Initializer.ParameterMap, TEXT("FirstUpdate") );
-		NumUpdates.Bind( Initializer.ParameterMap, TEXT("NumUpdates") );
-		UpdateBuffer.Bind( Initializer.ParameterMap, TEXT("UpdateBuffer") );
-	}
-
-	static bool ShouldCompilePermutation(const FGlobalShaderPermutationParameters& Parameters)
-	{
-		return true;
-	}
-
-	void SetParameters(FRHIBatchedShaderParameters& BatchedParameters, FIntPoint InPageTableSize, uint32 InFirstUpdate, uint32 InNumUpdates, FRHIShaderResourceView* InUpdateBuffer)
-	{
-		SetShaderValue(BatchedParameters, PageTableSize, InPageTableSize);
-		SetShaderValue(BatchedParameters, FirstUpdate, InFirstUpdate);
-		SetShaderValue(BatchedParameters, NumUpdates, InNumUpdates);
-		SetSRVParameter(BatchedParameters, UpdateBuffer, InUpdateBuffer);
-	}
-
-	LAYOUT_FIELD(FShaderParameter, PageTableSize);
-	LAYOUT_FIELD(FShaderParameter, FirstUpdate);
-	LAYOUT_FIELD(FShaderParameter, NumUpdates);
-	LAYOUT_FIELD(FShaderResourceParameter, UpdateBuffer);
-};
-
-class FPageTableUpdatePS : public FGlobalShader
-{
-	DECLARE_INLINE_TYPE_LAYOUT(FPageTableUpdatePS, NonVirtual);
-protected:
-	FPageTableUpdatePS() {}
-
-public:
-	FPageTableUpdatePS(const ShaderMetaType::CompiledShaderInitializerType& Initializer)
-		: FGlobalShader(Initializer)
-	{}
-	
-	static bool ShouldCompilePermutation(const FGlobalShaderPermutationParameters& Parameters)
-	{
-		return true;
-	}
+	SHADER_USE_PARAMETER_STRUCT(FPageTableUpdateVS, FGlobalShader);
+
+	BEGIN_SHADER_PARAMETER_STRUCT(FParameters, )
+		SHADER_PARAMETER_RDG_BUFFER_SRV(Buffer<uint4>, UpdateBuffer)
+		SHADER_PARAMETER(FUintVector2, PageTableSize)
+		SHADER_PARAMETER(uint32, FirstUpdate)
+		SHADER_PARAMETER(uint32, NumUpdates)
+	END_SHADER_PARAMETER_STRUCT()
 };
 
 template<bool Use16Bits>
 class TPageTableUpdateVS : public FPageTableUpdateVS
 {
-	DECLARE_SHADER_TYPE(TPageTableUpdateVS,Global);
-
-	TPageTableUpdateVS() {}
-
-public:
-	TPageTableUpdateVS(const ShaderMetaType::CompiledShaderInitializerType& Initializer)
-		: FPageTableUpdateVS(Initializer)
-	{}
+	DECLARE_GLOBAL_SHADER(TPageTableUpdateVS);
 
 	static void ModifyCompilationEnvironment(const FGlobalShaderPermutationParameters& Parameters, FShaderCompilerEnvironment& OutEnvironment)
 	{
@@ -258,17 +209,10 @@ public:
 };
 
 template<EPixelFormat TargetFormat>
-class TPageTableUpdatePS : public FPageTableUpdatePS
+class TPageTableUpdatePS : public FGlobalShader
 {
-	DECLARE_SHADER_TYPE(TPageTableUpdatePS, Global);
-
-	TPageTableUpdatePS() {}
+	DECLARE_GLOBAL_SHADER(TPageTableUpdatePS);
 
-public:
-	TPageTableUpdatePS(const ShaderMetaType::CompiledShaderInitializerType& Initializer)
-		: FPageTableUpdatePS(Initializer)
-	{}
-	
 	static void ModifyCompilationEnvironment(const FGlobalShaderPermutationParameters& Parameters, FShaderCompilerEnvironment& OutEnvironment)
 	{
 		FGlobalShader::ModifyCompilationEnvironment( Parameters, OutEnvironment );
@@ -285,6 +229,11 @@ IMPLEMENT_SHADER_TYPE(template<>, TPageTableUpdatePS<PF_R32_UINT>, TEXT("/Engine
 IMPLEMENT_SHADER_TYPE(template<>, TPageTableUpdatePS<PF_R32G32_UINT>, TEXT("/Engine/Private/PageTableUpdate.usf"), TEXT("PageTableUpdatePS_2"), SF_Pixel);
 IMPLEMENT_SHADER_TYPE(template<>, TPageTableUpdatePS<PF_R32G32B32A32_UINT>, TEXT("/Engine/Private/PageTableUpdate.usf"), TEXT("PageTableUpdatePS_4"), SF_Pixel);
 
+BEGIN_SHADER_PARAMETER_STRUCT(FPageTableUpdateParameters, )
+	SHADER_PARAMETER_STRUCT_INCLUDE(FPageTableUpdateVS::FParameters, VS)
+	RENDER_TARGET_BINDING_SLOTS()
+END_SHADER_PARAMETER_STRUCT()
+
 void FVirtualTextureSpace::QueueUpdateEntirePageTable()
 {
 	bForceEntireUpdate = true;
@@ -427,50 +376,27 @@ void FVirtualTextureSpace::ApplyUpdates(FVirtualTextureSystem* System, FRDGBuild
 		return;
 	}
 
-	FRHICommandListBase& RHICmdList = GraphBuilder.RHICmdList;
-
-	if (UpdateBuffer == nullptr || TotalNumUpdates * sizeof(FPageTableUpdate) > UpdateBuffer->GetSize())
-	{
-		// Resize Update Buffer
-		const uint32 MaxUpdates = FMath::RoundUpToPowerOfTwo(TotalNumUpdates);
-		uint32 NewBufferSize = MaxUpdates * sizeof(FPageTableUpdate);
-		if (UpdateBuffer)
-		{
-			NewBufferSize = FMath::Max(NewBufferSize, UpdateBuffer->GetSize() * 2u);
-		}
-
-		const FRHIBufferCreateDesc CreateDesc =
-			FRHIBufferCreateDesc::CreateVertex(TEXT("VirtualTexture_PageTableUpdateBuffer"), NewBufferSize)
-			.AddUsage(EBufferUsageFlags::ShaderResource | EBufferUsageFlags::Volatile)
-			.DetermineInitialState();
-
-		UpdateBuffer = RHICmdList.CreateBuffer(CreateDesc);
-		UpdateBufferSRV = RHICmdList.CreateShaderResourceView(
-			UpdateBuffer, 
-			FRHIViewDesc::CreateBufferSRV()
-				.SetType(FRHIViewDesc::EBufferType::Typed)
-				.SetFormat(PF_R16G16B16A16_UINT));
-	}
-
-	// This flushes the RHI thread!
+	void* StagingMemory = GraphBuilder.Alloc(TotalNumUpdates * sizeof(FPageTableUpdate));
+	uint8* WritePtr = (uint8*)StagingMemory;
+	for (uint32 LayerIndex = 0u; LayerIndex < Description.NumPageTableLayers; ++LayerIndex)
 	{
-		uint8* Buffer = (uint8*)RHICmdList.LockBuffer(UpdateBuffer, 0, TotalNumUpdates * sizeof(FPageTableUpdate), RLM_WriteOnly);
-		for (uint32 LayerIndex = 0u; LayerIndex < Description.NumPageTableLayers; ++LayerIndex)
+		for (uint32 Mip = 0; Mip < CachedNumPageTableLevels; Mip++)
 		{
-			for (uint32 Mip = 0; Mip < CachedNumPageTableLevels; Mip++)
+			const uint32 NumUpdates = ExpandedUpdates[LayerIndex][Mip].Num();
+			if (NumUpdates)
 			{
-				const uint32 NumUpdates = ExpandedUpdates[LayerIndex][Mip].Num();
-				if (NumUpdates)
-				{
-					size_t UploadSize = NumUpdates * sizeof(FPageTableUpdate);
-					FMemory::Memcpy(Buffer, ExpandedUpdates[LayerIndex][Mip].GetData(), UploadSize);
-					Buffer += UploadSize;
-				}
+				size_t UploadSize = NumUpdates * sizeof(FPageTableUpdate);
+				FMemory::Memcpy(WritePtr, ExpandedUpdates[LayerIndex][Mip].GetData(), UploadSize);
+				WritePtr += UploadSize;
 			}
 		}
-		RHICmdList.UnlockBuffer(UpdateBuffer);
 	}
 
+	FRDGBufferRef UpdateBuffer = GraphBuilder.CreateBuffer(FRDGBufferDesc::CreateBufferDesc(sizeof(FPageTableUpdate), TotalNumUpdates), TEXT("VirtualTexture_PageTableUpdateBuffer"));
+	GraphBuilder.QueueBufferUpload(UpdateBuffer, StagingMemory, TotalNumUpdates * sizeof(FPageTableUpdate), ERDGInitialDataFlags::NoCopy);
+	
+	FRDGBufferSRVRef UpdateBufferSRV = GraphBuilder.CreateSRV(FRDGBufferSRVDesc(UpdateBuffer, PF_R16G16B16A16_UINT));
+
 	// Draw
 	RDG_EVENT_SCOPE(GraphBuilder, "PageTableUpdate");
 
@@ -509,7 +435,7 @@ void FVirtualTextureSpace::ApplyUpdates(FVirtualTextureSystem* System, FRDGBuild
 		default: check(false); break;
 		}
 
-		TShaderRef<FPageTableUpdatePS> PixelShader;
+		TShaderRef<FGlobalShader> PixelShader;
 		switch (TexturePixelFormat[TextureIndex])
 		{
 		case PF_R16_UINT: PixelShader = ShaderMap->GetShader< TPageTableUpdatePS<PF_R16_UINT> >(); break;
@@ -529,14 +455,18 @@ void FVirtualTextureSpace::ApplyUpdates(FVirtualTextureSystem* System, FRDGBuild
 			const uint32 NumUpdates = ExpandedUpdates[LayerIndex][Mip].Num();
 			if (NumUpdates)
 			{
-				auto* PassParameters = GraphBuilder.AllocParameters<FRenderTargetParameters>();
+				auto* PassParameters = GraphBuilder.AllocParameters<FPageTableUpdateParameters>();
+				PassParameters->VS.UpdateBuffer = UpdateBufferSRV;
+				PassParameters->VS.PageTableSize = FUintVector2(CachedPageTableWidth, CachedPageTableHeight);
+				PassParameters->VS.FirstUpdate = FirstUpdate;
+				PassParameters->VS.NumUpdates = NumUpdates;
 				PassParameters->RenderTargets[0] = FRenderTargetBinding(PageTableTexture, ERenderTargetLoadAction::ELoad, Mip);
 
 				GraphBuilder.AddPass(
 					RDG_EVENT_NAME("PageTableUpdate (Mip: %d)", Mip),
 					PassParameters,
 					ERDGPassFlags::Raster,
-					[this, VertexShader, PixelShader, BlendStateRHI, FirstUpdate, NumUpdates, MipWidth, MipHeight](FRDGAsyncTask, FRHICommandList& RHICmdList)
+					[PassParameters, VertexShader, PixelShader, BlendStateRHI, NumUpdates, MipWidth, MipHeight](FRDGAsyncTask, FRHICommandList& RHICmdList)
 				{
 					RHICmdList.SetViewport(0, 0, 0.0f, MipWidth, MipHeight, 1.0f);
 
@@ -554,7 +484,7 @@ void FVirtualTextureSpace::ApplyUpdates(FVirtualTextureSystem* System, FRDGBuild
 
 					SetGraphicsPipelineState(RHICmdList, GraphicsPSOInit, 0);
 
-					SetShaderParametersLegacyVS(RHICmdList, VertexShader, FIntPoint(CachedPageTableWidth, CachedPageTableHeight), FirstUpdate, NumUpdates, UpdateBufferSRV);
+					SetShaderParameters(RHICmdList, VertexShader, VertexShader.GetVertexShader(), PassParameters->VS);
 
 					// needs to be the same on shader side (faster on NVIDIA and AMD)
 					uint32 QuadsPerInstance = 8;
diff --git a/Engine/Source/Runtime/Renderer/Private/VT/VirtualTextureSpace.h b/Engine/Source/Runtime/Renderer/Private/VT/VirtualTextureSpace.h
--- a/Engine/Source/Runtime/Renderer/Private/VT/VirtualTextureSpace.h
+++ b/Engine/Source/Runtime/Renderer/Private/VT/VirtualTextureSpace.h
@@ -136,9 +136,6 @@ private:
 
 	TArray<FPageTableUpdate> PageTableUpdates[VIRTUALTEXTURE_SPACE_MAXLAYERS];
 
-	FBufferRHIRef UpdateBuffer;
-	FShaderResourceViewRHIRef UpdateBufferSRV;
-
 	uint32 NumRefs = 0;
 	uint32 ReleasedFrame = 0;
 
