Mobile renderer: Add compute shader SSAO support (both synchronous and async)
- Disable binding the SSAO texture in the basepass if we're rendering deferred, since it causes an unnecessary barrier, which is inefficient in the async compute case

diff --git a/Engine/Shaders/Private/PostProcessAmbientOcclusionMobile.usf b/Engine/Shaders/Private/PostProcessAmbientOcclusionMobile.usf
--- a/Engine/Shaders/Private/PostProcessAmbientOcclusionMobile.usf
+++ b/Engine/Shaders/Private/PostProcessAmbientOcclusionMobile.usf
@@ -982,6 +982,9 @@ void MainPSandCS(in float4 UVAndScreenPos, float4 SvPosition, out float4 OutColo
 
 	float2 WeightAccumulator = 0.0001f;
 	
+	// if the feature is enabled and right side of screen
+	const bool bDebugLookups = DEBUG_LOOKUPS && ViewSpacePosition.x > 0;
+
 #if AO_SAMPLE_QUALITY != 0
 	// no SSAO in this pass, only upsampling
 
@@ -999,9 +1002,6 @@ void MainPSandCS(in float4 UVAndScreenPos, float4 SvPosition, out float4 OutColo
 	// jitter each frame a bit to get higher quality over multiple frames (only if TemporalAA is enabled), can cause ghosting effects
 	const float2 TemporalOffset = ScreenSpaceAOParams[3].xy;
 
-	// if the feature is enabled and right side of screen
-	const bool bDebugLookups = DEBUG_LOOKUPS && ViewSpacePosition.x > 0;
-
 	float2 RandomVec = (Texture2DSample(RandomNormalTexture, RandomNormalTextureSampler, TemporalOffset + UV * ViewportUVToRandomUV).rg * 2 - 1) * ActualAORadius;
 	{
 #endif // AO_SAMPLE_QUALITY == 
@@ -1138,7 +1138,7 @@ void MainPSandCS(in float4 UVAndScreenPos, float4 SvPosition, out float4 OutColo
 		OutColor.r = pow(1 - OutColor.r, 16);	// constnt is tweaked with radius and sample count
 	}
 
-#if QUAD_MESSAGE_PASSING_BLUR > 0 
+#if !COMPUTE_SHADER && QUAD_MESSAGE_PASSING_BLUR > 0 
 	{	// .x: AO output, .y:SceneDepth .zw:view space normal
 		float4 CenterPixel = float4(OutColor.r, SceneDepth, normalize(ViewSpaceNormal).xy); 
 
@@ -1194,6 +1194,35 @@ void MainPS(in noperspective float4 UVAndScreenPos : TEXCOORD0, float4 SvPositio
 	MainPSandCS(UVAndScreenPos, SvPosition, OutColor);	
 }
 
+#if COMPUTE_SHADER
+/** Output target. In compute, this is a single value buffer. */
+RWTexture2D<float> OutTexture;
+[numthreads(THREADGROUP_SIZEX, THREADGROUP_SIZEY, 1)]
+void MainCS(
+	uint2 GroupId : SV_GroupID,
+	uint2 DispatchThreadId : SV_DispatchThreadID,
+	uint2 GroupThreadId : SV_GroupThreadID) 
+{
+	float ScaleFactor = ScreenSpaceAOParams[2].x;
+	int2 PixelPos = DispatchThreadId + AOViewport_ViewportMin; 
+	float2 PixelCenter = (float2)PixelPos + float2(0.5, 0.5);
+	float4 SvPosition = float4(PixelCenter, 0, 0) * ScaleFactor;	
+	float2 BufferUV = SvPositionToBufferUV(SvPosition);
+	SvPosition.zw = float2(LookupDeviceZ(BufferUV),1);
+	float4 OutColor = 1;
+	// Test for early exit with out of depth bound.
+	float SceneDepth = ConvertFromDeviceZ(SvPosition.z);
+	float FadeMul = ScreenSpaceAOParams[4].x;
+	float FadeAdd = ScreenSpaceAOParams[4].y;
+	BRANCH
+	if (SceneDepth * FadeMul + FadeAdd < 1)
+	{
+		MainPSandCS(float4(BufferUV, SvPositionToScreenPosition(SvPosition).xy), SvPosition, OutColor);
+	}
+	OutTexture[PixelPos] = OutColor.r;
+}
+#endif // COMPUTE_SHADER
+
 #endif // SSAO
 
 
diff --git a/Engine/Source/Runtime/Renderer/Private/MobileBasePassRendering.cpp b/Engine/Source/Runtime/Renderer/Private/MobileBasePassRendering.cpp
--- a/Engine/Source/Runtime/Renderer/Private/MobileBasePassRendering.cpp
+++ b/Engine/Source/Runtime/Renderer/Private/MobileBasePassRendering.cpp
@@ -267,7 +267,8 @@ void SetupMobileBasePassUniformParameters(
 	EMobileSceneTextureSetupMode SetupMode,
 	FMobileBasePassUniformParameters& BasePassParameters,
 	const FMobileBasePassTextures& MobileBasePassTextures,
-	bool bForRealtimeSkyCapture)
+	bool bForRealtimeSkyCapture,
+	bool bIsRenderingDeferredBasepass)
 {
 	const FViewInfo* InstancedView = View.GetInstancedView();
 
@@ -311,7 +312,8 @@ void SetupMobileBasePassUniformParameters(
 	const bool bIsMobileMultiViewEnabled = View.bIsMobileMultiViewEnabled || View.Aspects.IsMobileMultiViewEnabled();
 	BasePassParameters.DBuffer = GetDBufferParameters(GraphBuilder, MobileBasePassTextures.DBufferTextures, View.GetShaderPlatform(), bIsMobileMultiViewEnabled);
 
-	if (HasBeenProduced(MobileBasePassTextures.AmbientOcclusionTexture))
+	// Don't bind the SSAO texture if we're rendering the deferred basepass. It's not needed and will result in additional barriers
+	if (!bIsRenderingDeferredBasepass && HasBeenProduced(MobileBasePassTextures.AmbientOcclusionTexture))
 	{
 		BasePassParameters.AmbientOcclusionTexture = MobileBasePassTextures.AmbientOcclusionTexture;
 	}
@@ -321,8 +323,8 @@ void SetupMobileBasePassUniformParameters(
 	}
 	BasePassParameters.AmbientOcclusionSampler = TStaticSamplerState<SF_Point, AM_Clamp, AM_Clamp, AM_Clamp>::GetRHI();
 	BasePassParameters.AmbientOcclusionStaticFraction = FMath::Clamp(View.FinalPostProcessSettings.AmbientOcclusionStaticFraction, 0.0f, 1.0f);
-
-	if (HasBeenProduced(MobileBasePassTextures.ScreenSpaceShadowMaskTexture))
+	
+	if (!bIsRenderingDeferredBasepass && HasBeenProduced(MobileBasePassTextures.ScreenSpaceShadowMaskTexture))
 	{
 		BasePassParameters.ScreenSpaceShadowMaskTexture = MobileBasePassTextures.ScreenSpaceShadowMaskTexture;
 		BasePassParameters.ScreenSpaceShadowMaskTextureArray = MobileBasePassTextures.ScreenSpaceShadowMaskTexture;
@@ -408,10 +410,11 @@ TRDGUniformBufferRef<FMobileBasePassUniformParameters> CreateMobileBasePassUnifo
 	EMobileBasePass BasePass,
 	EMobileSceneTextureSetupMode SetupMode,
 	const FMobileBasePassTextures& MobileBasePassTextures,
-	bool bForRealtimeSkyCapture)
+	bool bForRealtimeSkyCapture,
+	bool bIsRenderingDeferredBasepass)
 {
 	FMobileBasePassUniformParameters* BasePassParameters = GraphBuilder.AllocParameters<FMobileBasePassUniformParameters>();
-	SetupMobileBasePassUniformParameters(GraphBuilder, View, BasePass, SetupMode, *BasePassParameters, MobileBasePassTextures, bForRealtimeSkyCapture);
+	SetupMobileBasePassUniformParameters(GraphBuilder, View, BasePass, SetupMode, *BasePassParameters, MobileBasePassTextures, bForRealtimeSkyCapture, bIsRenderingDeferredBasepass);
 
 	return GraphBuilder.CreateUniformBuffer(BasePassParameters);
 }
@@ -642,15 +645,15 @@ void FMobileSceneRenderer::RenderMobileDebugPrimitives(FRHICommandList& RHICmdLi
 
 	if (View.DebugSimpleElementCollector.HasAnyPrimitives())
 	{
-		FMeshPassProcessorRenderState DrawRenderState;
-		DrawRenderState.SetDepthStencilAccess(FExclusiveDepthStencil::DepthWrite_StencilWrite);
-		DrawRenderState.SetBlendState(TStaticBlendState<CW_RGBA, BO_Add, BF_One, BF_One, BO_Add, BF_One, BF_One>::GetRHI());
-
-		DrawRenderState.SetDepthStencilState(TStaticDepthStencilState<true, CF_DepthNearOrEqual>::GetRHI());
-		View.DebugSimpleElementCollector.DrawBatchedElements(RHICmdList, DrawRenderState, View, EBlendModeFilter::OpaqueAndMasked, SDPG_World);
+	FMeshPassProcessorRenderState DrawRenderState;
+	DrawRenderState.SetDepthStencilAccess(FExclusiveDepthStencil::DepthWrite_StencilWrite);
+	DrawRenderState.SetBlendState(TStaticBlendState<CW_RGBA, BO_Add, BF_One, BF_One, BO_Add, BF_One, BF_One>::GetRHI());
+		
+	DrawRenderState.SetDepthStencilState(TStaticDepthStencilState<true, CF_DepthNearOrEqual>::GetRHI());
+	View.DebugSimpleElementCollector.DrawBatchedElements(RHICmdList, DrawRenderState, View, EBlendModeFilter::OpaqueAndMasked, SDPG_World);
 
-		DrawRenderState.SetDepthStencilState(TStaticDepthStencilState<true, CF_Always>::GetRHI());
-		View.DebugSimpleElementCollector.DrawBatchedElements(RHICmdList, DrawRenderState, View, EBlendModeFilter::OpaqueAndMasked, SDPG_Foreground);
-	}
+	DrawRenderState.SetDepthStencilState(TStaticDepthStencilState<true, CF_Always>::GetRHI());
+	View.DebugSimpleElementCollector.DrawBatchedElements(RHICmdList, DrawRenderState, View, EBlendModeFilter::OpaqueAndMasked, SDPG_Foreground);
+}
 }
 #endif
\ No newline at end of file
diff --git a/Engine/Source/Runtime/Renderer/Private/MobileBasePassRendering.h b/Engine/Source/Runtime/Renderer/Private/MobileBasePassRendering.h
--- a/Engine/Source/Runtime/Renderer/Private/MobileBasePassRendering.h
+++ b/Engine/Source/Runtime/Renderer/Private/MobileBasePassRendering.h
@@ -88,7 +88,8 @@ extern void SetupMobileBasePassUniformParameters(
 	EMobileSceneTextureSetupMode SetupMode,
 	FMobileBasePassUniformParameters& BasePassParameters, 
 	const FMobileBasePassTextures& MobileBasePassTextures = {},
-	bool bForRealtimeSkyCapture = false);
+	bool bForRealtimeSkyCapture = false,
+	bool bIsRenderingDeferredBasepass = false);
 
 extern TRDGUniformBufferRef<FMobileBasePassUniformParameters> CreateMobileBasePassUniformBuffer(
 	FRDGBuilder& GraphBuilder,
@@ -96,7 +97,8 @@ extern TRDGUniformBufferRef<FMobileBasePassUniformParameters> CreateMobileBasePa
 	EMobileBasePass BasePass,
 	EMobileSceneTextureSetupMode SetupMode,
 	const FMobileBasePassTextures& MobileBasePassTextures = {},
-	bool bForRealtimeSkyCapture = false);
+	bool bForRealtimeSkyCapture = false,
+	bool bIsRenderingDeferredBasepass = false);
 
 extern void SetupMobileDirectionalLightUniformParameters(
 	const FScene& Scene,
diff --git a/Engine/Source/Runtime/Renderer/Private/MobileShadingRenderer.cpp b/Engine/Source/Runtime/Renderer/Private/MobileShadingRenderer.cpp
--- a/Engine/Source/Runtime/Renderer/Private/MobileShadingRenderer.cpp
+++ b/Engine/Source/Runtime/Renderer/Private/MobileShadingRenderer.cpp
@@ -2425,7 +2425,7 @@ void FMobileSceneRenderer::RenderDeferredMultiPass(FRDGBuilder& GraphBuilder, FS
 		auto PassParameters = &ParameterCollection->PassParameters;
 		PassParameters->View = View.GetShaderParameters();
 		EMobileSceneTextureSetupMode SetupMode = bIsFullDepthPrepassEnabled ? EMobileSceneTextureSetupMode::SceneDepth : EMobileSceneTextureSetupMode::None;
-		PassParameters->MobileBasePass = CreateMobileBasePassUniformBuffer(GraphBuilder, View, EMobileBasePass::Opaque, SetupMode, BasePassTextures);
+		PassParameters->MobileBasePass = CreateMobileBasePassUniformBuffer(GraphBuilder, View, EMobileBasePass::Opaque, SetupMode, BasePassTextures, false, true);
 		PassParameters->RenderTargets = BasePassRenderTargets;
 		if (!ViewContext.bIsFirstView)
 		{
diff --git a/Engine/Source/Runtime/Renderer/Private/PostProcess/PostProcessAmbientOcclusionMobile.cpp b/Engine/Source/Runtime/Renderer/Private/PostProcess/PostProcessAmbientOcclusionMobile.cpp
--- a/Engine/Source/Runtime/Renderer/Private/PostProcess/PostProcessAmbientOcclusionMobile.cpp
+++ b/Engine/Source/Runtime/Renderer/Private/PostProcess/PostProcessAmbientOcclusionMobile.cpp
@@ -56,6 +56,7 @@ static TAutoConsoleVariable<int32> CVarMobileAmbientOcclusionQuality(
 	ECVF_Scalability | ECVF_RenderThreadSafe
 );
 
+// NOTE: this only applies to GTAO
 static TAutoConsoleVariable<int32> CVarMobileAmbientOcclusionShaderType(
 	TEXT("r.Mobile.AmbientOcclusionShaderType"),
 	2,
@@ -65,6 +66,17 @@ static TAutoConsoleVariable<int32> CVarMobileAmbientOcclusionShaderType(
 	ECVF_RenderThreadSafe
 );
 
+static TAutoConsoleVariable<int32> CVarMobileAmbientOcclusionSSAOCompute(
+	TEXT("r.Mobile.AmbientOcclusionSSAOCompute"),
+	0,
+	TEXT("Whether to use the compute shader variant of mobile SSAO\n")
+	TEXT("0: Disabled (PS).\n")
+	TEXT("1: Compute (graphics pipe)\n")
+	TEXT("2: Compute (async pipe)\n"),
+	ECVF_Scalability | ECVF_RenderThreadSafe
+);
+
+
 static TAutoConsoleVariable<int32> CVarMobileAmbientOcclusionCompression(
 	TEXT("r.Mobile.AmbientOcclusionCompression"),
 	1,
@@ -468,7 +480,7 @@ static void RenderGTAO(FRDGBuilder& GraphBuilder, const FSceneTextures& SceneTex
 
 	const int32 MobileGTAOPreIntegratedTextureType = FMath::Min(CVarMobileGTAOPreIntegratedTextureType.GetValueOnRenderThread(), 2);
 	const int32 MobileAmbientOcclusionQuality = FMath::Min(CVarMobileAmbientOcclusionQuality.GetValueOnRenderThread(), 3);
-		
+
 	FRDGTextureRef SceneDepthTexture = SceneTextures.Depth.Resolve;
 	const FIntPoint DepthBufferSize = SceneDepthTexture->Desc.Extent;
 	const FIntPoint BufferSize = AmbientOcclusionTexture->Desc.Extent;
@@ -492,7 +504,7 @@ static void RenderGTAO(FRDGBuilder& GraphBuilder, const FSceneTextures& SceneTex
 	FMath::SinCos(&SinDeltaAngle, &CosDeltaAngle, PI / NumAngles);
 
 	const EMobileAmbientOcclusionShaderType ShaderType = GetCVarMobileAmbientOcclusionShaderTypeOnRenderThread();
-	
+
 	FRDGTextureUAVRef AmbientOcclusionTextureUAV{};
 	if (IsComputeShaderType(ShaderType))
 	{
@@ -739,6 +751,40 @@ public:
 };
 IMPLEMENT_GLOBAL_SHADER(FMobileAmbientOcclusionPS, "/Engine/Private/PostProcessAmbientOcclusionMobile.usf", "MainPS", SF_Pixel);
 
+class FMobileAmbientOcclusionCS : public FGlobalShader
+{
+public:
+	DECLARE_GLOBAL_SHADER(FMobileAmbientOcclusionCS);
+	SHADER_USE_PARAMETER_STRUCT(FMobileAmbientOcclusionCS, FGlobalShader);
+
+	class FUseUpsampleDim : SHADER_PERMUTATION_BOOL("USE_UPSAMPLE");
+	class FShaderQualityDim : SHADER_PERMUTATION_INT("SHADER_QUALITY", 5);
+
+	using FPermutationDomain = TShaderPermutationDomain<FUseUpsampleDim, FShaderQualityDim>;
+
+	static void ModifyCompilationEnvironment(const FGlobalShaderPermutationParameters& Parameters, FShaderCompilerEnvironment& OutEnvironment)
+	{
+		FGlobalShader::ModifyCompilationEnvironment(Parameters, OutEnvironment);
+		OutEnvironment.SetDefine(TEXT("COMPUTE_SHADER"), 1);
+		OutEnvironment.SetDefine(TEXT("THREADGROUP_SIZEX"), 16);
+		OutEnvironment.SetDefine(TEXT("THREADGROUP_SIZEY"), 16);
+		OutEnvironment.SetDefine(TEXT("SSAO"), 1);
+		OutEnvironment.SetDefine(TEXT("USE_AO_SETUP_AS_INPUT"), 0);
+	}
+
+	static bool ShouldCompilePermutation(const FGlobalShaderPermutationParameters& Parameters)
+	{
+		return IsMobileAmbientOcclusionEnabled(Parameters.Platform);
+	}
+
+	BEGIN_SHADER_PARAMETER_STRUCT(FParameters, )
+		SHADER_PARAMETER_STRUCT_INCLUDE(FMobileAmbientOcclusionParameters, SharedParameters)
+		SHADER_PARAMETER_RDG_TEXTURE_UAV(RWTexture2D, OutTexture)
+	END_SHADER_PARAMETER_STRUCT();
+};
+IMPLEMENT_GLOBAL_SHADER(FMobileAmbientOcclusionCS, "/Engine/Private/PostProcessAmbientOcclusionMobile.usf", "MainCS", SF_Compute);
+
+
 class FMobileAmbientOcclusionUpsamplePS : public FGlobalShader
 {
 	DECLARE_GLOBAL_SHADER(FMobileAmbientOcclusionUpsamplePS);
@@ -927,78 +973,106 @@ static void AddMobileAmbientOcclusionPass(
 	SharedParameters.RandomNormalTexture = GraphBuilder.RegisterExternalTexture(GSystemTextures.SSAORandomization, TEXT("SSAORandomization"));
 	SharedParameters.RandomNormalTextureSampler = TStaticSamplerState<SF_Point, AM_Wrap, AM_Wrap, AM_Wrap>::GetRHI();
 
-	FRDGEventName EventName(TEXT("AmbientOcclusionPS %dx%d"),
-		OutputViewport.Rect.Width(), OutputViewport.Rect.Height());
+	const int32 MobileAmbientOcclusionQuality = FMath::Min(CVarMobileAmbientOcclusionQuality.GetValueOnRenderThread(), 3) - 1;
 
-	FMobileAmbientOcclusionPS::FParameters* PassParameters = GraphBuilder.AllocParameters<FMobileAmbientOcclusionPS::FParameters>();
-	PassParameters->SharedParameters = MoveTemp(SharedParameters);
-	PassParameters->RenderTargets[0] = Output.GetRenderTargetBinding();
-	if (bDepthBoundsTestEnabled)
+	int32 ComputeMethod = CVarMobileAmbientOcclusionSSAOCompute.GetValueOnRenderThread();
+	if (ComputeMethod>0)
 	{
-		PassParameters->RenderTargets.DepthStencil = DepthStencilBinding;
-	}
-
-	const int32 MobileAmbientOcclusionQuality = FMath::Min(CVarMobileAmbientOcclusionQuality.GetValueOnRenderThread(), 3) - 1; 
-	FMobileAmbientOcclusionPS::FPermutationDomain PermutationVector;
-	PermutationVector.Set<FMobileAmbientOcclusionPS::FShaderQualityDim>(MobileAmbientOcclusionQuality);
-	PermutationVector.Set<FMobileAmbientOcclusionPS::FOutputDepth>(bOutputDepth);
-
-	TShaderMapRef<FMobileAmbientOcclusionPS> PixelShader(View.ShaderMap, PermutationVector);
-	TShaderMapRef<FScreenPassVS> VertexShader(View.ShaderMap);
-
-	check(PassParameters);
-	ClearUnusedGraphResources(PixelShader, PassParameters);
+		FRDGEventName EventName(TEXT("AmbientOcclusionCS %dx%d"),
+		OutputViewport.Rect.Width(), OutputViewport.Rect.Height());
 
-	GraphBuilder.AddPass(
-		MoveTemp(EventName),
-		PassParameters,
-		ERDGPassFlags::Raster,
-		[&View, Output, OutputViewport, InputViewport, VertexShader, PixelShader, PassParameters, bHalfResolution, bDepthBoundsTestEnabled, DepthFar](FRHICommandListImmediate& RHICmdList)
+		// Compute Shader Path
+		FMobileAmbientOcclusionCS::FParameters* PassParameters = GraphBuilder.AllocParameters<FMobileAmbientOcclusionCS::FParameters>();
+		PassParameters->SharedParameters = MoveTemp(SharedParameters);
+		PassParameters->OutTexture = GraphBuilder.CreateUAV(Output.Texture);
+
+		FMobileAmbientOcclusionCS::FPermutationDomain PermutationVector;
+		PermutationVector.Set<FMobileAmbientOcclusionCS::FUseUpsampleDim>(bHalfResolution);
+		PermutationVector.Set<FMobileAmbientOcclusionCS::FShaderQualityDim>(MobileAmbientOcclusionQuality);
+
+		TShaderMapRef<FMobileAmbientOcclusionCS> ComputeShader(View.ShaderMap, PermutationVector);
+		FComputeShaderUtils::AddPass(
+			GraphBuilder,
+			MoveTemp(EventName),
+			ComputeMethod == 1 ? ERDGPassFlags::Compute : ERDGPassFlags::AsyncCompute,
+			ComputeShader,
+			PassParameters,
+			FComputeShaderUtils::GetGroupCount(OutputViewport.Rect.Size(), 16));
+	}
+	else
+	{
+		FRDGEventName EventName(TEXT("AmbientOcclusionPS %dx%d"),
+		OutputViewport.Rect.Width(), OutputViewport.Rect.Height());
+    
+		FMobileAmbientOcclusionPS::FParameters* PassParameters = GraphBuilder.AllocParameters<FMobileAmbientOcclusionPS::FParameters>();
+		PassParameters->SharedParameters = MoveTemp(SharedParameters);
+		PassParameters->RenderTargets[0] = Output.GetRenderTargetBinding();
+		if (bDepthBoundsTestEnabled)
 		{
-			const FIntRect InputRect = InputViewport.Rect;
-			const FIntPoint InputSize = InputViewport.Extent;
-			const FIntRect OutputRect = OutputViewport.Rect;
-			const FIntPoint OutputSize = OutputRect.Size();
-
-			RHICmdList.SetViewport(OutputRect.Min.X, OutputRect.Min.Y, 0.0f, OutputRect.Max.X, OutputRect.Max.Y, 1.0f);
-
-			FGraphicsPipelineStateInitializer GraphicsPSOInit;
-			RHICmdList.ApplyCachedRenderTargets(GraphicsPSOInit);
-
-			GraphicsPSOInit.BlendState = TStaticBlendState<>::GetRHI();
-			GraphicsPSOInit.RasterizerState = TStaticRasterizerState<FM_Solid, CM_None>::GetRHI();
-			GraphicsPSOInit.DepthStencilState = TStaticDepthStencilState<false, CF_Always>::GetRHI();
-			GraphicsPSOInit.BoundShaderState.VertexDeclarationRHI = GFilterVertexDeclaration.VertexDeclarationRHI;
-			GraphicsPSOInit.BoundShaderState.VertexShaderRHI = VertexShader.GetVertexShader();
-			GraphicsPSOInit.BoundShaderState.PixelShaderRHI = PixelShader.GetPixelShader();
-			GraphicsPSOInit.PrimitiveType = PT_TriangleList;
-			GraphicsPSOInit.bDepthBounds = bDepthBoundsTestEnabled;
-
-			SetGraphicsPipelineState(RHICmdList, GraphicsPSOInit, 0);
-
-			SetShaderParameters(RHICmdList, PixelShader, PixelShader.GetPixelShader(), *PassParameters);
-
-			if (bDepthBoundsTestEnabled)
-			{
-				RHICmdList.SetDepthBounds(DepthFar, 1.0f);
-			}
-
-			DrawPostProcessPass(
-				RHICmdList,
-				0, 0, OutputSize.X, OutputSize.Y,
-				InputRect.Min.X, InputRect.Min.Y, InputRect.Width(), InputRect.Height(),
-				OutputSize,
-				InputSize,
-				VertexShader,
-				View.StereoViewIndex,
-				false,
-				EDRF_UseTriangleOptimization);
-
-			if (bDepthBoundsTestEnabled)
+			PassParameters->RenderTargets.DepthStencil = DepthStencilBinding;
+		}
+    
+		FMobileAmbientOcclusionPS::FPermutationDomain PermutationVector;
+		PermutationVector.Set<FMobileAmbientOcclusionPS::FShaderQualityDim>(MobileAmbientOcclusionQuality);
+		PermutationVector.Set<FMobileAmbientOcclusionPS::FOutputDepth>(bOutputDepth);
+    
+		TShaderMapRef<FMobileAmbientOcclusionPS> PixelShader(View.ShaderMap, PermutationVector);
+		TShaderMapRef<FScreenPassVS> VertexShader(View.ShaderMap);
+    
+		check(PassParameters);
+		ClearUnusedGraphResources(PixelShader, PassParameters);
+    
+		GraphBuilder.AddPass(
+			MoveTemp(EventName),
+			PassParameters,
+			ERDGPassFlags::Raster,
+			[&View, Output, OutputViewport, InputViewport, VertexShader, PixelShader, PassParameters, bHalfResolution, bDepthBoundsTestEnabled, DepthFar](FRHICommandListImmediate& RHICmdList)
 			{
-				RHICmdList.SetDepthBounds(0, 1.0f);
-			}
-		});	
+				const FIntRect InputRect = InputViewport.Rect;
+				const FIntPoint InputSize = InputViewport.Extent;
+				const FIntRect OutputRect = OutputViewport.Rect;
+				const FIntPoint OutputSize = OutputRect.Size();
+    
+				RHICmdList.SetViewport(OutputRect.Min.X, OutputRect.Min.Y, 0.0f, OutputRect.Max.X, OutputRect.Max.Y, 1.0f);
+    
+				FGraphicsPipelineStateInitializer GraphicsPSOInit;
+				RHICmdList.ApplyCachedRenderTargets(GraphicsPSOInit);
+    
+				GraphicsPSOInit.BlendState = TStaticBlendState<>::GetRHI();
+				GraphicsPSOInit.RasterizerState = TStaticRasterizerState<FM_Solid, CM_None>::GetRHI();
+				GraphicsPSOInit.DepthStencilState = TStaticDepthStencilState<false, CF_Always>::GetRHI();
+				GraphicsPSOInit.BoundShaderState.VertexDeclarationRHI = GFilterVertexDeclaration.VertexDeclarationRHI;
+				GraphicsPSOInit.BoundShaderState.VertexShaderRHI = VertexShader.GetVertexShader();
+				GraphicsPSOInit.BoundShaderState.PixelShaderRHI = PixelShader.GetPixelShader();
+				GraphicsPSOInit.PrimitiveType = PT_TriangleList;
+				GraphicsPSOInit.bDepthBounds = bDepthBoundsTestEnabled;
+    
+				SetGraphicsPipelineState(RHICmdList, GraphicsPSOInit, 0);
+    
+				SetShaderParameters(RHICmdList, PixelShader, PixelShader.GetPixelShader(), *PassParameters);
+    
+				if (bDepthBoundsTestEnabled)
+				{
+					RHICmdList.SetDepthBounds(DepthFar, 1.0f);
+				}
+    
+				DrawPostProcessPass(
+					RHICmdList,
+					0, 0, OutputSize.X, OutputSize.Y,
+					InputRect.Min.X, InputRect.Min.Y, InputRect.Width(), InputRect.Height(),
+					OutputSize,
+					InputSize,
+					VertexShader,
+					View.StereoViewIndex,
+					false,
+					EDRF_UseTriangleOptimization);
+    
+				if (bDepthBoundsTestEnabled)
+				{
+					RHICmdList.SetDepthBounds(0, 1.0f);
+				}
+			});	
+	}
 }
 
 static void AddMobileAmbientOcclusionUpsamplePass(
@@ -1130,8 +1204,13 @@ static void RenderSSAO(FRDGBuilder& GraphBuilder, const FSceneTextures& SceneTex
 		// Bilateral requires 32bit format for AO + depth.
 		const EPixelFormat Format = bBilateralUpsample ? PF_R8G8B8A8 : PF_R8;
 		const FIntPoint Extent = GetDownscaledExtent(AmbientOcclusionTexture->Desc.Extent, 2);
+		ETextureCreateFlags TextureCreateFlags = TexCreate_ShaderResource | TexCreate_RenderTargetable;
+		if (CVarMobileAmbientOcclusionSSAOCompute.GetValueOnRenderThread() > 0)
+		{
+			TextureCreateFlags |= TexCreate_UAV;
+		}
 		HalfResolutionTexture = GraphBuilder.CreateTexture(
-			FRDGTextureDesc::Create2D(Extent, Format, FClearValueBinding::None, TexCreate_ShaderResource | TexCreate_RenderTargetable),
+			FRDGTextureDesc::Create2D(Extent, Format, FClearValueBinding::None, TextureCreateFlags),
 			TEXT("HalfResolutionScreenSpaceAO"));
 	}
 
@@ -1174,17 +1253,17 @@ FRDGTextureRef CreateMobileScreenSpaceAOTexture(FRDGBuilder& GraphBuilder, const
 	const FIntPoint Extent = FIntPoint::DivideAndRoundUp(Config.Extent, DownsampleFactor);
 	EPixelFormat Format = PF_R8;
 
-	const EMobileAmbientOcclusionShaderType ShaderType = GetCVarMobileAmbientOcclusionShaderTypeOnRenderThread();
+	const EMobileAmbientOcclusionShaderType GtaoShaderType = GetCVarMobileAmbientOcclusionShaderTypeOnRenderThread();
+
+	const bool bComputeShader = ( bGTAO && IsComputeShaderType(GtaoShaderType) ) || CVarMobileAmbientOcclusionSSAOCompute.GetValueOnRenderThread() > 0;
+
 	// G8 isn't supported as UAV on Android OpenGLES, fall back to RGBA8 for compute shader usage.
-	if (bGTAO
-		&& IsOpenGLPlatform(Config.ShaderPlatform)
-		&& IsComputeShaderType(ShaderType)
-		)
+	if (bComputeShader && IsOpenGLPlatform(Config.ShaderPlatform) )
 	{
 		Format = PF_R8G8B8A8;
 	}
 	ETextureCreateFlags TextureCreateFlags = TexCreate_ShaderResource | TexCreate_RenderTargetable;
-	if (IsComputeShaderType(ShaderType))
+	if (bComputeShader)
 	{
 		TextureCreateFlags |= TexCreate_UAV;
 	}
