[LightGrid] Refactored light grid feedback status system to use a global resource instead of per-view message handler stored in FSceneViewState.
- avoids race condition calling GPUMessage::RegisterHandler(...) which caused a crash.

diff --git a/Engine/Source/Runtime/Renderer/Private/LightGrid.h b/Engine/Source/Runtime/Renderer/Private/LightGrid.h
deleted file mode 100644
diff --git a/Engine/Source/Runtime/Renderer/Private/LightGridInjection.cpp b/Engine/Source/Runtime/Renderer/Private/LightGridInjection.cpp
--- a/Engine/Source/Runtime/Renderer/Private/LightGridInjection.cpp
+++ b/Engine/Source/Runtime/Renderer/Private/LightGridInjection.cpp
@@ -4,8 +4,6 @@
 LightGridInjection.cpp
 =============================================================================*/
 
-#include "LightGrid.h"
-
 #include "CoreMinimal.h"
 #include "Stats/Stats.h"
 #include "HAL/IConsoleManager.h"
@@ -185,6 +183,17 @@ FAutoConsoleVariableRef CVarLightCullingMaxDistanceOverride(
 
 DECLARE_GPU_STAT(LightGrid);
 
+#if !UE_BUILD_SHIPPING
+void LightGridFeedbackStatus(
+	FRDGBuilder& GraphBuilder,
+	FViewInfo& View,
+	FRDGBufferRef CulledLightDataAllocatorBuffer,
+	uint32 NumCulledLightDataEntries,
+	FRDGBufferRef CulledLightLinkAllocatorBuffer,
+	uint32 NumCulledLightLinks,
+	bool bUseAsyncCompute);
+#endif // !UE_BUILD_SHIPPING
+
 bool ShouldVisualizeLightGrid(EShaderPlatform InShaderPlatform)
 {
 	const bool bEnabled = GForwardLightGridDebug > 0;
@@ -609,7 +618,6 @@ FLightGrid LightGridInjection(
 	FRDGBufferSRVRef LightViewSpaceDirAndPreprocAngleSRV,
 	FRDGBufferSRVRef LightViewSpaceRectPlanesSRV,
 	FRDGBufferSRVRef IndirectionIndicesSRV,
-	FLightGridViewState* LightGridViewState,
 	bool bThreadGroupPerCell,
 	bool bThreadGroupSize32,
 	// parent params
@@ -626,7 +634,7 @@ FLightGrid LightGridInjection(
 		ensureMsgf(NumLocalLights <= LIGHT_GRID_CELL_WRITER_MAX_NUM_PRIMITIVES, TEXT("NumLocalLights limited to 16M by FCellWriter."));
 		ensureMsgf(NumReflectionCaptures <= LIGHT_GRID_CELL_WRITER_MAX_NUM_PRIMITIVES, TEXT("NumLocalLights limited to 16M by FCellWriter."));
 
-		NumCulledLightLinks = FMath::Min(NumCulledLightLinks, (uint32)LIGHT_GRID_CELL_WRITER_MAX_NUM_LINKS); // limited to 16M by FCellWriter (will cause warning if exceeded, see FLightGridViewState::Update())
+		NumCulledLightLinks = FMath::Min(NumCulledLightLinks, (uint32)LIGHT_GRID_CELL_WRITER_MAX_NUM_LINKS); // limited to 16M by FCellWriter (will cause warning if exceeded, see FLightGridFeedbackStatus)
 	}
 
 	const FIntVector ParentGridSize = FIntVector::DivideAndRoundUp(GridSize, ParentGridSizeFactor);
@@ -798,10 +806,9 @@ FLightGrid LightGridInjection(
 	Output.CulledLightDataGridSRV = CulledLightDataGridSRV;
 	Output.NumCulledLightsGridSRV = NumCulledLightsGridSRV;
 
-	if (LightGridViewState != nullptr)
-	{
-		LightGridViewState->FeedbackStatus(GraphBuilder, View, CulledLightDataAllocatorBuffer, NumCulledLightEntries, CulledLightLinkAllocatorBuffer, NumCulledLightLinks, bUseAsyncCompute);
-	}
+#if !UE_BUILD_SHIPPING
+	LightGridFeedbackStatus(GraphBuilder, View, CulledLightDataAllocatorBuffer, NumCulledLightEntries, CulledLightLinkAllocatorBuffer, NumCulledLightLinks, bUseAsyncCompute);
+#endif
 
 	return Output;
 }
@@ -1421,7 +1428,6 @@ FComputeLightGridOutput FSceneRenderer::ComputeLightGrid(FRDGBuilder& GraphBuild
 				LightViewSpaceDirAndPreprocAngleSRV,
 				LightViewSpaceRectPlanesSRV,
 				nullptr,
-				View.ViewState != nullptr ? &View.ViewState->LightGrid : nullptr,
 				/*bThreadGroupPerCell*/ true,
 				/*bThreadGroupSize32*/ false,
 				nullptr,
@@ -1463,7 +1469,6 @@ FComputeLightGridOutput FSceneRenderer::ComputeLightGrid(FRDGBuilder& GraphBuild
 			LightViewSpaceDirAndPreprocAngleSRV,
 			LightViewSpaceRectPlanesSRV,
 			IndirectionIndicesSRV,
-			View.ViewState != nullptr ? &View.ViewState->LightGrid : nullptr,
 			NumThreadsPerCell > 1,
 			NumThreadsPerCell == 32,
 			ParentLightGrid.NumCulledLightsGridSRV,
@@ -1824,9 +1829,27 @@ class FLightGridFeedbackStatusCS : public FGlobalShader
 };
 IMPLEMENT_GLOBAL_SHADER(FLightGridFeedbackStatusCS, "/Engine/Private/LightGridInjection.usf", "FeedbackStatusCS", SF_Compute);
 
-FLightGridViewState::FLightGridViewState()
-{
 #if !UE_BUILD_SHIPPING
+class FLightGridFeedbackStatus : public FRenderResource
+{
+public:
+	virtual ~FLightGridFeedbackStatus() {}
+
+	virtual void InitRHI(FRHICommandListBase& RHICmdList);
+
+	virtual void ReleaseRHI();
+
+	uint32 GetStatusMessageId() const { return StatusFeedbackSocket.GetMessageId().GetIndex(); }
+
+private:
+	GPUMessage::FSocket StatusFeedbackSocket;
+	uint32 MaxEntriesHighWaterMark = 0;
+};
+
+TGlobalResource<FLightGridFeedbackStatus> GLightGridFeedbackStatus;
+
+void FLightGridFeedbackStatus::InitRHI(FRHICommandListBase& RHICmdList)
+{
 	StatusFeedbackSocket = GPUMessage::RegisterHandler(TEXT("LightGrid.StatusFeedback"),
 		[this](GPUMessage::FReader Message)
 		{
@@ -1860,12 +1883,22 @@ FLightGridViewState::FLightGridViewState()
 				}
 			}
 		});
-#endif
 }
 
-void FLightGridViewState::FeedbackStatus(FRDGBuilder& GraphBuilder, FViewInfo& View, FRDGBufferRef CulledLightDataAllocatorBuffer, uint32 NumCulledLightDataEntries, FRDGBufferRef CulledLightLinkAllocatorBuffer, uint32 NumCulledLightLinks, bool bUseAsyncCompute)
+void FLightGridFeedbackStatus::ReleaseRHI()
+{
+	
+}
+
+void LightGridFeedbackStatus(
+	FRDGBuilder& GraphBuilder,
+	FViewInfo& View,
+	FRDGBufferRef CulledLightDataAllocatorBuffer,
+	uint32 NumCulledLightDataEntries,
+	FRDGBufferRef CulledLightLinkAllocatorBuffer,
+	uint32 NumCulledLightLinks,
+	bool bUseAsyncCompute)
 {
-#if !UE_BUILD_SHIPPING
 	FLightGridFeedbackStatusCS::FParameters* PassParameters = GraphBuilder.AllocParameters<FLightGridFeedbackStatusCS::FParameters>();
 
 	PassParameters->CulledLightDataAllocatorBuffer = GraphBuilder.CreateSRV(CulledLightDataAllocatorBuffer);
@@ -1875,7 +1908,7 @@ void FLightGridViewState::FeedbackStatus(FRDGBuilder& GraphBuilder, FViewInfo& V
 	PassParameters->NumAvailableLinks = NumCulledLightLinks;
 
 	PassParameters->GPUMessageParams = GPUMessage::GetShaderParameters(GraphBuilder);
-	PassParameters->StatusMessageId = GetStatusMessageId();
+	PassParameters->StatusMessageId = GLightGridFeedbackStatus.GetStatusMessageId();
 
 	auto ComputeShader = View.ShaderMap->GetShader<FLightGridFeedbackStatusCS>();
 
@@ -1887,5 +1920,6 @@ void FLightGridViewState::FeedbackStatus(FRDGBuilder& GraphBuilder, FViewInfo& V
 		PassParameters,
 		FIntVector(1, 1, 1)
 	);
+}
+
 #endif // !UE_BUILD_SHIPPING
-}
\ No newline at end of file
diff --git a/Engine/Source/Runtime/Renderer/Private/ScenePrivate.h b/Engine/Source/Runtime/Renderer/Private/ScenePrivate.h
--- a/Engine/Source/Runtime/Renderer/Private/ScenePrivate.h
+++ b/Engine/Source/Runtime/Renderer/Private/ScenePrivate.h
@@ -43,7 +43,6 @@
 #include "RayTracing/RayTracingShaderBindingTable.h"
 #endif
 #include "Nanite/Nanite.h"
-#include "LightGrid.h"
 #include "Lumen/LumenViewState.h"
 #include "MegaLights/MegaLightsViewState.h"
 #include "StochasticLighting/StochasticLightingViewState.h"
@@ -1000,8 +999,6 @@ public:
 	// Burley Subsurface scattering variance texture from the last frame.
 	TRefCountPtr<IPooledRenderTarget> SubsurfaceScatteringQualityHistoryRT;
 
-	FLightGridViewState LightGrid;
-
 	FLumenViewState Lumen;
 	FMegaLightsViewState MegaLights;
 
