* Add min/max voxel size definition for hair voxelization to ensure value are consistent throughout the code.
* Increase max limit in order to handle very large groom.

diff --git a/Engine/Shaders/Private/HairStrands/HairStrandsAABB.usf b/Engine/Shaders/Private/HairStrands/HairStrandsAABB.usf
--- a/Engine/Shaders/Private/HairStrands/HairStrandsAABB.usf
+++ b/Engine/Shaders/Private/HairStrands/HairStrandsAABB.usf
@@ -2,6 +2,7 @@
 
 #include "../Common.ush"
 #include "HairStrandsAABBCommon.ush"
+#include "/Engine/Shared/HairStrandsDefinitions.h"
 
 #if SHADER_AABBUPDATE 
 
@@ -92,8 +93,7 @@ void Main(uint2 DispatchThreadId : SV_DispatchThreadID, uint LinearIndex : SV_Gr
 			VoxelWorldSize = MaxExtent / float(VoxelPageResolution);
 		}
 
-		const float MinVoxelWorldSize = 0.01f; // 0.01mm
-		OutMacroGroupVoxelSizeBuffer[MacroGroupId] = max(MinVoxelWorldSize, VoxelWorldSize);
+		OutMacroGroupVoxelSizeBuffer[MacroGroupId] = max(HAIR_MIN_VOXEL_SIZE, VoxelWorldSize);
 	}
 }
 
diff --git a/Engine/Shaders/Private/HairStrands/HairStrandsVoxelPageCommon.ush b/Engine/Shaders/Private/HairStrands/HairStrandsVoxelPageCommon.ush
--- a/Engine/Shaders/Private/HairStrands/HairStrandsVoxelPageCommon.ush
+++ b/Engine/Shaders/Private/HairStrands/HairStrandsVoxelPageCommon.ush
@@ -3,6 +3,7 @@
 #pragma once
 
 #include "HairStrandsVoxelPageCommonStruct.ush"
+#include "/Engine/Shared/HairStrandsDefinitions.h"
 
 ///////////////////////////////////////////////////////////////////////////////////////////////////
 // Define voxel traversal type
@@ -87,13 +88,13 @@ struct FVirtualVoxelNodeDesc
 
 uint PackVoxelWorldSize(float In)
 {
-	// Encode voxel world size [0..10] onto 10bits, i.e.: 10.f / 1023u
-	const float NormalizedVoxelWorldSize = min(In, 10.f) * 0.1f;
+	// Encode voxel world size [0..HAIR_MAX_VOXEL_SIZE] onto 10bits
+	const float NormalizedVoxelWorldSize = min(In, HAIR_MAX_VOXEL_SIZE) * (1.f/HAIR_MAX_VOXEL_SIZE);
 	return PackUnorm10(NormalizedVoxelWorldSize);
 }
 float UnpackVoxelWorldSize(uint In)
 {
-	return UnpackUnorm10(In) * 10.f;
+	return UnpackUnorm10(In) * HAIR_MAX_VOXEL_SIZE;
 }
 
 float QuantizeVoxelWorldSize(float In)
diff --git a/Engine/Shaders/Private/HairStrands/HairStrandsVoxelRasterCompute.usf b/Engine/Shaders/Private/HairStrands/HairStrandsVoxelRasterCompute.usf
--- a/Engine/Shaders/Private/HairStrands/HairStrandsVoxelRasterCompute.usf
+++ b/Engine/Shaders/Private/HairStrands/HairStrandsVoxelRasterCompute.usf
@@ -590,10 +590,8 @@ void FeedbackCS(uint3 DispatchThreadId : SV_DispatchThreadID)
 		NextVoxelWorldSize = CPUMinVoxelWorldSize;
 	}
 
-	// Clamp voxel size into a reasonable amount (e.g. 0.1mm - 100mm)
-	const float ClampMinVoxelWorldSize = 0.01f;
-	const float ClampMaxVoxelWorldSize = 10.0f;
-	NextVoxelWorldSize = clamp(RoundHairVoxelSize(NextVoxelWorldSize), ClampMinVoxelWorldSize, ClampMaxVoxelWorldSize);
+	// Clamp voxel size
+	NextVoxelWorldSize = clamp(RoundHairVoxelSize(NextVoxelWorldSize), HAIR_MIN_VOXEL_SIZE, HAIR_MAX_VOXEL_SIZE);
 
 	// Debug
 #if DEBUG_ENABLE
diff --git a/Engine/Shaders/Shared/HairStrandsDefinitions.h b/Engine/Shaders/Shared/HairStrandsDefinitions.h
--- a/Engine/Shaders/Shared/HairStrandsDefinitions.h
+++ b/Engine/Shaders/Shared/HairStrandsDefinitions.h
@@ -32,7 +32,6 @@
 #define HAIR_ATTRIBUTE_INVALID_OFFSET 0xFFFFFFFF
 
 // Hair interpolation (in bytes)
-
 #define HAIR_INTERPOLATION_CURVE_STRIDE 8
 #define HAIR_INTERPOLATION_POINT_STRIDE 2
 #define HAIR_INTERPOLATION_MAX_GUIDE_COUNT 2
@@ -106,6 +105,9 @@
 
 #define HAIR_CARDS_MAX_TEXTURE_COUNT 6
 
+#define HAIR_MIN_VOXEL_SIZE 0.01f
+#define HAIR_MAX_VOXEL_SIZE 100.f
+
 // Hair instance flags 
 #define HAIR_FLAGS_SCATTER_SCENE_LIGHT 0x1u
 #define HAIR_FLAGS_STABLE_RASTER 0x2u
diff --git a/Engine/Source/Runtime/Renderer/Private/HairStrands/HairStrandsVoxelization.cpp b/Engine/Source/Runtime/Renderer/Private/HairStrands/HairStrandsVoxelization.cpp
--- a/Engine/Source/Runtime/Renderer/Private/HairStrands/HairStrandsVoxelization.cpp
+++ b/Engine/Source/Runtime/Renderer/Private/HairStrands/HairStrandsVoxelization.cpp
@@ -841,9 +841,7 @@ static FHairStrandsVoxelResources AllocateVirtualVoxelResources(
 	// Init. default page table size and voxel size
 	FIntVector PageCountResolution = FIntVector(GHairVirtualVoxel_PageCountPerDim, GHairVirtualVoxel_PageCountPerDim, GHairVirtualVoxel_PageCountPerDim);
 	uint32 PageCount = PageCountResolution.X * PageCountResolution.Y * PageCountResolution.Z;
-	const float ClampMinVoxelWorldSize = 0.01f;
-	const float ClampMaxVoxelWorldSize = 10.f;
-	const float CPUMinVoxelWorldSize = RoundHairVoxeliSize(FMath::Clamp(GHairVirtualVoxel_VoxelWorldSize, ClampMinVoxelWorldSize, ClampMaxVoxelWorldSize));
+	const float CPUMinVoxelWorldSize = RoundHairVoxeliSize(FMath::Clamp(GHairVirtualVoxel_VoxelWorldSize, HAIR_MIN_VOXEL_SIZE, HAIR_MAX_VOXEL_SIZE));
 	const bool bVoxelAllocationFeedbackEnable = OutViewStateData && IsHairStrandsAdaptiveVoxelAllocationEnable();
 
 	FHairStrandsVoxelResources Out;
