[Ray Traced Translucency] RT translucency can run when Lumen is disabled so still create RGS for it

#tests tested in editor, ran the failing test and passed

diff --git a/Engine/Source/Runtime/Renderer/Private/DeferredShadingRenderer.cpp b/Engine/Source/Runtime/Renderer/Private/DeferredShadingRenderer.cpp
--- a/Engine/Source/Runtime/Renderer/Private/DeferredShadingRenderer.cpp
+++ b/Engine/Source/Runtime/Renderer/Private/DeferredShadingRenderer.cpp
@@ -703,16 +703,19 @@ bool FDeferredShadingSceneRenderer::SetupRayTracingPipelineStatesAndSBT(FRDGBuil
 					PrepareRayTracingVolumetricFogShadows(View, *Scene, RayGenShaders);
 
 					if (DoesPlatformSupportLumenGI(ShaderPlatform) 
-						&& Lumen::UseHardwareRayTracing(ViewFamily)
-						&& IsLumenEnabled(View))
+						&& Lumen::UseHardwareRayTracing(ViewFamily))
 					{
-						PrepareLumenHardwareRayTracingScreenProbeGather(View, RayGenShaders);
-						PrepareLumenHardwareRayTracingShortRangeAO(View, RayGenShaders);
-						PrepareLumenHardwareRayTracingRadianceCache(View, RayGenShaders);
-						PrepareLumenHardwareRayTracingReflections(View, RayGenShaders);
+						if (IsLumenEnabled(View))
+						{
+							PrepareLumenHardwareRayTracingScreenProbeGather(View, RayGenShaders);
+							PrepareLumenHardwareRayTracingShortRangeAO(View, RayGenShaders);
+							PrepareLumenHardwareRayTracingRadianceCache(View, RayGenShaders);
+							PrepareLumenHardwareRayTracingReflections(View, RayGenShaders);
+							PrepareLumenHardwareRayTracingReSTIR(View, RayGenShaders);
+							PrepareLumenHardwareRayTracingVisualize(View, RayGenShaders);
+						}
+
 						PrepareHardwareRayTracingTranslucency(View, RayGenShaders);
-						PrepareLumenHardwareRayTracingReSTIR(View, RayGenShaders);
-						PrepareLumenHardwareRayTracingVisualize(View, RayGenShaders);
 					}
 
 					PrepareMegaLightsHardwareRayTracing(View, *Scene, RayGenShaders);
