Fix to prevent SceneColorTint from reaching zero, which causes preExposure ensure()

diff --git a/Engine/Source/Runtime/Renderer/Private/PostProcess/PostProcessEyeAdaptation.cpp b/Engine/Source/Runtime/Renderer/Private/PostProcess/PostProcessEyeAdaptation.cpp
--- a/Engine/Source/Runtime/Renderer/Private/PostProcess/PostProcessEyeAdaptation.cpp
+++ b/Engine/Source/Runtime/Renderer/Private/PostProcess/PostProcessEyeAdaptation.cpp
@@ -1441,7 +1441,8 @@ void FViewInfo::UpdatePreExposure()
 	else
 	{
 		// How much the SceneColorTint changes the overall brightness of the image.
-		const float SceneColorTint = FinalPostProcessSettings.SceneColorTint.GetLuminance();
+		// (prevent tint form being zero, because would trigger preExposure ensure)
+		const float SceneColorTint = FMath::Max(FinalPostProcessSettings.SceneColorTint.GetLuminance(), 0.0001);
 
 		// How much the vignette changes the overall brightness of the image. ComputeVignetteMask() always returns 1 at the center and dim
 		// the edge of the screen.
