Unshelved from pending changelist '44679834':

Add Beer Shadow Map alternative for volumetric shadows on hard and translucent surfaces.

r.HeterogeneousVolumes.Shadows.ShadowType (default = 1)
0: Adaptive Volumetric Shadow Map
1: Beer Shadow Map

We prefer Beer Shadow Maps by default for Epic scalability, since this amounts to a significant runtime improvement. It is expected that cinematics may prefer AVSMs as they are higher quality, but come at a cost, especially when supporting multiple volumes.

#tests Dev testing on custom project. Engine tests

diff --git a/Engine/Shaders/Private/DeferredLightPixelShaders.usf b/Engine/Shaders/Private/DeferredLightPixelShaders.usf
--- a/Engine/Shaders/Private/DeferredLightPixelShaders.usf
+++ b/Engine/Shaders/Private/DeferredLightPixelShaders.usf
@@ -71,13 +71,6 @@
 #include "VirtualShadowMaps/VirtualShadowMapMaskBitsCommon.ush"
 #endif
 
-#if USE_ADAPTIVE_VOLUMETRIC_SHADOW_MAP
-#include "HeterogeneousVolumes/HeterogeneousVolumesAdaptiveVolumetricShadowMapSampling.ush"
-
-//#include "HeterogeneousVolumes/HeterogeneousVolumesVoxelGridUtils2.ush"
-//#include "HeterogeneousVolumes/HeterogeneousVolumesRayMarchingUtils.ush"
-#endif // USE_ADAPTIVE_VOLUMETRIC_SHADOW_MAP
-
 // Remap light shape permutation into light type
 #if LIGHT_SOURCE_SHAPE == 0
 #define CURRENT_LIGHT_TYPE LIGHT_TYPE_DIRECTIONAL
@@ -322,10 +315,6 @@ void DeferredLightPixelMain(
 		#endif
 			SubstrateShadowTermInputParameters.TranslatedWorldPosition				= DerivedParams.TranslatedWorldPosition;
 			SubstrateShadowTermInputParameters.LightAttenuation						= GetLightAttenuationFromShadow(InputParams, SceneDepth, DerivedParams.TranslatedWorldPosition);
-#if USE_ADAPTIVE_VOLUMETRIC_SHADOW_MAP
-			//SubstrateShadowTermInputParameters.LightAttenuation *= Luminance(ComputeTransmittance(DerivedParams.TranslatedWorldPosition, LightData.TranslatedWorldPosition, 256));
-			SubstrateShadowTermInputParameters.LightAttenuation *= AVSM_SampleTransmittance(DerivedParams.TranslatedWorldPosition, LightData.TranslatedWorldPosition);
-#endif // USE_ADAPTIVE_VOLUMETRIC_SHADOW_MAP
 			SubstrateShadowTermInputParameters.Dither								= InterleavedGradientNoise(InputParams.PixelPos, View.StateFrameIndexMod8);
 			
 		#if SUBSTRATE_LOAD_FROM_MATERIALCONTAINER==0
@@ -402,24 +391,6 @@ void DeferredLightPixelMain(
 		float SurfaceShadow = 1.0f;
 		
 		float4 LightAttenuation = GetLightAttenuationFromShadow(InputParams, SceneDepth, DerivedParams.TranslatedWorldPosition);
-#if USE_ADAPTIVE_VOLUMETRIC_SHADOW_MAP
-		float3 WorldPosition = DFFastSubtractDemote(DerivedParams.TranslatedWorldPosition, PrimaryView.PreViewTranslation);
-		float3 WorldToLight = LightData.Direction * 10000;
-		if (LightData.bSpotLight || LightData.bRectLight)
-		{
-			float3 WorldLightPosition = DFFastSubtractDemote(LightData.TranslatedWorldPosition, PrimaryView.PreViewTranslation);
-			WorldToLight = WorldLightPosition - WorldPosition;
-		}
-
-		if (!AVSM.bIsEmpty)
-		{
-			LightAttenuation *= AVSM_SampleTransmittance(DerivedParams.TranslatedWorldPosition, LightData.TranslatedWorldPosition);
-		}
-		else
-		{
-			//LightAttenuation *= Luminance(ComputeTransmittance(WorldPosition, WorldToLight, 512));
-		}
-#endif // USE_ADAPTIVE_VOLUMETRIC_SHADOW_MAP
 		float4 Radiance = GetDynamicLighting(DerivedParams.TranslatedWorldPosition, DerivedParams.CameraVector, ScreenSpaceData.GBuffer, ScreenSpaceData.AmbientOcclusion, LightData, LightAttenuation, Dither, uint2(InputParams.PixelPos), SurfaceShadow);
 
 		OutColor += Radiance;
diff --git a/Engine/Shaders/Private/DeferredLightingCommon.ush b/Engine/Shaders/Private/DeferredLightingCommon.ush
--- a/Engine/Shaders/Private/DeferredLightingCommon.ush
+++ b/Engine/Shaders/Private/DeferredLightingCommon.ush
@@ -19,21 +19,10 @@
 #include "LightFunctionAtlas/LightFunctionAtlasCommon.usf"
 #endif
 
-#ifndef ADAPTIVE_VOLUMETRIC_SHADOW_MAP
-#define ADAPTIVE_VOLUMETRIC_SHADOW_MAP 0
-#endif
-
 #ifndef ALLOW_LOCAL_LIGHT_DISTANCE_ATTENUATION
 #define ALLOW_LOCAL_LIGHT_DISTANCE_ATTENUATION 0
 #endif
 
-#if ADAPTIVE_VOLUMETRIC_SHADOW_MAP
-#include "HeterogeneousVolumes/HeterogeneousVolumesAdaptiveVolumetricShadowMapSampling.ush"
-
-//#include "HeterogeneousVolumes/HeterogeneousVolumesVoxelGridUtils2.ush"
-//#include "HeterogeneousVolumes/HeterogeneousVolumesRayMarchingUtils.ush"
-#endif // ADAPTIVE_VOLUMETRIC_SHADOW_MAP
-
 // This Substrate include assumes that if inline shading is required, it needs to 
 // be defined prior to DeferredLightingCommong.ush
 // SubstrateStruct.MaterialTextureArray is only available to translucent materials or global shaders computing contact shadows.
@@ -346,10 +335,6 @@ FLightAccumulator AccumulateDynamicLighting(
 	if (LightData.bRadialLight)
 	{
 		LightMask = GetLocalLightAttenuation( TranslatedWorldPosition, LightData, ToLight, L );
-#if ADAPTIVE_VOLUMETRIC_SHADOW_MAP
-		//LightAttenuation *= ComputeTransmittance(DerivedParams.TranslatedWorldPosition, LightData.TranslatedWorldPosition, 256);
-		LightAttenuation *= AVSM_SampleTransmittance(TranslatedWorldPosition, LightData.TranslatedWorldPosition);
-#endif // ADAPTIVE_VOLUMETRIC_SHADOW_MAP
 		MaskedLightColor *= LightMask;
 	}
 
diff --git a/Engine/Shaders/Private/HeterogeneousVolumes/HeterogeneousVolumesAdaptiveVolumetricShadowMapSampling.ush b/Engine/Shaders/Private/HeterogeneousVolumes/HeterogeneousVolumesAdaptiveVolumetricShadowMapSampling.ush
--- a/Engine/Shaders/Private/HeterogeneousVolumes/HeterogeneousVolumesAdaptiveVolumetricShadowMapSampling.ush
+++ b/Engine/Shaders/Private/HeterogeneousVolumes/HeterogeneousVolumesAdaptiveVolumetricShadowMapSampling.ush
@@ -3,6 +3,7 @@
 #pragma once
 
 #include "HeterogeneousVolumesAdaptiveVolumetricShadowMapUtils.ush"
+#include "HeterogeneousVolumesBeerShadowMapUtils.ush"
 
 #define AVSM_SAMPLE_MODE_DISABLED 0
 #define AVSM_SAMPLE_MODE_TWO_LEVEL_VECTORIZED 1
@@ -15,6 +16,8 @@
 #define AVSM_SAMPLE_MODE AVSM_SAMPLE_MODE_TWO_LEVEL_VECTORIZED_LOOP
 #endif // AVSM_SAMPLE_MODE
 
+// Adaptive Volumetric Shadow Map implementation
+
 struct FAdaptiveVolumetricShadowMap
 {
 	float4x4 TranslatedWorldToShadow[6];
@@ -445,14 +448,6 @@ float AVSM_SampleTransmittance(
 		ShadowMap.NumShadowMatrices - 1
 	};
 
-#if 0
-	if (View.GeneralPurposeTweak >= 0)
-	{
-		Face[0] = max(View.GeneralPurposeTweak - 1, 0);
-		Face[1] = min(View.GeneralPurposeTweak, ShadowMap.NumShadowMatrices - 1);
-	}
-	else
-#endif
 	if (AVSM.bIsDirectionalLight)
 	{
 		// Find the highest quality cascade for the shade point
@@ -623,4 +618,237 @@ float4 AVSM_SampleCameraRadianceAndTransmittance4(
 	float Weight = saturate((1.0 - Transmittance) * SafeRcp(1.0 - Result.a));
 
 	return float4(Radiance * Weight, Transmittance);
-}
\ No newline at end of file
+}
+
+// Beer Shadow Map implementation
+
+struct FBeerShadowMap
+{
+	float4x4 TranslatedWorldToShadow[6];
+	float3 TranslatedWorldOrigin[6];
+	float4 TranslatedWorldPlane[6];
+	float4 SplitDepths[6];
+	float DownsampleFactor;
+
+	int2 Resolution;
+	int NumShadowMatrices;
+	int MaxSampleCount;
+	bool bIsEmpty;
+	bool bIsDirectionalLight;
+
+	Texture2D<float4> BeerShadowMapTexture;
+	Texture2D<float4> RadianceTexture;
+	SamplerState TextureSampler;
+};
+
+FBeerShadowMap GetBeerShadowMap()
+{
+	FBeerShadowMap ShadowMap;
+	for (int i = 0; i < 6; ++i)
+	{
+		ShadowMap.TranslatedWorldToShadow[i] = BeerShadowMap.TranslatedWorldToShadow[i];
+		ShadowMap.TranslatedWorldOrigin[i] = BeerShadowMap.TranslatedWorldOrigin[i].xyz;
+		ShadowMap.TranslatedWorldPlane[i] = BeerShadowMap.TranslatedWorldPlane[i];
+		ShadowMap.SplitDepths[i] = BeerShadowMap.SplitDepths[i];
+	}
+
+	ShadowMap.Resolution = BeerShadowMap.Resolution;
+	ShadowMap.DownsampleFactor = BeerShadowMap.DownsampleFactor;
+	ShadowMap.NumShadowMatrices = BeerShadowMap.NumShadowMatrices;
+	ShadowMap.MaxSampleCount = BeerShadowMap.MaxSampleCount;
+	ShadowMap.bIsEmpty = BeerShadowMap.bIsEmpty;
+	ShadowMap.bIsDirectionalLight = BeerShadowMap.bIsDirectionalLight;
+
+	ShadowMap.BeerShadowMapTexture = BeerShadowMap.BeerShadowMapTexture;
+	ShadowMap.RadianceTexture = BeerShadowMap.RadianceTexture;
+	ShadowMap.TextureSampler = BeerShadowMap.TextureSampler;
+	return ShadowMap;
+}
+
+float BeerShadowMap_SampleTransmittanceSimple(
+	inout FBeerShadowMap ShadowMap,
+	float3 TranslatedWorldPosition,
+	float3 LightTranslatedWorldPosition
+)
+{
+	if (any(ShadowMap.Resolution == 0))
+	{
+		return 1.0;
+	}
+
+	float Transmittance[] = { 1.0, 1.0 };
+	float Alpha = 0.0;
+
+	float Depth = length(TranslatedWorldPosition - View.TranslatedWorldCameraOrigin);
+
+	int Face[] = {
+		ShadowMap.NumShadowMatrices - 1,
+		ShadowMap.NumShadowMatrices - 1
+	};
+
+	float4 ShadowPositionAsFloat4 = 0;
+	float3 ShadowPosition = 0;
+	for (int Index = 0; Index < 1; ++Index)
+	{
+		ShadowPositionAsFloat4 = mul(float4(TranslatedWorldPosition, 1), ShadowMap.TranslatedWorldToShadow[Face[Index]]);
+		if (ShadowPositionAsFloat4.w > 0)
+		{
+			ShadowPosition = ShadowPositionAsFloat4.xyz / ShadowPositionAsFloat4.w;
+
+			// Remap ShadowPosition to Texture UV
+			ShadowPosition.x = (Face[Index] + ShadowPosition.x) / BeerShadowMap.NumShadowMatrices;
+			float3 ClipUV[] = {
+				float3(float(Face[Index]) / BeerShadowMap.NumShadowMatrices, 0, -1),
+				float3(float(Face[Index] + 1) / BeerShadowMap.NumShadowMatrices, 1, 1)
+			};
+
+			if (all(ShadowPosition > ClipUV[0]) && all(ShadowPosition < ClipUV[1]))
+			{
+				float4 PackedData = Texture2DSample(ShadowMap.BeerShadowMapTexture, ShadowMap.TextureSampler, ShadowPosition.xy);
+				FBeerShadowMapData ShadowMapData = UnpackBeerShadowMapData(PackedData);
+
+				float ShadowZ = length(TranslatedWorldPosition - ShadowMap.TranslatedWorldOrigin[Face[Index]]);
+				if (ShadowMap.bIsDirectionalLight)
+				{
+					ShadowZ = dot(ShadowMap.TranslatedWorldPlane[Face[Index]].xyz, TranslatedWorldPosition) + ShadowMap.TranslatedWorldPlane[Face[Index]].w;
+				}
+
+				FAVSMSampleData SampleData[] = {
+					AVSM_CreateSampleData(ShadowMapData.T[0], ShadowMapData.Tau[0]),
+					AVSM_CreateSampleData(ShadowMapData.T[1], ShadowMapData.Tau[1])
+				};
+				Transmittance[Index] = Interpolate(SampleData, ShadowZ);
+			}
+		}
+	}
+
+	return ((1.0 - Alpha) * Transmittance[0]) + (Alpha * Transmittance[1]);
+}
+
+float BeerShadowMap_SampleTransmittance(
+	inout FBeerShadowMap ShadowMap,
+	float3 TranslatedWorldPosition,
+	float3 LightTranslatedWorldPosition
+)
+{
+	if (any(ShadowMap.Resolution == 0))
+	{
+		return 1.0;
+	}
+
+	float Transmittance[] = { 1.0, 1.0 };
+	float Alpha = 1.0;
+
+	float Depth = length(TranslatedWorldPosition - View.TranslatedWorldCameraOrigin);
+
+	int Face[] = {
+		ShadowMap.NumShadowMatrices - 1,
+		ShadowMap.NumShadowMatrices - 1
+	};
+
+	if (BeerShadowMap.bIsDirectionalLight)
+	{
+		// Find the highest quality cascade for the shade point
+		for (; Face[1] > 0; --Face[1])
+		{
+			float Far = ShadowMap.SplitDepths[Face[1]].y;
+			if (Depth <= Far)
+			{
+				break;
+			}
+		}
+		Face[0] = max(Face[1] - 1, 0);
+
+		// Apply near-far distance fades
+		float Near = ShadowMap.SplitDepths[Face[1]].x;
+		float NearFadeDist = ShadowMap.SplitDepths[Face[1]].z;
+		if (Depth - Near < NearFadeDist)
+		{
+			Alpha = saturate((Depth - Near) / NearFadeDist);
+			Face[0] = min(Face[1] + 1, ShadowMap.NumShadowMatrices - 1);
+		}
+
+		float Far = ShadowMap.SplitDepths[Face[1]].y;
+		float FarFadeDist = ShadowMap.SplitDepths[Face[1]].w;
+		if (Far - Depth < FarFadeDist)
+		{
+			Alpha = saturate((Far - Depth) / FarFadeDist);
+			Face[0] = max(Face[1] - 1, 0);
+		}
+	}
+	else if (BeerShadowMap.NumShadowMatrices > 1)
+	{
+		float3 LightDir = TranslatedWorldPosition - ShadowMap.TranslatedWorldOrigin[0];
+		float3 LightDirAbs = abs(LightDir);
+
+		Face[0] = LightDirAbs[1] > LightDirAbs[0] ? 1 : 0;
+		Face[0] = LightDirAbs[2] > LightDirAbs[Face[0]] ? 2 : Face[0];
+
+		int FaceOffset = (sign(LightDir[Face[0]]) > 0) ? 1 : 0;
+		Face[0] = Face[1] = Face[0] * 2 + FaceOffset;
+	}
+
+	float4 ShadowPositionAsFloat4 = 0;
+	float3 ShadowPosition = 0;
+	for (int Index = 0; Index < 2; ++Index)
+	{
+		ShadowPositionAsFloat4 = mul(float4(TranslatedWorldPosition, 1), ShadowMap.TranslatedWorldToShadow[Face[Index]]);
+		if (ShadowPositionAsFloat4.w > 0)
+		{
+			ShadowPosition = ShadowPositionAsFloat4.xyz / ShadowPositionAsFloat4.w;
+
+			// Remap ShadowPosition to Texture UV
+			ShadowPosition.x = (Face[Index] + ShadowPosition.x) / BeerShadowMap.NumShadowMatrices;
+			float3 ClipUV[] = {
+				float3(float(Face[Index]) / BeerShadowMap.NumShadowMatrices, 0, -1),
+				float3(float(Face[Index] + 1) / BeerShadowMap.NumShadowMatrices, 1, 1)
+			};
+
+			if (all(ShadowPosition > ClipUV[0]) && all(ShadowPosition < ClipUV[1]))
+			{
+				float4 PackedData = Texture2DSample(ShadowMap.BeerShadowMapTexture, ShadowMap.TextureSampler, ShadowPosition.xy);
+				FBeerShadowMapData ShadowMapData = UnpackBeerShadowMapData(PackedData);
+
+				float ShadowZ = length(TranslatedWorldPosition - ShadowMap.TranslatedWorldOrigin[Face[Index]]);
+				if (ShadowMap.bIsDirectionalLight)
+				{
+					ShadowZ = dot(ShadowMap.TranslatedWorldPlane[Face[Index]].xyz, TranslatedWorldPosition) + ShadowMap.TranslatedWorldPlane[Face[Index]].w;
+				}
+
+				FAVSMSampleData SampleData[] = {
+					AVSM_CreateSampleData(ShadowMapData.T[0], ShadowMapData.Tau[0]),
+					AVSM_CreateSampleData(ShadowMapData.T[1], ShadowMapData.Tau[1])
+				};
+				Transmittance[Index] = Interpolate(SampleData, ShadowZ);
+			}
+		}
+	}
+
+	return ((1.0 - Alpha) * Transmittance[0]) + (Alpha * Transmittance[1]);
+}
+
+float BeerShadowMap_SampleTransmittance(
+	float3 TranslatedWorldPosition,
+	float3 LightTranslatedWorldPosition
+)
+{
+	FBeerShadowMap BeerShadowMap = GetBeerShadowMap();
+	return BeerShadowMap_SampleTransmittance(BeerShadowMap, TranslatedWorldPosition, LightTranslatedWorldPosition);
+}
+
+float4 BeerShadowMap_SampleCameraRadianceAndTransmittance4(
+	float2 ScreenUV,
+	float3 TranslatedWorldPosition,
+	float3 LightTranslatedWorldPosition
+)
+{
+	FBeerShadowMap BeerShadowMap = GetBeerShadowMap();
+	float Transmittance = BeerShadowMap_SampleTransmittanceSimple(BeerShadowMap, TranslatedWorldPosition, LightTranslatedWorldPosition);
+
+	// Extract radiance from precomposited render target and weight based on relative opacity
+	float4 Result = Texture2DSample(BeerShadowMap.RadianceTexture, BeerShadowMap.TextureSampler, ScreenUV);
+	float3 Radiance = Result.rgb * View.OneOverPreExposure;
+	float Weight = saturate((1.0 - Transmittance) * SafeRcp(1.0 - Result.a));
+
+	return float4(Radiance * Weight, Transmittance);
+}
diff --git a/Engine/Shaders/Private/HeterogeneousVolumes/HeterogeneousVolumesBeerShadowMapUtils.ush b/Engine/Shaders/Private/HeterogeneousVolumes/HeterogeneousVolumesBeerShadowMapUtils.ush
new file mode 100644
--- /dev/null
+++ b/Engine/Shaders/Private/HeterogeneousVolumes/HeterogeneousVolumesBeerShadowMapUtils.ush
@@ -0,0 +1,107 @@
+// Copyright Epic Games, Inc. All Rights Reserved.
+
+#pragma once
+
+// Beer Shadow Map implementation
+struct FBeerShadowMapData
+{
+	float T[2];
+	float Tau[2];
+};
+
+FBeerShadowMapData InitializeBeerShadowMapData(float Transmittance, float MaxTraceDistance)
+{
+	FBeerShadowMapData Data;
+	Data.T[0] = MaxTraceDistance;
+	Data.T[1] = 0.0;
+	Data.Tau[0] = Transmittance;
+	Data.Tau[1] = Transmittance;
+	return Data;
+}
+
+FBeerShadowMapData UnpackBeerShadowMapData(float4 PackedData)
+{
+	FBeerShadowMapData Data;
+	Data.T[0] = PackedData.r;
+	Data.T[1] = PackedData.g;
+	Data.Tau[0] = PackedData.b;
+	Data.Tau[1] = PackedData.a;
+	return Data;
+}
+
+float4 PackBeerShadowMapData(FBeerShadowMapData Data)
+{
+	float4 PackedData;
+	PackedData.r = Data.T[0];
+	PackedData.g = Data.T[1];
+	PackedData.b = Data.Tau[0];
+	PackedData.a = Data.Tau[1];
+	return PackedData;
+}
+
+#define BEER_SHADOW_MAP_COMBINE_METHOD_FIRST 0
+#define BEER_SHADOW_MAP_COMBINE_METHOD_EXPAND 1
+
+FBeerShadowMapData ExpandBeerShadowMapData(FBeerShadowMapData Data0, FBeerShadowMapData Data1)
+{
+	FBeerShadowMapData Result;
+	Result.T[0] = min(Data0.T[0], Data1.T[0]);
+	Result.T[1] = max(Data0.T[1], Data1.T[1]);
+	Result.Tau[0] = max(Data0.Tau[0], Data0.Tau[0]);
+	Result.Tau[1] = min(Data0.Tau[1], Data1.Tau[1]);
+	return Result;
+}
+
+bool IsValid(FBeerShadowMapData Data)
+{
+	return Data.T[1] > Data.T[0];
+}
+
+FBeerShadowMapData FirstBeerShadowMapData(FBeerShadowMapData Data0, FBeerShadowMapData Data1)
+{
+	bool bIsValid0 = IsValid(Data0);
+	bool bIsValid1 = IsValid(Data1);
+
+	if (bIsValid0 && bIsValid1)
+	{
+		if (Data0.T[0] < Data1.T[0])
+		{
+			return Data0;
+		}
+		else
+		{
+			return Data1;
+		}
+	}
+	else if (bIsValid0)
+	{
+		return Data0;
+	}
+	// else
+	return Data1;
+}
+
+FBeerShadowMapData CombineBeerShadowMapData(FBeerShadowMapData Data0, FBeerShadowMapData Data1, int CombineMethod)
+{
+	FBeerShadowMapData Result;
+
+	if (!IsValid(Data0))
+	{
+		return Data1;
+	}
+
+	if (!IsValid(Data1))
+	{
+		return Data0;
+	}
+
+	if (CombineMethod == BEER_SHADOW_MAP_COMBINE_METHOD_FIRST)
+	{
+		Result = FirstBeerShadowMapData(Data0, Data1);
+	}
+	else // if (CombineMethod == BEER_SHADOW_MAP_COMBINE_METHOD_EXPAND)
+	{
+		Result = ExpandBeerShadowMapData(Data0, Data1);
+	}
+	return Result;
+}
diff --git a/Engine/Shaders/Private/HeterogeneousVolumes/HeterogeneousVolumesLiveShadingPipeline.usf b/Engine/Shaders/Private/HeterogeneousVolumes/HeterogeneousVolumesLiveShadingPipeline.usf
--- a/Engine/Shaders/Private/HeterogeneousVolumes/HeterogeneousVolumesLiveShadingPipeline.usf
+++ b/Engine/Shaders/Private/HeterogeneousVolumes/HeterogeneousVolumesLiveShadingPipeline.usf
@@ -44,6 +44,7 @@
 #endif // SUPPORT_OVERLAPPING_VOLUMES
 
 #include "/Engine/Generated/Material.ush"
+#include "HeterogeneousVolumesBeerShadowMapUtils.ush"
 #include "HeterogeneousVolumesStochasticFiltering.ush"
 #include "HeterogeneousVolumesLiveShadingUtils.ush"
 #include "HeterogeneousVolumesTracingUtils.ush"
@@ -632,94 +633,6 @@ void RayMarchSingleScatteringAdaptive(
 
 #endif // USE_CAMERA_AVSM
 
-struct FBeerShadowMapData
-{
-	float T[2];
-	float Tau[2];
-};
-
-FBeerShadowMapData InitializeBeerShadowMapData(float Transmittance)
-{
-	FBeerShadowMapData Data;
-	Data.T[0] = GetMaxTraceDistance();
-	Data.T[1] = 0.0;
-	Data.Tau[0] = Transmittance;
-	Data.Tau[1] = Transmittance;
-	return Data;
-}
-
-FBeerShadowMapData UnpackBeerShadowMapData(float4 PackedData)
-{
-	FBeerShadowMapData Data;
-	Data.T[0] = PackedData.r;
-	Data.T[1] = PackedData.g;
-	Data.Tau[0] = PackedData.b;
-	Data.Tau[1] = PackedData.a;
-	return Data;
-}
-
-float4 PackBeerShadowMapData(FBeerShadowMapData Data)
-{
-	float4 PackedData;
-	PackedData.r = Data.T[0];
-	PackedData.g = Data.T[1];
-	PackedData.b = Data.Tau[0];
-	PackedData.a = Data.Tau[1];
-	return PackedData;
-}
-
-#define BEER_SHADOW_MAP_COMBINE_METHOD_FIRST 0
-#define BEER_SHADOW_MAP_COMBINE_METHOD_EXPAND 1
-
-FBeerShadowMapData ExpandBeerShadowMapData(FBeerShadowMapData Data0, FBeerShadowMapData Data1)
-{
-	FBeerShadowMapData Result;
-	Result.T[0] = min(Data0.T[0], Data1.T[0]);
-	Result.T[1] = max(Data0.T[1], Data1.T[1]);
-	Result.Tau[0] = Data0.Tau[0];
-	Result.Tau[1] = Data1.Tau[1];
-	return Result;
-}
-
-FBeerShadowMapData FirstBeerShadowMapData(FBeerShadowMapData Data0, FBeerShadowMapData Data1)
-{
-	bool bIsValid0 = Data0.T[1] > Data0.T[0];
-	bool bIsValid1 = Data1.T[1] > Data1.T[0];
-
-	if (bIsValid0 && bIsValid1)
-	{
-		if (Data0.T[0] < Data1.T[0])
-		{
-			return Data0;
-		}
-		else
-		{
-			return Data1;
-		}
-	}
-	else if (bIsValid0)
-	{
-		return Data0;
-	}
-	// else
-	return Data1;
-}
-
-FBeerShadowMapData CombineBeerShadowMapData(FBeerShadowMapData Data0, FBeerShadowMapData Data1, int CombineMethod)
-{
-	FBeerShadowMapData Result;
-
-	if (CombineMethod == BEER_SHADOW_MAP_COMBINE_METHOD_FIRST)
-	{
-		Result = FirstBeerShadowMapData(Data0, Data1);
-	}
-	else // if (CombineMethod == BEER_SHADOW_MAP_COMBINE_METHOD_EXPAND)
-	{
-		Result = ExpandBeerShadowMapData(Data0, Data1);
-	}
-	return Result;
-}
-
 void RenderSingleScatteringWithLiveShading(
 	uint2 DispatchThreadId
 )
@@ -746,7 +659,7 @@ void RenderSingleScatteringWithLiveShading(
 	float PrevTransmittance = RWLightingTexture[PixelCoord].a;
 	float PrevOpacity = Luminance(1.0 - PrevTransmittance);
 
-	FBeerShadowMapData BeerShadowMapData = InitializeBeerShadowMapData(PrevTransmittance);
+	FBeerShadowMapData BeerShadowMapData = InitializeBeerShadowMapData(PrevTransmittance, GetMaxTraceDistance());
 
 	// Extract depth
 	float DeviceZ = SceneDepthTexture.Load(int3(ViewPixelCoord, 0)).r;
diff --git a/Engine/Shaders/Private/HeterogeneousVolumes/HeterogeneousVolumesLiveShadingShadows.usf b/Engine/Shaders/Private/HeterogeneousVolumes/HeterogeneousVolumesLiveShadingShadows.usf
--- a/Engine/Shaders/Private/HeterogeneousVolumes/HeterogeneousVolumesLiveShadingShadows.usf
+++ b/Engine/Shaders/Private/HeterogeneousVolumes/HeterogeneousVolumesLiveShadingShadows.usf
@@ -15,6 +15,7 @@
 #include "HeterogeneousVolumesLiveShadingUtils.ush"
 #include "HeterogeneousVolumesTracingUtils.ush"
 #include "HeterogeneousVolumesAdaptiveVolumetricShadowMapUtils.ush"
+#include "HeterogeneousVolumesBeerShadowMapUtils.ush"
 #include "../BlueNoise.ush"
 
 #ifndef THREADGROUP_SIZE_1D
@@ -29,6 +30,13 @@
 #define THREADGROUP_SIZE_3D 1
 #endif // THREADGROUP_SIZE_3D
 
+#define SHADOW_MAP_TYPE_AVSM 0
+#define SHADOW_MAP_TYPE_BEER 1
+
+#ifndef SHADOW_MAP_TYPE
+#define SHADOW_MAP_TYPE SHADOW_MAP_TYPE_AVSM
+#endif // SHADOW_MAP_TYPE
+
 // Object data
 float4x4 LocalToWorld;
 float4x4 WorldToLocal;
@@ -251,6 +259,12 @@ uint GetLinearIndex(uint Face, uint2 PixelCoord)
 	return LinearIndex;
 }
 
+uint2 GetTexture2DIndex(uint Face, uint2 PixelCoord)
+{
+	int2 Texture2DIndex = uint2(Face * ShadowResolution.x + PixelCoord.x, PixelCoord.y);
+	return Texture2DIndex;
+}
+
 // Cinematic feature options
 bool UseInscatteringVolume()
 {
@@ -403,6 +417,120 @@ void AddOrReplaceSample(FAVSMSampleData Data, int PixelIndex, inout int PrevSamp
 
 int CameraDownsampleFactor;
 
+float3 RayMarchTransmittance(
+	inout FRayMarchingContext RayMarchingContext,
+	int PixelIndex,
+	inout FBeerShadowMapData BeerShadowMapData
+)
+{
+	float3 LocalBoundsMin = LocalBoundsOrigin - LocalBoundsExtent;
+	float3 LocalBoundsMax = LocalBoundsOrigin + LocalBoundsExtent;
+
+	float3 Transmittance = 1;
+	InitializeBeerShadowMapData(1.0, GetMaxTraceDistance());
+
+	uint StepCount = CalcStepCount(RayMarchingContext);
+
+	float3 ExpectedExtinction = 0.0;
+	float3 ExpectedTransmittance = 1.0;
+	int PrevElementIndex = AVSM_NULL_PTR;
+	int ElementIndex = 0;
+
+	float StepSize = RayMarchingContext.StepSize;
+	float JitterSize = StepSize * RayMarchingContext.Jitter;
+
+	float LocalHitT = RayMarchingContext.LocalRayTMin;
+	float WorldHitT = LocalHitT * RayMarchingContext.LocalToWorldScale;
+	float LastWorldHitT = WorldHitT;
+	for (uint StepIndex = 0; StepIndex < StepCount; ++StepIndex)
+	{
+		LocalHitT = min(RayMarchingContext.LocalRayTMin + StepSize * (RayMarchingContext.Jitter + StepIndex), RayMarchingContext.LocalRayTMax);
+		if (StepIndex == StepCount - 1)
+		{
+			float UnjitteredLocalHitT = LocalHitT - JitterSize;
+			StepSize = RayMarchingContext.LocalRayTMax - UnjitteredLocalHitT;
+		}
+
+		WorldHitT = LocalHitT * RayMarchingContext.LocalToWorldScale;
+		float3 LocalPosition = RayMarchingContext.LocalRayOrigin + RayMarchingContext.LocalRayDirection * LocalHitT;
+		float3 WorldPosition = RayMarchingContext.WorldRayOrigin + RayMarchingContext.WorldRayDirection * WorldHitT;
+
+		FVolumeSampleContext SampleContext = CreateVolumeSampleContext(LocalPosition, WorldPosition, RayMarchingContext.MipLevel);
+		float3 Extinction = SampleExtinction(SampleContext);
+		Transmittance *= exp(-Extinction * StepSize);
+
+		if (any(Extinction > 0.0))
+		{
+			LastWorldHitT = WorldHitT;
+			bool bFirstEntry = (ElementIndex == 0);
+
+#if SHADOW_MAP_TYPE == SHADOW_MAP_TYPE_BEER
+			if (bFirstEntry)
+			{
+				BeerShadowMapData.T[0] = WorldHitT;
+				BeerShadowMapData.Tau[0] = Luminance(Transmittance);
+				ElementIndex = 1;
+			}
+#else // SHADOW_MAP_TYPE == SHADOW_MAP_TYPE_AVSM
+			if (bFirstEntry)
+			{
+				FAVSMSampleData SampleData = AVSM_CreateSampleData(WorldHitT, Transmittance);
+				AddSample(SampleData, PixelIndex, PrevElementIndex, ElementIndex);
+
+				// Set expected values
+				ExpectedExtinction = Extinction;
+				ExpectedTransmittance = Transmittance;
+			}
+			else
+			{
+				// If extrapolated transmittance is beyond relative error tolerance
+				ExpectedTransmittance *= exp(-ExpectedExtinction * StepSize);
+				float AbsoluteError = length(Transmittance - ExpectedTransmittance);
+				float RelativeError = length((Transmittance - ExpectedTransmittance) / ExpectedTransmittance);
+				if ((AbsoluteError > AbsoluteErrorThreshold) && (RelativeError > RelativeErrorThreshold))
+				{
+					FAVSMSampleData SampleData = AVSM_CreateSampleData(WorldHitT, Transmittance);
+					AddOrReplaceSample(SampleData, PixelIndex, PrevElementIndex, ElementIndex);
+
+					// Reset expected values
+					ExpectedExtinction = Extinction;
+					ExpectedTransmittance = Transmittance;
+				}
+			}
+#endif
+		}
+
+		// Early termination
+		float Epsilon = 1.0e-7;
+		if (all(Transmittance < Epsilon))
+		{
+			Transmittance = 0.0;
+			break;
+		}
+	}
+
+#if SHADOW_TYPE == SHADOW_TYPE_AVSM
+	// Stitch up the previous entry
+	if (PrevElementIndex != AVSM_NULL_PTR)
+	{
+		AVSM_SetLinkedListNextPtr(RWVolumetricShadowLinkedListBuffer[PixelIndex + PrevElementIndex], ElementIndex);
+
+	}
+
+	// Last Entry
+	AVSM_SetLinkedListSampleData(RWVolumetricShadowLinkedListBuffer[PixelIndex + ElementIndex], AVSM_CreateSampleData(LastWorldHitT, Transmittance));
+	AVSM_SetLinkedListPrevPtr(RWVolumetricShadowLinkedListBuffer[PixelIndex + ElementIndex], PrevElementIndex);
+	AVSM_SetLinkedListNextPtr(RWVolumetricShadowLinkedListBuffer[PixelIndex + ElementIndex], AVSM_NULL_PTR);
+#endif
+
+#if SHADOW_MAP_TYPE == SHADOW_MAP_TYPE_BEER
+	BeerShadowMapData.T[1] = LastWorldHitT;
+	BeerShadowMapData.Tau[1] = saturate(Luminance(Transmittance));
+#endif
+
+	return Transmittance;
+}
+
 [numthreads(THREADGROUP_SIZE_2D, THREADGROUP_SIZE_2D, 1)]
 void RenderVolumetricShadowMapForLightWithLiveShadingCS(
 	uint2 GroupThreadId : SV_GroupThreadID,
@@ -466,6 +594,7 @@ void RenderVolumetricShadowMapForLightWithLiveShadingCS(
 	float3 LocalBoundsMax = LocalBoundsOrigin + LocalBoundsExtent;
 
 	float3 Transmittance = 1;
+	FBeerShadowMapData BeerShadowMapData = InitializeBeerShadowMapData(1.0, GetMaxTraceDistance());
 
 	// Intersect bounding volume
 	float2 HitT = IntersectAABB(LocalRayOrigin, LocalRayDirection, 0.0, LocalRayLength, LocalBoundsMin, LocalBoundsMax);
@@ -495,94 +624,27 @@ void RenderVolumetricShadowMapForLightWithLiveShadingCS(
 			bApplyDirectLighting,
 			bApplyShadowTransmittance
 		);
-
 		RayMarchingContext.RandSequence = RandomSequenceCreate(LinearIndex, View.StateFrameIndex);
-		uint StepCount = CalcStepCount(RayMarchingContext);
-
-		float3 ExpectedExtinction = 0.0;
-		float3 ExpectedTransmittance = 1.0;
-		int PrevElementIndex = AVSM_NULL_PTR;
-		int ElementIndex = 0;
-
-		float StepSize = RayMarchingContext.StepSize;
-		float JitterSize = StepSize * RayMarchingContext.Jitter;
 
-		float LocalHitT = HitT.x;
-		float WorldHitT = LocalHitT * RayMarchingContext.LocalToWorldScale;
-		float LastWorldHitT = WorldHitT;
-		for (uint StepIndex = 0; StepIndex < StepCount; ++StepIndex)
-		{
-			LocalHitT = min(RayMarchingContext.LocalRayTMin + StepSize * (RayMarchingContext.Jitter + StepIndex), RayMarchingContext.LocalRayTMax);
-			if (StepIndex == StepCount - 1)
-			{
-				float UnjitteredLocalHitT = LocalHitT - JitterSize;
-				StepSize = RayMarchingContext.LocalRayTMax - UnjitteredLocalHitT;
-			}
-
-			WorldHitT = LocalHitT * RayMarchingContext.LocalToWorldScale;
-			float3 LocalPosition = RayMarchingContext.LocalRayOrigin + RayMarchingContext.LocalRayDirection * LocalHitT;
-			float3 WorldPosition = RayMarchingContext.WorldRayOrigin + RayMarchingContext.WorldRayDirection * WorldHitT;
-
-			FVolumeSampleContext SampleContext = CreateVolumeSampleContext(LocalPosition, WorldPosition, RayMarchingContext.MipLevel);
-			float3 Extinction = SampleExtinction(SampleContext);
-			Transmittance *= exp(-Extinction * StepSize);
-
-			if (any(Extinction > 0.0))
-			{
-				LastWorldHitT = WorldHitT;
-				bool bFirstEntry = (ElementIndex == 0);
-				if (bFirstEntry)
-				{
-					FAVSMSampleData SampleData = AVSM_CreateSampleData(WorldHitT, Transmittance);
-					AddSample(SampleData, PixelIndex, PrevElementIndex, ElementIndex);
-
-					// Set expected values
-					ExpectedExtinction = Extinction;
-					ExpectedTransmittance = Transmittance;
-				}
-				else
-				{
-					// If extrapolated transmittance is beyond relative error tolerance
-					ExpectedTransmittance *= exp(-ExpectedExtinction * StepSize);
-					float AbsoluteError = length(Transmittance - ExpectedTransmittance);
-					float RelativeError = length((Transmittance - ExpectedTransmittance) / ExpectedTransmittance);
-					if ((AbsoluteError > AbsoluteErrorThreshold) && (RelativeError > RelativeErrorThreshold))
-					{
-						FAVSMSampleData SampleData = AVSM_CreateSampleData(WorldHitT, Transmittance);
-						AddOrReplaceSample(SampleData, PixelIndex, PrevElementIndex, ElementIndex);
-
-						// Reset expected values
-						ExpectedExtinction = Extinction;
-						ExpectedTransmittance = Transmittance;
-					}
-				}
-			}
-
-			// Early termination
-			float Epsilon = 1.0e-7;
-			if (all(Transmittance < Epsilon))
-			{
-				Transmittance = 0.0;
-				break;
-			}
-		}
-
-		// Stitch up the previous entry
-		if (PrevElementIndex != AVSM_NULL_PTR)
-		{
-			AVSM_SetLinkedListNextPtr(RWVolumetricShadowLinkedListBuffer[PixelIndex + PrevElementIndex], ElementIndex);
-
-		}
-
-		// Last Entry
-		AVSM_SetLinkedListSampleData(RWVolumetricShadowLinkedListBuffer[PixelIndex + ElementIndex], AVSM_CreateSampleData(LastWorldHitT, Transmittance));
-		AVSM_SetLinkedListPrevPtr(RWVolumetricShadowLinkedListBuffer[PixelIndex + ElementIndex], PrevElementIndex);
-		AVSM_SetLinkedListNextPtr(RWVolumetricShadowLinkedListBuffer[PixelIndex + ElementIndex], AVSM_NULL_PTR);
+		Transmittance = RayMarchTransmittance(RayMarchingContext, PixelIndex, BeerShadowMapData);
 	}
+#if SHADOW_MAP_TYPE == SHADOW_MAP_TYPE_AVSM
 	else
+#endif
 	{
 		AVSM_SetLinkedListSampleData(RWVolumetricShadowLinkedListBuffer[PixelIndex], AVSM_CreateSampleData(0.0, Transmittance));
 		AVSM_SetLinkedListPrevPtr(RWVolumetricShadowLinkedListBuffer[PixelIndex], AVSM_NULL_PTR);
 		AVSM_SetLinkedListNextPtr(RWVolumetricShadowLinkedListBuffer[PixelIndex], AVSM_NULL_PTR);
 	}
+
+#if SHADOW_MAP_TYPE == SHADOW_MAP_TYPE_BEER
+	uint2 BeerShadowMapCoord = GetTexture2DIndex(Face, PixelCoord);
+	RWBeerShadowMapTexture[BeerShadowMapCoord] = PackBeerShadowMapData(
+		CombineBeerShadowMapData(
+			UnpackBeerShadowMapData(RWBeerShadowMapTexture[BeerShadowMapCoord]),
+			BeerShadowMapData,
+			BEER_SHADOW_MAP_COMBINE_METHOD_EXPAND)
+		// BEER_SHADOW_MAP_COMBINE_METHOD_FIRST)
+	);
+#endif
 }
diff --git a/Engine/Shaders/Private/HeterogeneousVolumes/HeterogeneousVolumesRayMarchingUtils.ush b/Engine/Shaders/Private/HeterogeneousVolumes/HeterogeneousVolumesRayMarchingUtils.ush
--- a/Engine/Shaders/Private/HeterogeneousVolumes/HeterogeneousVolumesRayMarchingUtils.ush
+++ b/Engine/Shaders/Private/HeterogeneousVolumes/HeterogeneousVolumesRayMarchingUtils.ush
@@ -206,6 +206,12 @@ float3 ComputeTransmittance(
 			float3 LightTranslatedWorldPosition = TranslatedWorldPosition + ToLight;
 			Transmittance = AVSM_SampleTransmittance(TranslatedWorldPosition, LightTranslatedWorldPosition);
 		}
+		else if (!BeerShadowMap.bIsEmpty)
+		{
+			float3 TranslatedWorldPosition = DFFastAddDemote(WorldRayOrigin, PrimaryView.PreViewTranslation);
+			float3 LightTranslatedWorldPosition = TranslatedWorldPosition + ToLight;
+			Transmittance = BeerShadowMap_SampleTransmittance(TranslatedWorldPosition, LightTranslatedWorldPosition);
+		}
 		else
 	#endif
 		{
diff --git a/Engine/Shaders/Private/HeterogeneousVolumes/HeterogeneousVolumesShadowMask.usf b/Engine/Shaders/Private/HeterogeneousVolumes/HeterogeneousVolumesShadowMask.usf
new file mode 100644
--- /dev/null
+++ b/Engine/Shaders/Private/HeterogeneousVolumes/HeterogeneousVolumesShadowMask.usf
@@ -0,0 +1,91 @@
+// Copyright Epic Games, Inc. All Rights Reserved.
+
+#define SUPPORT_CONTACT_SHADOWS 0
+
+#define DYNAMICALLY_SHADOWED 1
+#define TREAT_MAXDEPTH_UNSHADOWED 1
+
+#define SHADOW_QUALITY 2
+#define NO_TRANSLUCENCY_AVAILABLE
+
+#include "../Common.ush"
+#include "../DeferredLightingCommon.ush"
+#include "HeterogeneousVolumesAdaptiveVolumetricShadowMapSampling.ush"
+
+#ifndef THREADGROUP_SIZE_1D
+#define THREADGROUP_SIZE_1D 1
+#endif // THREADGROUP_SIZE_1D
+
+#ifndef THREADGROUP_SIZE_2D
+#define THREADGROUP_SIZE_2D 1
+#endif // THREADGROUP_SIZE_2D
+
+#define SHADOW_TYPE_VOLUMETRIC_SHADOW_MAP 0
+#define SHADOW_TYPE_BEER_SHADOW_MAP 1
+
+#ifndef SHADOW_TYPE
+#define SHADOW_TYPE SHADOW_TYPE_VOLUMETRIC_SHADOW_MAP
+#endif // SHADOW_TYPE
+
+// Light data
+int LightType;
+
+// Dispatch data
+int3 GroupCount;
+int DownsampleFactor;
+
+// Shadow mask
+Texture2D<float> ShadowMaskTexture;
+
+// Output
+RWTexture2D<float4> RWVolumetricShadowMaskTexture;
+
+int2 GetDownsampledResolution(int2 Resolution, int Factor)
+{
+	return (Resolution + Factor - 1) / Factor;
+}
+
+int2 GetScaledViewRect()
+{
+	return GetDownsampledResolution(View.ViewSizeAndInvSize.xy, DownsampleFactor);
+}
+
+float SampleTransmittance(
+	float3 TranslatedWorldPosition,
+	float3 LightTranslatedWorldPosition
+)
+{
+#if SHADOW_TYPE == SHADOW_TYPE_VOLUMETRIC_SHADOW_MAP
+	return AVSM_SampleTransmittance(TranslatedWorldPosition, LightTranslatedWorldPosition);
+#else
+	return BeerShadowMap_SampleTransmittance(TranslatedWorldPosition, LightTranslatedWorldPosition);
+#endif
+}
+
+[numthreads(THREADGROUP_SIZE_2D, THREADGROUP_SIZE_2D, 1)]
+void RenderVolumetricShadowMaskCS(
+	uint GroupId : SV_GroupID,
+	uint GroupThreadId : SV_GroupThreadID,
+	uint2 DispatchThreadId : SV_DispatchThreadID
+)
+{
+	if (any(DispatchThreadId.xy >= GetScaledViewRect()))
+	{
+		return;
+	}
+	uint2 PixelCoord = DispatchThreadId.xy + View.ViewRectMin.xy / DownsampleFactor;
+
+	// Extract depth
+	float DeviceZ = SceneDepthTexture.Load(int3(PixelCoord, 0)).r;
+#if HAS_INVERTED_Z_BUFFER
+	DeviceZ = max(0.000000000001, DeviceZ);
+#endif // HAS_INVERTED_Z_BUFFER
+
+	// Initialize as input
+	RWVolumetricShadowMaskTexture[PixelCoord] = ShadowMaskTexture[PixelCoord];
+
+	// Query Beer Shadow Map
+	const float3 TranslatedWorldPosition = SvPositionToTranslatedWorld(float4(PixelCoord + 0.5, DeviceZ, 1.0));
+	float Transmittance = SampleTransmittance(TranslatedWorldPosition, GetDeferredLightTranslatedWorldPosition());
+	RWVolumetricShadowMaskTexture[PixelCoord] *= sqrt(Transmittance);
+}
\ No newline at end of file
diff --git a/Engine/Source/Runtime/Renderer/Private/HeterogeneousVolumes/HeterogeneousVolumes.cpp b/Engine/Source/Runtime/Renderer/Private/HeterogeneousVolumes/HeterogeneousVolumes.cpp
--- a/Engine/Source/Runtime/Renderer/Private/HeterogeneousVolumes/HeterogeneousVolumes.cpp
+++ b/Engine/Source/Runtime/Renderer/Private/HeterogeneousVolumes/HeterogeneousVolumes.cpp
@@ -23,7 +23,7 @@ static TAutoConsoleVariable<int32> CVarHeterogeneousVolumesShadows(
 	TEXT("r.HeterogeneousVolumes.Shadows"),
 	0,
 	TEXT("Enables heterogeneous volume-casting shadows (default = 0)"),
-	ECVF_RenderThreadSafe | ECVF_ReadOnly
+	ECVF_RenderThreadSafe
 );
 
 static TAutoConsoleVariable<int32> CVarTranslucencyHeterogeneousVolumes(
@@ -81,10 +81,18 @@ static TAutoConsoleVariable<int32> CVarHeterogeneousVolumesFilterWidth(
 	ECVF_RenderThreadSafe
 );
 
-static TAutoConsoleVariable<int32> CVarHeterogeneousVolumesShadowMode(
-	TEXT("r.HeterogeneousVolumes.Shadows.Mode"),
+static TAutoConsoleVariable<int32> CVarHeterogeneousVolumesShadowType(
+	TEXT("r.HeterogeneousVolumes.Shadows.Type"),
+	1,
+	TEXT("0: Adaptive Volumetric Shadow Map\n")
+	TEXT("1: Beer Shadow Map (Default)"),
+	ECVF_RenderThreadSafe
+);
+
+static TAutoConsoleVariable<int32> CVarHeterogeneousVolumesShadowPipeline(
+	TEXT("r.HeterogeneousVolumes.Shadows.Pipeline"),
 	0,
-	TEXT("0: Live-Shading (Default)")
+	TEXT("0: Live-Shading (Default)\n")
 	TEXT("1: Preshaded Voxel Grid"),
 	ECVF_RenderThreadSafe
 );
@@ -349,6 +357,12 @@ bool ShouldCompositeHeterogeneousVolumesWithTranslucency()
 
 namespace HeterogeneousVolumes
 {
+	EShadowType GetShadowType()
+	{
+		int32 ShadowValue = FMath::Clamp(CVarHeterogeneousVolumesShadowType.GetValueOnAnyThread(), 0, 1);
+		return static_cast<EShadowType>(ShadowValue);
+	}
+
 	ETranslucencyCompositingMode GetTranslucencyCompositingMode()
 	{
 		int32 ClampedValue = FMath::Clamp(CVarTranslucencyHeterogeneousVolumes.GetValueOnRenderThread(), 0, 2);
@@ -621,9 +635,10 @@ namespace HeterogeneousVolumes
 		return CVarHeterogeneousVolumesDebug.GetValueOnRenderThread();
 	}
 
-	EShadowMode GetShadowMode()
+	EShadowPipeline GetShadowPipeline()
 	{
-		return static_cast<EShadowMode>(CVarHeterogeneousVolumesShadowMode.GetValueOnRenderThread());
+		int32 ShadowValue = FMath::Clamp(CVarHeterogeneousVolumesShadowPipeline.GetValueOnRenderThread(), 0, 1);
+		return static_cast<EShadowPipeline>(ShadowValue);
 	}
 
 	EStochasticFilteringMode GetStochasticFilteringMode()
@@ -894,7 +909,7 @@ bool ShouldBuildVoxelGrids(const FScene* Scene)
 		return true;
 	}
 
-	if (HeterogeneousVolumes::GetShadowMode() == HeterogeneousVolumes::EShadowMode::VoxelGrid)
+	if (HeterogeneousVolumes::GetShadowPipeline() == HeterogeneousVolumes::EShadowPipeline::VoxelGrid)
 	{
 		return true;
 	}
@@ -910,7 +925,7 @@ bool ShouldCacheVoxelGrids(const FScene* Scene, FSceneViewState* ViewState)
 		return false;
 	}
 
-	if (HeterogeneousVolumes::GetShadowMode() == HeterogeneousVolumes::EShadowMode::VoxelGrid)
+	if (HeterogeneousVolumes::GetShadowPipeline() == HeterogeneousVolumes::EShadowPipeline::VoxelGrid)
 	{
 		return true;
 	}
@@ -934,7 +949,7 @@ void RenderHeterogeneousVolumeShadows(
 	RDG_GPU_STAT_SCOPE(GraphBuilder, HeterogeneousVolumeShadowsStat);
 	SCOPED_NAMED_EVENT(HeterogeneousVolumes, FColor::Emerald);
 
-	if (HeterogeneousVolumes::GetShadowMode() == HeterogeneousVolumes::EShadowMode::LiveShading)
+	if (HeterogeneousVolumes::GetShadowPipeline() == HeterogeneousVolumes::EShadowPipeline::LiveShading)
 	{
 		RenderAdaptiveVolumetricShadowMapWithLiveShading(
 			GraphBuilder,
@@ -962,7 +977,7 @@ void FDeferredShadingSceneRenderer::RenderHeterogeneousVolumeShadows(
 
 	TRDGUniformBufferRef<FOrthoVoxelGridUniformBufferParameters> OrthoGridUniformBuffer = nullptr;
 	TRDGUniformBufferRef<FFrustumVoxelGridUniformBufferParameters> FrustumGridUniformBuffer = nullptr;
-	if (HeterogeneousVolumes::GetShadowMode() == HeterogeneousVolumes::EShadowMode::VoxelGrid)
+	if (HeterogeneousVolumes::GetShadowPipeline() == HeterogeneousVolumes::EShadowPipeline::VoxelGrid)
 	{
 		FVoxelGridBuildOptions BuildOptions;
 		BuildOptions.VoxelGridBuildMode = EVoxelGridBuildMode::Shadows;
@@ -983,7 +998,7 @@ void FDeferredShadingSceneRenderer::RenderHeterogeneousVolumeShadows(
 
 		if (HeterogeneousVolumes::GetTranslucencyCompositingMode() == HeterogeneousVolumes::ETranslucencyCompositingMode::VolumetricShadowMap)
 		{
-			if (HeterogeneousVolumes::GetShadowMode() == HeterogeneousVolumes::EShadowMode::LiveShading)
+			if (HeterogeneousVolumes::GetShadowPipeline() == HeterogeneousVolumes::EShadowPipeline::LiveShading)
 			{
 				RenderAdaptiveVolumetricCameraMapWithLiveShading(
 					GraphBuilder,
@@ -1009,7 +1024,7 @@ void FDeferredShadingSceneRenderer::RenderHeterogeneousVolumeShadows(
 
 		if (ShouldHeterogeneousVolumesCastShadows())
 		{
-			if (HeterogeneousVolumes::GetShadowMode() == HeterogeneousVolumes::EShadowMode::LiveShading)
+			if (HeterogeneousVolumes::GetShadowPipeline() == HeterogeneousVolumes::EShadowPipeline::LiveShading)
 			{
 				// This path is taken care of now in ShadowDepthRendering
 #if 0
diff --git a/Engine/Source/Runtime/Renderer/Private/HeterogeneousVolumes/HeterogeneousVolumes.h b/Engine/Source/Runtime/Renderer/Private/HeterogeneousVolumes/HeterogeneousVolumes.h
--- a/Engine/Source/Runtime/Renderer/Private/HeterogeneousVolumes/HeterogeneousVolumes.h
+++ b/Engine/Source/Runtime/Renderer/Private/HeterogeneousVolumes/HeterogeneousVolumes.h
@@ -87,12 +87,19 @@ namespace HeterogeneousVolumes
 	EScalabilityMode GetScalabilityMode();
 
 	// Shadow generation
-	enum class EShadowMode
+	enum class EShadowType : uint8
+	{
+		AdaptiveVolumetricShadowMap,
+		BeerShadowMap
+	};
+	EShadowType GetShadowType();
+
+	enum class EShadowPipeline
 	{
 		LiveShading,
 		VoxelGrid
 	};
-	EShadowMode GetShadowMode();
+	EShadowPipeline GetShadowPipeline();
 	FIntPoint GetShadowMapResolution();
 	float GetShadingRateForShadows();
 	float GetOutOfFrustumShadingRateForShadows();
@@ -520,6 +527,42 @@ BEGIN_UNIFORM_BUFFER_STRUCT(FAdaptiveVolumetricShadowMaps, )
 	SHADER_PARAMETER_STRUCT(FAdaptiveVolumetricShadowMapParameters, CameraAVSM)
 END_UNIFORM_BUFFER_STRUCT()
 
+BEGIN_UNIFORM_BUFFER_STRUCT(FBeerShadowMapUniformBufferParameters, )
+	SHADER_PARAMETER_ARRAY(FMatrix44f, TranslatedWorldToShadow, [6])
+	SHADER_PARAMETER_ARRAY(FVector4f, TranslatedWorldOrigin, [6])
+	SHADER_PARAMETER_ARRAY(FVector4f, TranslatedWorldPlane, [6])
+	SHADER_PARAMETER_ARRAY(FVector4f, SplitDepths, [6])
+
+	SHADER_PARAMETER(FIntPoint, Resolution)
+	SHADER_PARAMETER(int32, NumShadowMatrices)
+	SHADER_PARAMETER(int32, MaxSampleCount)
+	SHADER_PARAMETER(int32, bIsEmpty)
+	SHADER_PARAMETER(int32, bIsDirectionalLight)
+	SHADER_PARAMETER(float, DownsampleFactor)
+
+	SHADER_PARAMETER_RDG_TEXTURE_SRV(Texture2D, BeerShadowMapTexture)
+	SHADER_PARAMETER_RDG_TEXTURE_SRV(Texture2D, RadianceTexture)
+	SHADER_PARAMETER_SAMPLER(SamplerState, TextureSampler)
+END_UNIFORM_BUFFER_STRUCT()
+
+BEGIN_SHADER_PARAMETER_STRUCT(FBeerShadowMapParameters, RENDERER_API)
+	SHADER_PARAMETER_ARRAY(FMatrix44f, TranslatedWorldToShadow, [6])
+	SHADER_PARAMETER_ARRAY(FVector4f, TranslatedWorldOrigin, [6])
+	SHADER_PARAMETER_ARRAY(FVector4f, TranslatedWorldPlane, [6])
+	SHADER_PARAMETER_ARRAY(FVector4f, SplitDepths, [6])
+
+	SHADER_PARAMETER(FIntPoint, Resolution)
+	SHADER_PARAMETER(int32, NumShadowMatrices)
+	SHADER_PARAMETER(int32, MaxSampleCount)
+	SHADER_PARAMETER(int32, bIsEmpty)
+	SHADER_PARAMETER(int32, bIsDirectionalLight)
+	SHADER_PARAMETER(float, DownsampleFactor)
+
+	SHADER_PARAMETER_RDG_TEXTURE_SRV(Texture2D, BeerShadowMapTexture)
+	SHADER_PARAMETER_RDG_TEXTURE_SRV(Texture2D, RadianceTexture)
+	SHADER_PARAMETER_SAMPLER(SamplerState, TextureSampler)
+END_SHADER_PARAMETER_STRUCT()
+
 namespace HeterogeneousVolumes {
 	TRDGUniformBufferRef<FAdaptiveVolumetricShadowMapUniformBufferParameters> GetAdaptiveVolumetricShadowMapUniformBuffer(
 		FRDGBuilder& GraphBuilder,
@@ -545,6 +588,20 @@ namespace HeterogeneousVolumes {
 		FRDGBuilder& GraphBuilder
 	);
 
+	TRDGUniformBufferRef<FBeerShadowMapUniformBufferParameters> GetBeerShadowMapUniformBuffer(
+		FRDGBuilder& GraphBuilder,
+		FSceneViewState* ViewState,
+		const FLightSceneInfo* LightSceneInfo
+	);
+
+	void DestroyBeerShadowMapUniformBuffer(
+		TRDGUniformBufferRef<FBeerShadowMapUniformBufferParameters>& BeerShadowMapUniformBuffer
+	);
+
+	TRDGUniformBufferRef<FBeerShadowMapUniformBufferParameters> CreateEmptyBeerShadowMapUniformBuffer(
+		FRDGBuilder& GraphBuilder
+	);
+
 	TRDGUniformBufferRef<FFrustumVoxelGridUniformBufferParameters> GetFrustumVoxelGridUniformBuffer(
 		FRDGBuilder& GraphBuilder,
 		FSceneViewState* ViewState
@@ -643,6 +700,21 @@ void RenderAdaptiveVolumetricShadowMapWithLiveShading(
 	TArray<FVisibleLightInfo, SceneRenderingAllocator>& VisibleLightInfos
 );
 
+void RenderVolumetricShadowMaskForLight(
+	FRDGBuilder& GraphBuilder,
+	// Scene data
+	const FSceneTextures& SceneTextures,
+	FScene* Scene,
+	const FViewInfo& View,
+	// Light data
+	uint32 LightType,
+	const FLightSceneInfo* LightSceneInfo,
+	// Current shadow mask
+	FRDGTextureRef ShadowMaskTexture,
+	// Output
+	FRDGTextureRef& VolumetricShadowMaskTexture
+);
+
 void RenderAdaptiveVolumetricCameraMapWithVoxelGrid(
 	FRDGBuilder& GraphBuilder,
 	// Scene data
@@ -747,6 +819,38 @@ void CreateAdaptiveVolumetricShadowMapUniformBuffer(
 	TRDGUniformBufferRef<FAdaptiveVolumetricShadowMapUniformBufferParameters>& AdaptiveVolumetricShadowMapUniformBuffer
 );
 
+void CreateBeerShadowMapUniformBufferParameters(
+	FRDGBuilder& GraphBuilder,
+	const FVector3f* TranslatedWorldOrigin,
+	const FVector4f* TranslatedWorldPlane,
+	const FMatrix44f* TranslatedWorldToShadow,
+	const FVector4f* SplitDepths,
+	FIntPoint VolumetricShadowMapResolution,
+	float VolumetricShadowMapDownsampleFactor,
+	int32 NumShadowMatrices,
+	uint32 VolumetricShadowMapMaxSampleCount,
+	bool bIsDirectionalLight,
+	FRDGBufferRef VolumetricShadowMapLinkedListBuffer,
+	FRDGBufferRef VolumetricShadowMapIndirectionBuffer,
+	FRDGBufferRef VolumetricShadowMapSampleBuffer,
+	FBeerShadowMapUniformBufferParameters*& BeerShadowMapUniformBufferParameters
+);
+
+void CreateBeerShadowMapUniformBuffer(
+	FRDGBuilder& GraphBuilder,
+	const FVector3f* TranslatedWorldOrigin,
+	const FVector4f* TranslatedWorldPlane,
+	const FMatrix44f* TranslatedWorldToShadow,
+	const FVector4f* SplitDepths,
+	FIntPoint VolumetricShadowMapResolution,
+	float VolumetricShadowMapDownsampleFactor,
+	int32 NumShadowMatrices,
+	uint32 VolumetricShadowMapMaxSampleCount,
+	bool bIsDirectionalLight,
+	FRDGTextureRef BeerShadowMapTexture,
+	TRDGUniformBufferRef<FBeerShadowMapUniformBufferParameters>& BeerShadowMapUniformBuffer
+);
+
 void RenderSingleScatteringWithVoxelGrid(
 	FRDGBuilder& GraphBuilder,
 	// Scene data
diff --git a/Engine/Source/Runtime/Renderer/Private/HeterogeneousVolumes/HeterogeneousVolumesLiveShadingPipeline.cpp b/Engine/Source/Runtime/Renderer/Private/HeterogeneousVolumes/HeterogeneousVolumesLiveShadingPipeline.cpp
--- a/Engine/Source/Runtime/Renderer/Private/HeterogeneousVolumes/HeterogeneousVolumesLiveShadingPipeline.cpp
+++ b/Engine/Source/Runtime/Renderer/Private/HeterogeneousVolumes/HeterogeneousVolumesLiveShadingPipeline.cpp
@@ -316,6 +316,7 @@ class FRenderLightingCacheWithLiveShadingCS : public FMeshMaterialShader
 		SHADER_PARAMETER_STRUCT_INCLUDE(FVirtualShadowMapSamplingParameters, VirtualShadowMapSamplingParameters)
 		SHADER_PARAMETER(int32, VirtualShadowMapId)
 		SHADER_PARAMETER_RDG_UNIFORM_BUFFER(FAdaptiveVolumetricShadowMapUniformBufferParameters, AVSM)
+		SHADER_PARAMETER_RDG_UNIFORM_BUFFER(FBeerShadowMapUniformBufferParameters, BeerShadowMap)
 
 		// Global illumination data
 		SHADER_PARAMETER_RDG_UNIFORM_BUFFER(FLumenTranslucencyLightingUniforms, LumenGIVolumeStruct)
@@ -500,6 +501,7 @@ class FRenderSingleScatteringWithLiveShadingCS : public FMeshMaterialShader
 
 		SHADER_PARAMETER_RDG_UNIFORM_BUFFER(FAdaptiveVolumetricShadowMaps, AVSMs)
 		SHADER_PARAMETER_RDG_UNIFORM_BUFFER(FAdaptiveVolumetricShadowMapUniformBufferParameters, AVSM)
+		SHADER_PARAMETER_RDG_UNIFORM_BUFFER(FBeerShadowMapUniformBufferParameters, BeerShadowMap)
 
 		// Atmosphere
 		SHADER_PARAMETER_RDG_UNIFORM_BUFFER(FFogUniformParameters, FogStruct)
@@ -856,6 +858,7 @@ static void RenderLightingCacheWithLiveShading(
 		}
 		PassParameters->VirtualShadowMapSamplingParameters = VirtualShadowMapArray.GetSamplingParameters(GraphBuilder, ViewIndex);
 		PassParameters->AVSM = HeterogeneousVolumes::GetAdaptiveVolumetricShadowMapUniformBuffer(GraphBuilder, View.ViewState, LightSceneInfo);
+		PassParameters->BeerShadowMap = HeterogeneousVolumes::GetBeerShadowMapUniformBuffer(GraphBuilder, View.ViewState, LightSceneInfo);
 
 		// Global illumination data
 		auto* LumenUniforms = GraphBuilder.AllocParameters<FLumenTranslucencyLightingUniforms>();
@@ -1311,6 +1314,7 @@ void RenderSingleScatteringWithLiveShading(
 		{
 			PassParameters->AVSM = HeterogeneousVolumes::GetAdaptiveVolumetricShadowMapUniformBuffer(GraphBuilder, View.ViewState, LightSceneInfo);
 		}
+		PassParameters->BeerShadowMap = HeterogeneousVolumes::GetBeerShadowMapUniformBuffer(GraphBuilder, View.ViewState, LightSceneInfo);
 
 		TRDGUniformBufferRef<FFogUniformParameters> FogBuffer = CreateFogUniformBuffer(GraphBuilder, View);
 		PassParameters->FogStruct = FogBuffer;
@@ -1956,7 +1960,8 @@ class FRenderVolumetricShadowMapForLightWithLiveShadingCS : public FMeshMaterial
 	class FUseAVSMCompression : SHADER_PERMUTATION_BOOL("USE_AVSM_COMPRESSION");
 	class FUseCameraSceneDepth : SHADER_PERMUTATION_BOOL("USE_CAMERA_SCENE_DEPTH");
 	class FUseAnalyticDerivatives : SHADER_PERMUTATION_BOOL("USE_ANALYTIC_DERIVATIVES");
-	using FPermutationDomain = TShaderPermutationDomain<FUseAVSMCompression, FUseCameraSceneDepth, FUseAnalyticDerivatives>;
+	class FShadowMapType : SHADER_PERMUTATION_INT("SHADOW_MAP_TYPE", 2);
+	using FPermutationDomain = TShaderPermutationDomain<FUseAVSMCompression, FUseCameraSceneDepth, FUseAnalyticDerivatives, FShadowMapType>;
 
 
 	BEGIN_SHADER_PARAMETER_STRUCT(FParameters, )
@@ -2001,6 +2006,7 @@ class FRenderVolumetricShadowMapForLightWithLiveShadingCS : public FMeshMaterial
 		// Output
 		SHADER_PARAMETER_RDG_BUFFER_UAV(RWBuffer<int>, RWVolumetricShadowLinkedListAllocatorBuffer)
 		SHADER_PARAMETER_RDG_BUFFER_UAV(RWStructuredBuffer<int2>, RWVolumetricShadowLinkedListBuffer)
+		SHADER_PARAMETER_RDG_TEXTURE_UAV(RWTexture2D<float>, RWBeerShadowMapTexture)
 
 		// Debug
 		//SHADER_PARAMETER_RDG_BUFFER_UAV(RWStructuredBuffer<FVolumetricShadowMapDebugData>, RWDebugBuffer)
@@ -2268,6 +2274,7 @@ bool RenderVolumetricShadowMapForLightForHeterogeneousVolumeWithLiveShading(
 		// Output
 		PassParameters->RWVolumetricShadowLinkedListAllocatorBuffer = GraphBuilder.CreateUAV(VolumetricShadowLinkedListAllocatorBuffer, PF_R32_UINT);
 		PassParameters->RWVolumetricShadowLinkedListBuffer = GraphBuilder.CreateUAV(VolumetricShadowLinkedListBuffer);
+		PassParameters->RWBeerShadowMapTexture = GraphBuilder.CreateUAV(BeerShadowMapTexture);
 		//PassParameters->RWDebugBuffer = GraphBuilder.CreateUAV(DebugBuffer);
 	}
 
@@ -2284,10 +2291,13 @@ bool RenderVolumetricShadowMapForLightForHeterogeneousVolumeWithLiveShading(
 	}
 #endif // WANTS_DRAW_MESH_EVENTS
 
+	const int32 ShadowmapType = static_cast<int32>(HeterogeneousVolumes::GetShadowType());
+
 	FRenderVolumetricShadowMapForLightWithLiveShadingCS::FPermutationDomain PermutationVector;
 	PermutationVector.Set<FRenderVolumetricShadowMapForLightWithLiveShadingCS::FUseAVSMCompression>(HeterogeneousVolumes::UseAVSMCompression());
 	PermutationVector.Set<FRenderVolumetricShadowMapForLightWithLiveShadingCS::FUseCameraSceneDepth>(bUseCameraSceneDepth && HeterogeneousVolumes::ShadowsUseCameraSceneDepth());
 	PermutationVector.Set<FRenderVolumetricShadowMapForLightWithLiveShadingCS::FUseAnalyticDerivatives>(HeterogeneousVolumes::UseAnalyticDerivatives());
+	PermutationVector.Set<FRenderVolumetricShadowMapForLightWithLiveShadingCS::FShadowMapType>(ShadowmapType);
 	TShaderRef<FRenderVolumetricShadowMapForLightWithLiveShadingCS> ComputeShader = Material.GetShader<FRenderVolumetricShadowMapForLightWithLiveShadingCS>(&FLocalVertexFactory::StaticType, PermutationVector, false);
 	if (!ComputeShader.IsNull())
 	{
@@ -2806,6 +2816,13 @@ bool RenderVolumetricShadowMapForLightWithLiveShading(
 		);
 		//ClearVolumetricShadowMap(GraphBuilder, Scene, View, ShadowMapResolution, GroupCount, VolumetricShadowLinkedListBuffer);
 
+		FRDGTextureDesc Desc = SceneTextures.Color.Target->Desc;
+		Desc.Format = PF_FloatRGBA;
+		Desc.Flags &= ~(TexCreate_FastVRAM);
+		Desc.Extent = FIntPoint(ShadowMapResolution.X * GroupCount.Z, ShadowMapResolution.Y);
+		BeerShadowMapTexture = GraphBuilder.CreateTexture(Desc, TEXT("HeterogeneousVolumes.Light.BeerShadowMapTexture"));
+		AddClearUAVPass(GraphBuilder, GraphBuilder.CreateUAV(BeerShadowMapTexture), FLinearColor::Transparent);
+
 		bool bUseCameraSceneDepth = false;
 		RenderVolumetricShadowMapForLightForHeterogeneousVolumeWithLiveShading(
 			GraphBuilder,
@@ -2871,33 +2888,37 @@ bool RenderVolumetricShadowMapForLightWithLiveShading(
 				VolumetricShadowLinkedListBuffer1
 			);
 
-			FRDGBufferRef VolumetricShadowLinkedListBuffer2 = GraphBuilder.CreateBuffer(
-				FRDGBufferDesc::CreateStructuredDesc(sizeof(FAVSMLinkedListPackedData), VolumetricShadowLinkedListElementCount),
-				TEXT("HeterogeneousVolume.VolumetricShadowLinkedListBuffer2")
-			);
+			if (HeterogeneousVolumes::GetShadowType() == HeterogeneousVolumes::EShadowType::AdaptiveVolumetricShadowMap)
+			{
+				FRDGBufferRef VolumetricShadowLinkedListBuffer2 = GraphBuilder.CreateBuffer(
+					FRDGBufferDesc::CreateStructuredDesc(sizeof(FAVSMLinkedListPackedData), VolumetricShadowLinkedListElementCount),
+					TEXT("HeterogeneousVolume.VolumetricShadowLinkedListBuffer2")
+				);
 
-			CombineVolumetricShadowMap(
-				GraphBuilder,
-				View,
-				GroupCount,
-				LightType,
-				ShadowMapResolution,
-				MaxSampleCount,
-				VolumetricShadowLinkedListBuffer,
-				VolumetricShadowLinkedListBuffer1,
-				VolumetricShadowLinkedListBuffer2
-			);
+				CombineVolumetricShadowMap(
+					GraphBuilder,
+					View,
+					GroupCount,
+					LightType,
+					ShadowMapResolution,
+					MaxSampleCount,
+					VolumetricShadowLinkedListBuffer,
+					VolumetricShadowLinkedListBuffer1,
+					VolumetricShadowLinkedListBuffer2
+				);
 
-			VolumetricShadowLinkedListBuffer = VolumetricShadowLinkedListBuffer2;
+				VolumetricShadowLinkedListBuffer = VolumetricShadowLinkedListBuffer2;
+			}
 		}
 
-		if (HeterogeneousVolumes::GetDebugMode() != 0)
+		if (HeterogeneousVolumes::GetDebugMode() &&
+			HeterogeneousVolumes::GetShadowType() == HeterogeneousVolumes::EShadowType::AdaptiveVolumetricShadowMap)
 		{
-			// Create debug beer shadow map texture
 			FString LightName;
 			FSceneRenderer::GetLightNameForDrawEvent(LightSceneInfo->Proxy, LightName);
 			RDG_EVENT_SCOPE(GraphBuilder, "%s", *LightName);
 
+			// Create debug beer shadow map texture
 			ConvertVolumetricShadowMapToBeerLawShadowMap(
 				GraphBuilder,
 				SceneTextures,
@@ -2966,6 +2987,7 @@ void RenderAdaptiveVolumetricShadowMapWithLiveShading(
 			}
 
 			TRDGUniformBufferRef<FAdaptiveVolumetricShadowMapUniformBufferParameters> AdaptiveVolumetricShadowMapUniformBuffer;
+			TRDGUniformBufferRef<FBeerShadowMapUniformBufferParameters> BeerShadowMapUniformBuffer;
 			bool bCreateShadowMap = bShouldRenderShadowMaps && bCastsVolumetricShadow && bDynamicallyShadowed && !ShouldRenderRayTracingShadowsForLight(*View.Family, LightSceneInfoCompact[PassIndex]);
 			if (bCreateShadowMap)
 			{
@@ -3048,7 +3070,7 @@ void RenderAdaptiveVolumetricShadowMapWithLiveShading(
 						TranslatedWorldPlane,
 						TranslatedWorldToShadow,
 						SplitDepths,
-						VolumetricShadowMapResolution,
+						HeterogeneousVolumes::GetShadowType() == HeterogeneousVolumes::EShadowType::AdaptiveVolumetricShadowMap ? VolumetricShadowMapResolution : FIntPoint::ZeroValue,
 						DownsampleFactor,
 						NumShadowMatrices,
 						VolumetricShadowMapMaxSampleCount,
@@ -3058,21 +3080,41 @@ void RenderAdaptiveVolumetricShadowMapWithLiveShading(
 						VolumetricShadowMapSampleBuffer,
 						AdaptiveVolumetricShadowMapUniformBuffer
 					);
+
+					CreateBeerShadowMapUniformBuffer(
+						GraphBuilder,
+						TranslatedWorldOrigin,
+						TranslatedWorldPlane,
+						TranslatedWorldToShadow,
+						SplitDepths,
+						HeterogeneousVolumes::GetShadowType() == HeterogeneousVolumes::EShadowType::BeerShadowMap ? VolumetricShadowMapResolution : FIntPoint::ZeroValue,
+						DownsampleFactor,
+						NumShadowMatrices,
+						VolumetricShadowMapMaxSampleCount,
+						bIsDirectionalLight,
+						BeerShadowMapTexture,
+						BeerShadowMapUniformBuffer
+					);
 				}
 				else
 				{
 					AdaptiveVolumetricShadowMapUniformBuffer = HeterogeneousVolumes::CreateEmptyAdaptiveVolumetricShadowMapUniformBuffer(GraphBuilder);
+					BeerShadowMapUniformBuffer = HeterogeneousVolumes::CreateEmptyBeerShadowMapUniformBuffer(GraphBuilder);
 				}
 			}
 			else
 			{
 				AdaptiveVolumetricShadowMapUniformBuffer = HeterogeneousVolumes::CreateEmptyAdaptiveVolumetricShadowMapUniformBuffer(GraphBuilder);
+				BeerShadowMapUniformBuffer = HeterogeneousVolumes::CreateEmptyBeerShadowMapUniformBuffer(GraphBuilder);
 			}
 
 			if (View.ViewState)
 			{
 				TRDGUniformBufferRef<FAdaptiveVolumetricShadowMapUniformBufferParameters>& AdaptiveVolumetricShadowMap = View.ViewState->AdaptiveVolumetricShadowMapUniformBufferMap.FindOrAdd(LightSceneInfo->Id);
 				AdaptiveVolumetricShadowMap = AdaptiveVolumetricShadowMapUniformBuffer;
+
+				TRDGUniformBufferRef<FBeerShadowMapUniformBufferParameters>& BeerShadowMap = View.ViewState->BeerShadowMapUniformBufferMap.FindOrAdd(LightSceneInfo->Id);
+				BeerShadowMap = BeerShadowMapUniformBuffer;
 			}
 
 			if (HeterogeneousVolumes::GetDebugMode() != 0)
@@ -3313,3 +3355,152 @@ void RenderAdaptiveVolumetricCameraMapWithLiveShading(
 		);
 	}
 }
+
+class FRenderVolumetricShadowMaskCS : public FGlobalShader
+{
+	DECLARE_GLOBAL_SHADER(FRenderVolumetricShadowMaskCS);
+	SHADER_USE_PARAMETER_STRUCT(FRenderVolumetricShadowMaskCS, FGlobalShader);
+
+	class FShadowType : SHADER_PERMUTATION_INT("SHADOW_TYPE", 2);
+	using FPermutationDomain = TShaderPermutationDomain<FShadowType>;
+
+	BEGIN_SHADER_PARAMETER_STRUCT(FParameters, )
+		// Scene data
+		SHADER_PARAMETER_STRUCT_REF(FViewUniformShaderParameters, View)
+		SHADER_PARAMETER_STRUCT_INCLUDE(FSceneTextureParameters, SceneTextures)
+		SHADER_PARAMETER_RDG_UNIFORM_BUFFER(FSceneUniformParameters, Scene)
+
+		// Light data
+		SHADER_PARAMETER(int, LightType)
+		SHADER_PARAMETER_STRUCT_REF(FDeferredLightUniformStruct, DeferredLight)
+
+		// Shadow data
+		SHADER_PARAMETER_RDG_UNIFORM_BUFFER(FAdaptiveVolumetricShadowMapUniformBufferParameters, AVSM)
+		SHADER_PARAMETER_RDG_UNIFORM_BUFFER(FBeerShadowMapUniformBufferParameters, BeerShadowMap)
+
+		// Dispatch data
+		SHADER_PARAMETER(FIntVector, GroupCount)
+		SHADER_PARAMETER(int32, DownsampleFactor)
+
+		// Shadow mask
+		SHADER_PARAMETER_RDG_TEXTURE_SRV(Texture2D<float>, ShadowMaskTexture)
+
+		// Output
+		SHADER_PARAMETER_RDG_TEXTURE_UAV(RWTexture2D<float4>, RWVolumetricShadowMaskTexture)
+	END_SHADER_PARAMETER_STRUCT()
+
+	static bool ShouldCompilePermutation(
+		const FGlobalShaderPermutationParameters& Parameters
+	)
+	{
+		FPermutationDomain PermutationVector(Parameters.PermutationId);
+		if (RemapPermutation(PermutationVector) != PermutationVector)
+		{
+			return false;
+		}
+
+		return DoesPlatformSupportHeterogeneousVolumes(Parameters.Platform);
+	}
+
+	static FPermutationDomain RemapPermutation(FPermutationDomain PermutationVector)
+	{
+		return PermutationVector;
+	}
+
+	static EShaderPermutationPrecacheRequest ShouldPrecachePermutation(const FGlobalShaderPermutationParameters& Parameters)
+	{
+		return EShaderPermutationPrecacheRequest::Precached;
+	}
+
+	static void ModifyCompilationEnvironment(
+		const FGlobalShaderPermutationParameters& Parameters,
+		FShaderCompilerEnvironment& OutEnvironment
+	)
+	{
+		FMaterialShader::ModifyCompilationEnvironment(Parameters, OutEnvironment);
+		OutEnvironment.SetDefine(TEXT("THREADGROUP_SIZE_1D"), GetThreadGroupSize1D());
+		OutEnvironment.SetDefine(TEXT("THREADGROUP_SIZE_2D"), GetThreadGroupSize2D());
+
+		bool bSupportVirtualShadowMap = IsFeatureLevelSupported(Parameters.Platform, ERHIFeatureLevel::SM5);
+		if (bSupportVirtualShadowMap)
+		{
+			OutEnvironment.SetDefine(TEXT("VIRTUAL_SHADOW_MAP"), 1);
+			FVirtualShadowMapArray::SetShaderDefines(OutEnvironment);
+		}
+
+		// This shader takes a very long time to compile with FXC, so we pre-compile it with DXC first and then forward the optimized HLSL to FXC.
+		OutEnvironment.CompilerFlags.Add(CFLAG_PrecompileWithDXC);
+		OutEnvironment.CompilerFlags.Add(CFLAG_AllowTypedUAVLoads);
+
+		OutEnvironment.SetDefine(TEXT("GET_PRIMITIVE_DATA_OVERRIDE"), 1);
+	}
+
+	static int32 GetThreadGroupSize1D() { return GetThreadGroupSize2D() * GetThreadGroupSize2D(); }
+	static int32 GetThreadGroupSize2D() { return 8; }
+};
+
+IMPLEMENT_GLOBAL_SHADER(FRenderVolumetricShadowMaskCS, "/Engine/Private/HeterogeneousVolumes/HeterogeneousVolumesShadowMask.usf", "RenderVolumetricShadowMaskCS", SF_Compute);
+
+void RenderVolumetricShadowMaskForLight(
+	FRDGBuilder& GraphBuilder,
+	// Scene data
+	const FSceneTextures& SceneTextures,
+	//const FMinimalSceneTextures& SceneTextures,
+	FScene* Scene,
+	const FViewInfo& View,
+	// Light data
+	uint32 LightType,
+	const FLightSceneInfo* LightSceneInfo,
+	// Current shadow mask
+	FRDGTextureRef ShadowMaskTexture,
+	// Output
+	FRDGTextureRef& VolumetricShadowMaskTexture
+)
+{
+	FIntVector GroupCount = FComputeShaderUtils::GetGroupCount(View.ViewRect.Size(), FRenderVolumetricShadowMaskCS::GetThreadGroupSize2D());
+
+	// Note must be done in the same scope as we add the pass otherwise the UB lifetime will not be guaranteed
+	FDeferredLightUniformStruct DeferredLightUniform = GetDeferredLightParameters(View, *LightSceneInfo);
+	TUniformBufferRef<FDeferredLightUniformStruct> DeferredLightUB = CreateUniformBufferImmediate(DeferredLightUniform, UniformBuffer_SingleDraw);
+
+	FRenderVolumetricShadowMaskCS::FParameters* PassParameters = GraphBuilder.AllocParameters<FRenderVolumetricShadowMaskCS::FParameters>();
+	{
+		// Scene data
+		PassParameters->View = View.ViewUniformBuffer;
+		PassParameters->SceneTextures = GetSceneTextureParameters(GraphBuilder, SceneTextures);
+		PassParameters->Scene = View.GetSceneUniforms().GetBuffer(GraphBuilder);
+
+		// Light data
+		PassParameters->LightType = LightType;
+		PassParameters->DeferredLight = DeferredLightUB;
+
+		// Shadow data
+		PassParameters->AVSM = HeterogeneousVolumes::GetAdaptiveVolumetricShadowMapUniformBuffer(GraphBuilder, View.ViewState, LightSceneInfo);
+		PassParameters->BeerShadowMap = HeterogeneousVolumes::GetBeerShadowMapUniformBuffer(GraphBuilder, View.ViewState, LightSceneInfo);
+
+		// Dispatch data
+		PassParameters->GroupCount = GroupCount;
+		PassParameters->DownsampleFactor = 1; // Disallow downsampling
+
+		// Shadow Mask
+		PassParameters->ShadowMaskTexture = GraphBuilder.CreateSRV(ShadowMaskTexture);
+
+		// Output
+		PassParameters->RWVolumetricShadowMaskTexture = GraphBuilder.CreateUAV(VolumetricShadowMaskTexture);
+	}
+
+	FRenderVolumetricShadowMaskCS::FPermutationDomain PermutationVector;
+	PermutationVector.Set<FRenderVolumetricShadowMaskCS::FShadowType>(static_cast<int32>(HeterogeneousVolumes::GetShadowType()));
+	PermutationVector = FRenderVolumetricShadowMaskCS::RemapPermutation(PermutationVector);
+	TShaderRef<FRenderVolumetricShadowMaskCS> ComputeShader = View.ShaderMap->GetShader<FRenderVolumetricShadowMaskCS>(PermutationVector);
+
+
+
+	FComputeShaderUtils::AddPass(
+		GraphBuilder,
+		//RDG_EVENT_NAME("%s", *PassName),
+		RDG_EVENT_NAME("RenderVolumetricShadowMask"),
+		ComputeShader,
+		PassParameters,
+		GroupCount);
+}
diff --git a/Engine/Source/Runtime/Renderer/Private/HeterogeneousVolumes/HeterogeneousVolumesPreshadingPipeline.cpp b/Engine/Source/Runtime/Renderer/Private/HeterogeneousVolumes/HeterogeneousVolumesPreshadingPipeline.cpp
--- a/Engine/Source/Runtime/Renderer/Private/HeterogeneousVolumes/HeterogeneousVolumesPreshadingPipeline.cpp
+++ b/Engine/Source/Runtime/Renderer/Private/HeterogeneousVolumes/HeterogeneousVolumesPreshadingPipeline.cpp
@@ -232,6 +232,7 @@ class FRenderLightingCacheWithPreshadingCS : public FGlobalShader
 		SHADER_PARAMETER_STRUCT_INCLUDE(FVirtualShadowMapSamplingParameters, VirtualShadowMapSamplingParameters)
 		SHADER_PARAMETER(int32, VirtualShadowMapId)
 		SHADER_PARAMETER_RDG_UNIFORM_BUFFER(FAdaptiveVolumetricShadowMapUniformBufferParameters, AVSM)
+		SHADER_PARAMETER_RDG_UNIFORM_BUFFER(FBeerShadowMapUniformBufferParameters, BeerShadowMap)
 
 		// Global illumination data
 		SHADER_PARAMETER_RDG_UNIFORM_BUFFER(FLumenTranslucencyLightingUniforms, LumenGIVolumeStruct)
@@ -343,6 +344,7 @@ class FRenderSingleScatteringWithPreshadingCS : public FGlobalShader
 		SHADER_PARAMETER_STRUCT_INCLUDE(FVirtualShadowMapSamplingParameters, VirtualShadowMapSamplingParameters)
 		SHADER_PARAMETER(int32, VirtualShadowMapId)
 		SHADER_PARAMETER_RDG_UNIFORM_BUFFER(FAdaptiveVolumetricShadowMapUniformBufferParameters, AVSM)
+		SHADER_PARAMETER_RDG_UNIFORM_BUFFER(FBeerShadowMapUniformBufferParameters, BeerShadowMap)
 
 		// Atmosphere
 		SHADER_PARAMETER_RDG_UNIFORM_BUFFER(FFogUniformParameters, FogStruct)
@@ -540,6 +542,7 @@ void RenderLightingCacheWithPreshadingCompute(
 		}
 		PassParameters->VirtualShadowMapSamplingParameters = VirtualShadowMapArray.GetSamplingParameters(GraphBuilder, ViewIndex);
 		PassParameters->AVSM = HeterogeneousVolumes::GetAdaptiveVolumetricShadowMapUniformBuffer(GraphBuilder, View.ViewState, LightSceneInfo);
+		PassParameters->BeerShadowMap = HeterogeneousVolumes::GetBeerShadowMapUniformBuffer(GraphBuilder, View.ViewState, LightSceneInfo);
 
 		// Global illumination data
 		auto* LumenUniforms = GraphBuilder.AllocParameters<FLumenTranslucencyLightingUniforms>();
@@ -673,6 +676,7 @@ void RenderSingleScatteringWithPreshadingCompute(
 		}
 		PassParameters->VirtualShadowMapSamplingParameters = VirtualShadowMapArray.GetSamplingParameters(GraphBuilder, ViewIndex);
 		PassParameters->AVSM = HeterogeneousVolumes::GetAdaptiveVolumetricShadowMapUniformBuffer(GraphBuilder, View.ViewState, LightSceneInfo);
+		PassParameters->BeerShadowMap = HeterogeneousVolumes::GetBeerShadowMapUniformBuffer(GraphBuilder, View.ViewState, LightSceneInfo);
 
 		TRDGUniformBufferRef<FFogUniformParameters> FogBuffer = CreateFogUniformBuffer(GraphBuilder, View);
 		PassParameters->FogStruct = FogBuffer;
diff --git a/Engine/Source/Runtime/Renderer/Private/HeterogeneousVolumes/HeterogeneousVolumesVoxelGridPipeline.cpp b/Engine/Source/Runtime/Renderer/Private/HeterogeneousVolumes/HeterogeneousVolumesVoxelGridPipeline.cpp
--- a/Engine/Source/Runtime/Renderer/Private/HeterogeneousVolumes/HeterogeneousVolumesVoxelGridPipeline.cpp
+++ b/Engine/Source/Runtime/Renderer/Private/HeterogeneousVolumes/HeterogeneousVolumesVoxelGridPipeline.cpp
@@ -322,6 +322,7 @@ static TAutoConsoleVariable<bool> CVarHeterogeneousShadowsForLightTypeRect(
 IMPLEMENT_UNIFORM_BUFFER_STRUCT(FOrthoVoxelGridUniformBufferParameters, "OrthoGridUniformBuffer");
 IMPLEMENT_UNIFORM_BUFFER_STRUCT(FFrustumVoxelGridUniformBufferParameters, "FrustumGridUniformBuffer");
 IMPLEMENT_UNIFORM_BUFFER_STRUCT(FAdaptiveVolumetricShadowMapUniformBufferParameters, "AVSM");
+IMPLEMENT_UNIFORM_BUFFER_STRUCT(FBeerShadowMapUniformBufferParameters, "BeerShadowMap");
 
 struct FRasterTileData
 {
@@ -2790,7 +2791,7 @@ void CollectHeterogeneousVolumeMeshBatches(
 	}
 
 	// TODO: Temporarily disable per-light primitive gathering because it causes a crash with VoxelGrid builder
-	if (HeterogeneousVolumes::GetShadowMode() == HeterogeneousVolumes::EShadowMode::VoxelGrid)
+	if (HeterogeneousVolumes::GetShadowPipeline() == HeterogeneousVolumes::EShadowPipeline::VoxelGrid)
 	{
 		return;
 	}
@@ -3477,7 +3478,6 @@ void CreateAdaptiveVolumetricShadowMapUniformBufferParameters(
 		Parameters->IndirectionBuffer = GraphBuilder.CreateSRV(VolumetricShadowMapIndirectionBuffer);
 		Parameters->SampleBuffer = GraphBuilder.CreateSRV(VolumetricShadowMapSampleBuffer);
 		Parameters->RadianceTexture = GraphBuilder.CreateSRV(GSystemTextures.GetBlackDummy(GraphBuilder));
-		// Test
 		Parameters->TextureSampler = TStaticSamplerState<SF_Bilinear, AM_Clamp, AM_Clamp, AM_Clamp>::GetRHI();
 	}
 }
@@ -3519,6 +3519,76 @@ void CreateAdaptiveVolumetricShadowMapUniformBuffer(
 	AdaptiveVolumetricShadowMapUniformBuffer = GraphBuilder.CreateUniformBuffer(UniformBufferParameters);
 }
 
+void CreateBeerShadowMapUniformBufferParameters(
+	FRDGBuilder& GraphBuilder,
+	const FVector3f* TranslatedWorldOrigin,
+	const FVector4f* TranslatedWorldPlane,
+	const FMatrix44f* TranslatedWorldToShadow,
+	const FVector4f* SplitDepths,
+	FIntPoint VolumetricShadowMapResolution,
+	float VolumetricShadowMapDownsampleFactor,
+	int32 NumShadowMatrices,
+	uint32 VolumetricShadowMapMaxSampleCount,
+	bool bIsDirectionalLight,
+	FRDGTextureRef BeerShadowMapTexture,
+	FBeerShadowMapUniformBufferParameters*& Parameters
+)
+{
+	Parameters = GraphBuilder.AllocParameters<FBeerShadowMapUniformBufferParameters>();
+	{
+		Parameters->NumShadowMatrices = NumShadowMatrices;
+		for (int32 i = 0; i < NumShadowMatrices; ++i)
+		{
+			Parameters->TranslatedWorldToShadow[i] = TranslatedWorldToShadow[i];
+			Parameters->TranslatedWorldOrigin[i] = TranslatedWorldOrigin[i];
+			Parameters->TranslatedWorldPlane[i] = TranslatedWorldPlane[i];
+			Parameters->SplitDepths[i] = SplitDepths[i];
+		}
+
+		Parameters->Resolution = VolumetricShadowMapResolution;
+		Parameters->DownsampleFactor = VolumetricShadowMapDownsampleFactor;
+		Parameters->MaxSampleCount = VolumetricShadowMapMaxSampleCount;
+		Parameters->bIsEmpty = (VolumetricShadowMapResolution == FIntPoint::ZeroValue);
+		Parameters->bIsDirectionalLight = bIsDirectionalLight;
+		Parameters->BeerShadowMapTexture = GraphBuilder.CreateSRV(BeerShadowMapTexture);
+		Parameters->RadianceTexture = GraphBuilder.CreateSRV(GSystemTextures.GetBlackDummy(GraphBuilder));
+		Parameters->TextureSampler = TStaticSamplerState<SF_Bilinear, AM_Clamp, AM_Clamp, AM_Clamp>::GetRHI();
+	}
+}
+
+void CreateBeerShadowMapUniformBuffer(
+	FRDGBuilder& GraphBuilder,
+	const FVector3f* TranslatedWorldOrigin,
+	const FVector4f* TranslatedWorldPlane,
+	const FMatrix44f* TranslatedWorldToShadow,
+	const FVector4f* SplitDepths,
+	FIntPoint VolumetricShadowMapResolution,
+	float VolumetricShadowMapDownsampleFactor,
+	int32 NumShadowMatrices,
+	uint32 VolumetricShadowMapMaxSampleCount,
+	bool bIsDirectionalLight,
+	FRDGTextureRef BeerShadowMapTexture,
+	TRDGUniformBufferRef<FBeerShadowMapUniformBufferParameters>& BeerShadowMapUniformBuffer
+)
+{
+	FBeerShadowMapUniformBufferParameters* UniformBufferParameters;
+	CreateBeerShadowMapUniformBufferParameters(
+		GraphBuilder,
+		TranslatedWorldOrigin,
+		TranslatedWorldPlane,
+		TranslatedWorldToShadow,
+		SplitDepths,
+		VolumetricShadowMapResolution,
+		VolumetricShadowMapDownsampleFactor,
+		NumShadowMatrices,
+		VolumetricShadowMapMaxSampleCount,
+		bIsDirectionalLight,
+		BeerShadowMapTexture,
+		UniformBufferParameters
+	);
+	BeerShadowMapUniformBuffer = GraphBuilder.CreateUniformBuffer(UniformBufferParameters);
+}
+
 namespace HeterogeneousVolumes
 {
 	TRDGUniformBufferRef<FAdaptiveVolumetricShadowMapUniformBufferParameters> CreateEmptyAdaptiveVolumetricShadowMapUniformBuffer(
@@ -3651,6 +3721,61 @@ namespace HeterogeneousVolumes
 		return Parameters;
 	}
 
+	// Beer Shadow Map Management Functions
+	TRDGUniformBufferRef<FBeerShadowMapUniformBufferParameters> CreateEmptyBeerShadowMapUniformBuffer(
+		FRDGBuilder& GraphBuilder
+	)
+	{
+		FBeerShadowMapUniformBufferParameters* UniformBufferParameters = GraphBuilder.AllocParameters<FBeerShadowMapUniformBufferParameters>();
+		{
+			UniformBufferParameters->NumShadowMatrices = 1;
+			for (int32 i = 0; i < 6; ++i)
+			{
+				UniformBufferParameters->TranslatedWorldToShadow[i] = FMatrix44f::Identity;
+				UniformBufferParameters->TranslatedWorldOrigin[i] = FVector3f::ZeroVector;
+				UniformBufferParameters->TranslatedWorldPlane[i] = FVector4f::Zero();
+				UniformBufferParameters->SplitDepths[i] = FVector4f::Zero();
+			}
+
+			UniformBufferParameters->Resolution = FIntPoint::ZeroValue;
+			UniformBufferParameters->DownsampleFactor = 1.0f;
+			UniformBufferParameters->MaxSampleCount = 0;
+			UniformBufferParameters->bIsEmpty = true;
+			UniformBufferParameters->bIsDirectionalLight = false;
+			UniformBufferParameters->BeerShadowMapTexture = GraphBuilder.CreateSRV(GSystemTextures.GetBlackDummy(GraphBuilder));
+			UniformBufferParameters->RadianceTexture = GraphBuilder.CreateSRV(GSystemTextures.GetBlackDummy(GraphBuilder));
+			UniformBufferParameters->TextureSampler = TStaticSamplerState<SF_Bilinear>::GetRHI();
+		}
+
+		return GraphBuilder.CreateUniformBuffer(UniformBufferParameters);
+	}
+
+	void DestroyBeerShadowMapUniformBuffer(
+		TRDGUniformBufferRef<FBeerShadowMapUniformBufferParameters>& BeerShadowMapUniformBuffer
+	)
+	{
+		BeerShadowMapUniformBuffer = nullptr;
+	}
+
+	TRDGUniformBufferRef<FBeerShadowMapUniformBufferParameters> GetBeerShadowMapUniformBuffer(
+		FRDGBuilder& GraphBuilder,
+		FSceneViewState* ViewState,
+		const FLightSceneInfo* LightSceneInfo
+	)
+	{
+		bool bIsEnabled = HeterogeneousVolumes::GetShadowType() == HeterogeneousVolumes::EShadowType::BeerShadowMap;
+		if (ViewState && LightSceneInfo && bIsEnabled)
+		{
+			TRDGUniformBufferRef<FBeerShadowMapUniformBufferParameters>* BeerShadowMapUniformBuffer = ViewState->BeerShadowMapUniformBufferMap.Find(LightSceneInfo->Id);
+			if (BeerShadowMapUniformBuffer != nullptr)
+			{
+				return *BeerShadowMapUniformBuffer;
+			}
+		}
+
+		return CreateEmptyBeerShadowMapUniformBuffer(GraphBuilder);
+	}
+
 } // namespace HeterogeneousVolumes
 
 void CompressVolumetricShadowMap(
diff --git a/Engine/Source/Runtime/Renderer/Private/LightRendering.cpp b/Engine/Source/Runtime/Renderer/Private/LightRendering.cpp
--- a/Engine/Source/Runtime/Renderer/Private/LightRendering.cpp
+++ b/Engine/Source/Runtime/Renderer/Private/LightRendering.cpp
@@ -212,6 +212,8 @@ static FAutoConsoleVariableRef CVarDebugLightDiscardProp(
 );
 #endif // ENABLE_DEBUG_DISCARD_PROP
 
+DECLARE_GPU_STAT_NAMED(VolumetricShadowMask, TEXT("Volumetric Shadow Mask"));
+
 #if RHI_RAYTRACING
 
 static bool ShouldRenderRayTracingShadowsForLightType(ELightComponentType LightType)
@@ -957,10 +959,6 @@ class FDeferredLightPS : public FGlobalShader
 		SHADER_PARAMETER(int32, VirtualShadowMapId)
 		SHADER_PARAMETER(int32, LightSceneId)
 		SHADER_PARAMETER_RDG_TEXTURE(Texture2D, ShadowMaskBits)
-		// Heterogeneous Volume data
-		//SHADER_PARAMETER_RDG_UNIFORM_BUFFER(FOrthoVoxelGridUniformBufferParameters, OrthoGridUniformBuffer)
-		//SHADER_PARAMETER_RDG_UNIFORM_BUFFER(FFrustumVoxelGridUniformBufferParameters, FrustumGridUniformBuffer)
-		SHADER_PARAMETER_RDG_UNIFORM_BUFFER(FAdaptiveVolumetricShadowMapUniformBufferParameters, AVSM)
 		RENDER_TARGET_BINDING_SLOTS()
 	END_SHADER_PARAMETER_STRUCT()
 
@@ -1121,8 +1119,6 @@ class FDeferredLightPS : public FGlobalShader
 			FForwardLightingParameters::ModifyCompilationEnvironment(Parameters.Platform, OutEnvironment);
 		}
 
-		OutEnvironment.SetDefine(TEXT("USE_ADAPTIVE_VOLUMETRIC_SHADOW_MAP"), ShouldHeterogeneousVolumesCastShadows());
-
 		// FXC is too slow at compiling Substrate lighting shaders
 		if (Substrate::IsSubstrateEnabled() && FDataDrivenShaderPlatformInfo::GetSupportsDxc(Parameters.Platform))
 		{
@@ -2242,6 +2238,27 @@ void FDeferredShadingSceneRenderer::RenderLights(
 						RenderFirstPersonSelfShadow(GraphBuilder, *this, Views, ScreenShadowMaskTexture, FirstPersonSelfShadowInputs, LightSceneInfo);
 					}
 
+					if (ShouldHeterogeneousVolumesCastShadows())
+					{
+						RDG_EVENT_SCOPE_STAT(GraphBuilder, VolumetricShadowMask, "VolumetricShadowMask");
+						RDG_GPU_STAT_SCOPE(GraphBuilder, VolumetricShadowMask);
+
+						// Render Volumetric Shadow Mask
+						FRDGTextureDesc Desc = FRDGTextureDesc::Create2D(
+							SceneTextures.Config.Extent,
+							PF_B8G8R8A8,
+							FClearValueBinding::White,
+							TexCreate_ShaderResource | TexCreate_RenderTargetable | TexCreate_UAV);
+						FRDGTextureRef VolumetricShadowMaskTexture = GraphBuilder.CreateTexture(Desc, TEXT("HeterogeneousVolumes.VolumetricShadowMaskTexture"));
+
+						for (const FViewInfo& View : Views)
+						{
+							FRDGTextureRef InputShadowMaskTexture = ScreenShadowMaskTexture ? ScreenShadowMaskTexture : GSystemTextures.GetWhiteDummy(GraphBuilder);
+							RenderVolumetricShadowMaskForLight(GraphBuilder, View.GetSceneTextures(), Scene, View, LightSceneInfo.Type, &LightSceneInfo, InputShadowMaskTexture, VolumetricShadowMaskTexture);
+						}
+						ScreenShadowMaskTexture = VolumetricShadowMaskTexture;
+					}
+
 					bUsedShadowMaskTexture = true;
 				} // if (bDrawShadows)
 
@@ -2907,11 +2924,6 @@ static void RenderLight(
 			PassParameters->SubstrateTileComplex = Substrate::SetTileParameters(GraphBuilder, View, ESubstrateTileType::EComplex);
 			PassParameters->SubstrateTileSpectialComplex = Substrate::SetTileParameters(GraphBuilder, View, ESubstrateTileType::EComplexSpecial);
 		}
-#if 0
-		PassParameters->PS.OrthoGridUniformBuffer = HeterogeneousVolumes::GetOrthoVoxelGridUniformBuffer(GraphBuilder, View.ViewState);
-		PassParameters->PS.FrustumGridUniformBuffer = HeterogeneousVolumes::GetFrustumVoxelGridUniformBuffer(GraphBuilder, View.ViewState);
-#endif
-		PassParameters->PS.AVSM = HeterogeneousVolumes::GetAdaptiveVolumetricShadowMapUniformBuffer(GraphBuilder, View.ViewState, LightSceneInfo);
 
 		FDeferredLightPS::FPermutationDomain PermutationVector;
 		PermutationVector.Set< FDeferredLightPS::FTransmissionDim >(bTransmission);
@@ -3260,7 +3272,6 @@ static FDeferredLightPS::FParameters GetRenderLightSimplePSParameters(
 	{
 		Out.RenderTargets.DepthStencil = FDepthStencilBinding(SceneTextures.Depth.Target, ERenderTargetLoadAction::ELoad, ERenderTargetLoadAction::ELoad, FExclusiveDepthStencil::DepthRead_StencilWrite);
 	}
-	Out.AVSM = HeterogeneousVolumes::CreateEmptyAdaptiveVolumetricShadowMapUniformBuffer(GraphBuilder);
 
 	return Out;
 }
diff --git a/Engine/Source/Runtime/Renderer/Private/ScenePrivate.h b/Engine/Source/Runtime/Renderer/Private/ScenePrivate.h
--- a/Engine/Source/Runtime/Renderer/Private/ScenePrivate.h
+++ b/Engine/Source/Runtime/Renderer/Private/ScenePrivate.h
@@ -1016,6 +1016,8 @@ public:
 	TRDGUniformBufferRef<FAdaptiveVolumetricShadowMapUniformBufferParameters> AdaptiveVolumetricCameraMapUniformBuffer = nullptr;
 	TMap<int32, TRDGUniformBufferRef<FAdaptiveVolumetricShadowMapUniformBufferParameters>> AdaptiveVolumetricShadowMapUniformBufferMap;
 
+	TMap<int32, TRDGUniformBufferRef<FBeerShadowMapUniformBufferParameters>> BeerShadowMapUniformBufferMap;
+
 	// Map from Light ID in GPU Scene to index in the View's ForwardLightData array
 	// This is stored in ViewState so we can access previous frame mapping
 	TMap<int32, int32> LightSceneIdToForwardLightIndex;
diff --git a/Engine/Source/Runtime/Renderer/Private/ShadowDepthRendering.cpp b/Engine/Source/Runtime/Renderer/Private/ShadowDepthRendering.cpp
--- a/Engine/Source/Runtime/Renderer/Private/ShadowDepthRendering.cpp
+++ b/Engine/Source/Runtime/Renderer/Private/ShadowDepthRendering.cpp
@@ -1966,7 +1966,7 @@ void FSceneRenderer::RenderShadowDepthMaps(FRDGBuilder& GraphBuilder, FDynamicSh
 	}
 
 	if (ShouldRenderHeterogeneousVolumes(Scene) && ShouldHeterogeneousVolumesCastShadows() &&
-		(HeterogeneousVolumes::GetShadowMode() == HeterogeneousVolumes::EShadowMode::LiveShading))
+		(HeterogeneousVolumes::GetShadowPipeline() == HeterogeneousVolumes::EShadowPipeline::LiveShading))
 	{
 		for (int32 ViewIndex = 0; ViewIndex < Views.Num(); ViewIndex++)
 		{
