Remove slow path for when allocating a virtual texture with NumAllocatedSpaces==16.
We used to urgently pay the slow path cost to help avoid the possibilty of the allocation taking us to NumAllocatedSpaces>16.
But now we no longer crash for that case so it is less important.
We don't want to be making things always slow for the genuinely valid case of NumAllocatedSpaces==16.
#rnx

diff --git a/Engine/Source/Runtime/Renderer/Private/VT/VirtualTextureSystem.cpp b/Engine/Source/Runtime/Renderer/Private/VT/VirtualTextureSystem.cpp
--- a/Engine/Source/Runtime/Renderer/Private/VT/VirtualTextureSystem.cpp
+++ b/Engine/Source/Runtime/Renderer/Private/VT/VirtualTextureSystem.cpp
@@ -717,15 +717,6 @@ IAllocatedVirtualTexture* FVirtualTextureSystem::AllocateVirtualTexture(FRHIComm
 		NumPhysicalGroups = Desc.NumTextureLayers;
 	}
 
-	if (UNLIKELY(NumAllocatedSpaces == VIRTUALTEXTURE_MAX_FEEDBACK_SPACES))
-	{
-		// We have already allocated the maximum number of fully supported spaces. 
-		// If this allocation needs a new space then it may no longer support feedback.
-		// To give us the best chance of success, flush all pending destroys now. This may free up a space.
-		// This is potentially an expensive operation, which is why we only do it if necessary.
-		DestroyPendingVirtualTextures(true);
-	}
-
 	AllocatedVT = new FAllocatedVirtualTexture(RHICmdList, this, Frame, Desc, ProducerForLayer, BlockWidthInTiles, BlockHeightInTiles, WidthInBlocks, HeightInBlocks, DepthInTiles);
 	AllocatedVT->NumRefs = 1;
 	if (bAnyLayerProducerWantsPersistentHighestMip)
