Fix mismatches of ACESGamutTable for declarations of SHADER_PARAMETER_SRV

diff --git a/Engine/Source/Runtime/Renderer/Private/PostProcess/PostProcessCombineLUTs.cpp b/Engine/Source/Runtime/Renderer/Private/PostProcess/PostProcessCombineLUTs.cpp
--- a/Engine/Source/Runtime/Renderer/Private/PostProcess/PostProcessCombineLUTs.cpp
+++ b/Engine/Source/Runtime/Renderer/Private/PostProcess/PostProcessCombineLUTs.cpp
@@ -116,7 +116,7 @@ BEGIN_SHADER_PARAMETER_STRUCT(FACESTonemapShaderParameters, )
 	SHADER_PARAMETER(float, ACESSceneColorMultiplier)
 	SHADER_PARAMETER(float, ACESGamutCompression)
 	SHADER_PARAMETER_SRV(Texture2D<float>, ACESReachTable)
-	SHADER_PARAMETER_SRV(Texture2D<FVector3f>, ACESGamutTable)
+	SHADER_PARAMETER_SRV(Texture2D<float>, ACESGamutTable)
 	SHADER_PARAMETER_SRV(Texture2D<float>, ACESGammaTable)
 END_SHADER_PARAMETER_STRUCT()
 
diff --git a/Engine/Source/Runtime/Renderer/Private/PostProcess/PostProcessDeviceEncodingOnly.cpp b/Engine/Source/Runtime/Renderer/Private/PostProcess/PostProcessDeviceEncodingOnly.cpp
--- a/Engine/Source/Runtime/Renderer/Private/PostProcess/PostProcessDeviceEncodingOnly.cpp
+++ b/Engine/Source/Runtime/Renderer/Private/PostProcess/PostProcessDeviceEncodingOnly.cpp
@@ -92,7 +92,7 @@ BEGIN_SHADER_PARAMETER_STRUCT(FDeviceEncodingACESTonemapShaderParameters, )
 	SHADER_PARAMETER(float, ACESSceneColorMultiplier)
 	SHADER_PARAMETER(float, ACESGamutCompression)
 	SHADER_PARAMETER_SRV(Texture2D<float>, ACESReachTable)
-	SHADER_PARAMETER_SRV(Texture2D<FVector3f>, ACESGamutTable)
+	SHADER_PARAMETER_SRV(Texture2D<float>, ACESGamutTable)
 	SHADER_PARAMETER_SRV(Texture2D<float>, ACESGammaTable)
 END_SHADER_PARAMETER_STRUCT()
 
