Mobile: When possible do not merge dynamic spotlight shadows into directional light atlas, to avoid atlas re-allocations

diff --git a/Engine/Source/Runtime/Renderer/Private/ShadowSetup.cpp b/Engine/Source/Runtime/Renderer/Private/ShadowSetup.cpp
--- a/Engine/Source/Runtime/Renderer/Private/ShadowSetup.cpp
+++ b/Engine/Source/Runtime/Renderer/Private/ShadowSetup.cpp
@@ -6411,7 +6411,9 @@ void FSceneRenderer::FilterDynamicShadows(FDynamicShadowsTaskData& TaskData)
 	FFilteredShadowArrays& ShadowArrays = TaskData.ShadowArrays;
 
 	const bool bMobile = Scene->GetFeatureLevel() < ERHIFeatureLevel::SM5;
-
+	// Whether to merge spot light shadows into directional light atlas
+	const bool bMergeMobileDynamicSpotlightShadows = bMobile && !IsMobileDeferredShadingEnabled(ShaderPlatform);
+	
 	TArray<FProjectedShadowInfo*, SceneRenderingAllocator> WholeSceneDirectionalShadowsForLight;
 	TArray<FProjectedShadowInfo*, SceneRenderingAllocator> MobileDynamicSpotlightShadows;
 
@@ -6585,7 +6587,7 @@ void FSceneRenderer::FilterDynamicShadows(FDynamicShadowsTaskData& TaskData)
 						check(ProjectedShadowInfo->bWholeSceneShadow);
 						ShadowArrays.CachedSpotlightShadows.Add(ProjectedShadowInfo);
 					}
-					else if (bMobile && ProjectedShadowInfo->bWholeSceneShadow)
+					else if (bMergeMobileDynamicSpotlightShadows && ProjectedShadowInfo->bWholeSceneShadow)
 					{
 						MobileDynamicSpotlightShadows.Add(ProjectedShadowInfo);
 					}
