Add stencil masking support to mesh decals.
Mesh decals (non DBuffer) were not obeying the receiver setting for Receive Decals.
The mobile renderer inverted the the stencil value written by the receiver when compared to the desktop renderer. Modified that to make the two renderers match.

diff --git a/Engine/Source/Runtime/Renderer/Private/CompositionLighting/PostProcessMeshDecals.cpp b/Engine/Source/Runtime/Renderer/Private/CompositionLighting/PostProcessMeshDecals.cpp
--- a/Engine/Source/Runtime/Renderer/Private/CompositionLighting/PostProcessMeshDecals.cpp
+++ b/Engine/Source/Runtime/Renderer/Private/CompositionLighting/PostProcessMeshDecals.cpp
@@ -37,6 +37,31 @@ static FAutoConsoleVariableRef CVarGDeferredDecalMeshMaterialPSOPrecache(
 	TEXT("Precache all possible required PSOs for loaded DeferredDecalMesh Materials."),
 	ECVF_ReadOnly);
 
+
+static constexpr uint32 GetMeshDecalStencilRef()
+{
+	return GET_STENCIL_BIT_MASK(RECEIVE_DECAL, 1);
+}
+
+static FRHIDepthStencilState* GetMeshDecalDepthStencilState(EDecalRenderTargetMode InRenderTargetMode)
+{
+	if (InRenderTargetMode == EDecalRenderTargetMode::DBuffer)
+	{
+		// DBuffer handles ReceiveDecal state in the receiver material.
+		return TStaticDepthStencilState<false, CF_DepthNearOrEqual>::GetRHI();
+	}
+	else
+	{
+		// GBuffer handles ReceiveDecal state using stencil.
+		return TStaticDepthStencilState<
+			false, CF_DepthNearOrEqual,
+			true, CF_Equal, SO_Keep, SO_Keep, SO_Keep,
+			false, CF_Always, SO_Keep, SO_Keep, SO_Keep,
+			GetMeshDecalStencilRef(), 0>::GetRHI();
+	}
+}
+
+
 class FMeshDecalsVS : public FMeshMaterialShader
 {
 public:
@@ -203,7 +228,7 @@ FMeshDecalMeshProcessor::FMeshDecalMeshProcessor(const FScene* Scene,
 	, PassDecalStage(DecalRendering::GetRenderStage(InRenderTargetMode, ShadingPath))
 	, RenderTargetMode(InRenderTargetMode)
 {
-	PassDrawRenderState.SetDepthStencilState(TStaticDepthStencilState<false, CF_DepthNearOrEqual>::GetRHI());
+	PassDrawRenderState.SetDepthStencilState(GetMeshDecalDepthStencilState(RenderTargetMode));
 }
 
 void FMeshDecalMeshProcessor::AddMeshBatch(const FMeshBatch& RESTRICT MeshBatch, uint64 BatchElementMask, const FPrimitiveSceneProxy* RESTRICT PrimitiveSceneProxy, int32 StaticMeshId)
@@ -272,6 +297,11 @@ bool FMeshDecalMeshProcessor::TryAddMeshBatch(
 					PassDrawRenderState.SetBlendState(DecalRendering::GetDecalBlendState(DecalBlendDesc, PassDecalStage, RenderTargetMode));
 				}
 
+				if (RenderTargetMode != EDecalRenderTargetMode::DBuffer)
+				{
+					PassDrawRenderState.SetStencilRef(GetMeshDecalStencilRef());
+				}
+
 				return Process(MeshBatch, BatchElementMask, StaticMeshId, PrimitiveSceneProxy, MaterialRenderProxy, Material, MeshFillMode, MeshCullMode);
 			}
 		}
@@ -359,6 +389,7 @@ void FMeshDecalMeshProcessor::CollectDeferredDecalMeshPSOInitializers(
 {
 	EDecalRenderTargetMode LocalRenderTargetMode = DecalRendering::GetRenderTargetMode(DecalBlendDesc, DecalRenderStage);
 	PassDrawRenderState.SetBlendState(DecalRendering::GetDecalBlendState(DecalBlendDesc, DecalRenderStage, LocalRenderTargetMode));
+	PassDrawRenderState.SetDepthStencilState(GetMeshDecalDepthStencilState(LocalRenderTargetMode));
 
 	const FMeshDrawingPolicyOverrideSettings OverrideSettings = ComputeMeshOverrideSettings(PreCacheParams);
 	ERasterizerFillMode MeshFillMode = ComputeMeshFillMode(Material, OverrideSettings);
diff --git a/Engine/Source/Runtime/Renderer/Private/MobileBasePass.cpp b/Engine/Source/Runtime/Renderer/Private/MobileBasePass.cpp
--- a/Engine/Source/Runtime/Renderer/Private/MobileBasePass.cpp
+++ b/Engine/Source/Runtime/Renderer/Private/MobileBasePass.cpp
@@ -111,7 +111,7 @@ void SetMobileBasePassDepthState(
 			0x00, 0xff >::GetRHI());
 	}
 
-	uint8 ReceiveDecals = (PrimitiveSceneProxy && !PrimitiveSceneProxy->ReceivesDecals() ? 0x01 : 0x00);
+	uint8 ReceiveDecals = (PrimitiveSceneProxy && PrimitiveSceneProxy->ReceivesDecals() ? 0x01 : 0x00);
 	StencilValue |= GET_STENCIL_BIT_MASK(RECEIVE_DECAL, ReceiveDecals);
 
 	if (bUsesDeferredShading)
diff --git a/Engine/Source/Runtime/Renderer/Private/MobileDecalRendering.cpp b/Engine/Source/Runtime/Renderer/Private/MobileDecalRendering.cpp
--- a/Engine/Source/Runtime/Renderer/Private/MobileDecalRendering.cpp
+++ b/Engine/Source/Runtime/Renderer/Private/MobileDecalRendering.cpp
@@ -113,13 +113,15 @@ void RenderDeferredDecalsMobile(FRHICommandList& RHICmdList, const FScene& Scene
 			EDecalRasterizerState DecalRasterizerState = DecalRendering::GetDecalRasterizerState(bInsideDecal, bReverseHanded, View.bReverseCulling);
 			GraphicsPSOInit.RasterizerState = DecalRendering::GetDecalRasterizerState(DecalRasterizerState);
 
+			constexpr uint32 StencilRef = GET_STENCIL_BIT_MASK(RECEIVE_DECAL, 1);
+
 			if (bInsideDecal)
 			{
 				GraphicsPSOInit.DepthStencilState = TStaticDepthStencilState<
 					false, CF_Always,
 					true, CF_Equal, SO_Keep, SO_Keep, SO_Keep,
 					false, CF_Always, SO_Keep, SO_Keep, SO_Keep,
-					GET_STENCIL_BIT_MASK(RECEIVE_DECAL, 1), 0x00>::GetRHI();
+					StencilRef, 0x00>::GetRHI();
 			}
 			else
 			{
@@ -127,13 +129,13 @@ void RenderDeferredDecalsMobile(FRHICommandList& RHICmdList, const FScene& Scene
 					false, CF_DepthNearOrEqual,
 					true, CF_Equal, SO_Keep, SO_Keep, SO_Keep,
 					false, CF_Always, SO_Keep, SO_Keep, SO_Keep,
-					GET_STENCIL_BIT_MASK(RECEIVE_DECAL, 1), 0x00>::GetRHI();
+					StencilRef, 0x00>::GetRHI();
 			}
 
 			GraphicsPSOInit.BlendState = DecalRendering::GetDecalBlendState(VisibleDecal.BlendDesc, DecalRenderStage, RenderTargetMode);
 
 			// Set shader params
-			DecalRendering::SetShader(RHICmdList, GraphicsPSOInit, 0, View, VisibleDecal, DecalRenderStage, FrustumComponentToClip, &Scene);
+			DecalRendering::SetShader(RHICmdList, GraphicsPSOInit, StencilRef, View, VisibleDecal, DecalRenderStage, FrustumComponentToClip, &Scene);
 
 			RHICmdList.DrawIndexedPrimitive(GetUnitCubeIndexBuffer(), 0, 0, 8, 0, UE_ARRAY_COUNT(GCubeIndices) / 3, View.GetStereoPassInstanceFactor());
 		}
