Virtual Texturing: Use 8 bits instead of 4 bits for SpaceID in the page mapping.
This is required after we raised the SpaceID limit above 16.
8 bits is almost certainly more than necessary, so we could take 2 or 3 bits back in future if needed.
#rnx

diff --git a/Engine/Source/Runtime/Renderer/Private/VT/TexturePagePool.cpp b/Engine/Source/Runtime/Renderer/Private/VT/TexturePagePool.cpp
--- a/Engine/Source/Runtime/Renderer/Private/VT/TexturePagePool.cpp
+++ b/Engine/Source/Runtime/Renderer/Private/VT/TexturePagePool.cpp
@@ -39,7 +39,6 @@ void FTexturePagePool::Initialize(uint32 InNumPages)
 	{
 		FPageMapping& Mapping = PageMapping[i];
 		FMemory::Memset(Mapping, 0xff);
-		Mapping.Pad = 0u;
 		Mapping.NextIndex = Mapping.PrevIndex = i;
 	}
 }
diff --git a/Engine/Source/Runtime/Renderer/Private/VT/TexturePagePool.h b/Engine/Source/Runtime/Renderer/Private/VT/TexturePagePool.h
--- a/Engine/Source/Runtime/Renderer/Private/VT/TexturePagePool.h
+++ b/Engine/Source/Runtime/Renderer/Private/VT/TexturePagePool.h
@@ -149,11 +149,10 @@ private:
 	{
 		uint32 vAddress : 24;
 		uint32 vLogSize : 4;
-		uint32 SpaceID : 4;
+		uint32 MaxLevel : 4;
 
 		uint32 NextIndex : 24;
-		uint32 MaxLevel : 4;
-		uint32 Pad : 4;
+		uint32 SpaceID : 8;
 
 		uint32 PrevIndex : 24;
 		uint32 PageTableLayerIndex : 8;
