Move slow path in AllocateVirtualTexture() when NumAllocatedSpaces==16.
We don't want to be making things always slower for the valid case of NumAllocatedSpaces==16.
The path was to help prevent a failure in AcquireSpace().
So now we only do the extra work if really needed in AcquireSpace().
#rnx

diff --git a/Engine/Source/Runtime/Renderer/Private/VT/VirtualTextureSystem.cpp b/Engine/Source/Runtime/Renderer/Private/VT/VirtualTextureSystem.cpp
--- a/Engine/Source/Runtime/Renderer/Private/VT/VirtualTextureSystem.cpp
+++ b/Engine/Source/Runtime/Renderer/Private/VT/VirtualTextureSystem.cpp
@@ -717,15 +717,6 @@ IAllocatedVirtualTexture* FVirtualTextureSystem::AllocateVirtualTexture(FRHIComm
 		NumPhysicalGroups = Desc.NumTextureLayers;
 	}
 
-	if (UNLIKELY(NumAllocatedSpaces == VIRTUALTEXTURE_MAX_FEEDBACK_SPACES))
-	{
-		// We have already allocated the maximum number of fully supported spaces. 
-		// If this allocation needs a new space then it may no longer support feedback.
-		// To give us the best chance of success, flush all pending destroys now. This may free up a space.
-		// This is potentially an expensive operation, which is why we only do it if necessary.
-		DestroyPendingVirtualTextures(true);
-	}
-
 	AllocatedVT = new FAllocatedVirtualTexture(RHICmdList, this, Frame, Desc, ProducerForLayer, BlockWidthInTiles, BlockHeightInTiles, WidthInBlocks, HeightInBlocks, DepthInTiles);
 	AllocatedVT->NumRefs = 1;
 	if (bAnyLayerProducerWantsPersistentHighestMip)
@@ -924,26 +915,43 @@ FVirtualTextureSpace* FVirtualTextureSystem::AcquireSpace(FRHICommandListBase& R
 
 	// Find an empty space slot.
 	int32 FreeSpaceIndex = INDEX_NONE;
-	for (int32 SpaceIndex = 0; FreeSpaceIndex == INDEX_NONE && SpaceIndex < Spaces.Num(); ++SpaceIndex)
+
+	// We don't support feedback on more than 16 spaces, because we only have 4 bits available in the 32bit feedback requests.
+	// If we are at the limit of spaces that support feedback then we prepare for a retry below.
+	const int32 TryCount = NumAllocatedSpaces >= VIRTUALTEXTURE_MAX_FEEDBACK_SPACES ? 2 : 1;
+
+	for (int32 TryIndex = 0; FreeSpaceIndex == INDEX_NONE && TryIndex < TryCount; ++TryIndex)
 	{
-		if (!Spaces[SpaceIndex].IsValid())
+		if (TryIndex == 1)
 		{
-			FreeSpaceIndex = SpaceIndex;
+			// We failed to find a space on the first iteration of the loop.
+			// Flush all pending destroys now which may free up a space.
+			// This is a potentially expensive operation, which is why we only do it if necessary.
+			DestroyPendingVirtualTextures(true);
 		}
-	}
 
-	// If no empty space was found then try to recycle a space. 
-	// A space may be available because we don't always release spaces immediately so as to avoid unnecessary page table recreation.
-	for (int32 SpaceIndex = 0; FreeSpaceIndex == INDEX_NONE && SpaceIndex < Spaces.Num(); ++SpaceIndex)
-	{
-		FVirtualTextureSpace* Space = Spaces[SpaceIndex].Get();
-		if (Space != nullptr && Space->GetRefCount() == 0)
+		// Use an empty space slot if it exists.
+		for (int32 SpaceIndex = 0; FreeSpaceIndex == INDEX_NONE && SpaceIndex < Spaces.Num(); ++SpaceIndex)
 		{
-			DEC_MEMORY_STAT_BY(STAT_TotalPagetableMemory, Space->GetSizeInBytes());
-			Space->ReleaseResource();
-			NumAllocatedSpaces--;
+			if (!Spaces[SpaceIndex].IsValid())
+			{
+				FreeSpaceIndex = SpaceIndex;
+			}
+		}
+
+		// If no empty space slot was found then try to recycle a space. 
+		// A space may be available because we don't always release spaces immediately so as to avoid unnecessary page table recreation.
+		for (int32 SpaceIndex = 0; FreeSpaceIndex == INDEX_NONE && SpaceIndex < Spaces.Num(); ++SpaceIndex)
+		{
+			FVirtualTextureSpace* Space = Spaces[SpaceIndex].Get();
+			if (Space != nullptr && Space->GetRefCount() == 0)
+			{
+				DEC_MEMORY_STAT_BY(STAT_TotalPagetableMemory, Space->GetSizeInBytes());
+				Space->ReleaseResource();
+				NumAllocatedSpaces--;
 
-			FreeSpaceIndex = SpaceIndex;
+				FreeSpaceIndex = SpaceIndex;
+			}
 		}
 	}
 
@@ -959,9 +967,8 @@ FVirtualTextureSpace* FVirtualTextureSystem::AcquireSpace(FRHICommandListBase& R
 
 	if (FreeSpaceIndex >= VIRTUALTEXTURE_MAX_FEEDBACK_SPACES)
 	{
-		// We don't support feedback on more than 16 spaces, because we only have 4 bits available in the 32bit feedback requests.
-		// We still create a new space to avoid crashing, but warn the user that things won't work well from this point on.
-		// To reduce the number of spaces created, try using r.VT.PageTableMode=0 or modifying RVTs to not use private spaces.
+		// We still create a new space but warn the user.
+		// If you hit this warning: To reduce the number of spaces created, try using r.VT.PageTableMode=0 or modifying RVTs to not use private spaces.
 		UE_LOG(LogVirtualTexturing, Warning, TEXT("Allocating a Virtual Texture space with ID %d. Feedback will not work for this space. A dump of the current pool state follows."), FreeSpaceIndex);
 		ListPhysicalPoolsFromConsole();
 	}
