Mobile Contact Shadows
- Inline Contact Shadows Removed, now mobile uses screen space contact shadows
- In  Mobile Forward(single and multi pass)  we will use the MOBILE_CAST_CONTACT_SHADOW Stencil mask
- In  Mobile Deferred multi pass we will read the PerObjectGBufferData and wont need any stencil bits.

diff --git a/Engine/Shaders/Private/MobileDeferredShading.usf b/Engine/Shaders/Private/MobileDeferredShading.usf
--- a/Engine/Shaders/Private/MobileDeferredShading.usf
+++ b/Engine/Shaders/Private/MobileDeferredShading.usf
@@ -293,14 +293,6 @@ FDeferredLightData InitDeferredLightFromLightParameters(FLightShaderParameters I
 	Out.bSpotLight				= InLightType == LIGHT_TYPE_SPOT;
 	Out.bRectLight				= InLightType == LIGHT_TYPE_RECT;
 	Out.HairTransmittance		= InitHairTransmittanceData();
-
-// In-line Contact Shadows only in mobile deferred multi-pass
-#if SUPPORT_CONTACT_SHADOWS
-	Out.ContactShadowLength					= abs(ContactShadows.x);
-	Out.ContactShadowLengthInWS				= ContactShadows.x < 0.0f;
-	Out.ContactShadowCastingIntensity		= ContactShadows.y;
-	Out.ContactShadowNonCastingIntensity	= ContactShadows.z;
-#endif
 	Out.ShadowedBits						= ShadowedBits;
 
 #if SUPPORT_SPOTLIGHTS_SHADOW
diff --git a/Engine/Shaders/Private/MobileLightingCommon.ush b/Engine/Shaders/Private/MobileLightingCommon.ush
--- a/Engine/Shaders/Private/MobileLightingCommon.ush
+++ b/Engine/Shaders/Private/MobileLightingCommon.ush
@@ -55,10 +55,6 @@
 #define REFLECTION_COMPOSITE_HAS_SPHERE_CAPTURES 1
 #define REFLECTION_COMPOSITE_HAS_BOX_CAPTURES 1
 #include "ReflectionEnvironmentComposite.ush"
-// In-line Contact Shadows only in mobile deferred multi-pass
-#if SUPPORT_CONTACT_SHADOWS
-float3 ContactShadows;
-#endif
 
 half GetSurfaceShadow(FGBufferData GBuffer, FShadowTerms ShadowTerms)
 {
@@ -308,14 +304,6 @@ FDeferredLightData GetDirectionalLightData(float4 ScreenPosition, inout half Out
 	LightData.SpecularScale = MobileDirectionalLight.DirectionalLightDistanceFadeMADAndSpecularScale.z;
 	LightData.DiffuseScale = MobileDirectionalLight.DirectionalLightDistanceFadeMADAndSpecularScale.w;
 
-// In-line Contact Shadows only in mobile deferred multi-pass
-#if SUPPORT_CONTACT_SHADOWS
-	LightData.ContactShadowLength = abs(ContactShadows.x);
-	LightData.ContactShadowLengthInWS = ContactShadows.x < 0.0f;
-	LightData.ContactShadowCastingIntensity = ContactShadows.y;
-	LightData.ContactShadowNonCastingIntensity = ContactShadows.z;
-#endif
-
 	LightData.ShadowedBits = MobileDirectionalLight.ShadowedBits;
 	LightData.HairTransmittance = InitHairTransmittanceData();
 	LightData.ShadowMapChannelMask = UnpackShadowMapChannelMask(MobileDirectionalLight.DirectionalLightShadowMapChannelMask);
diff --git a/Engine/Source/Runtime/Renderer/Private/MobileBasePass.cpp b/Engine/Source/Runtime/Renderer/Private/MobileBasePass.cpp
--- a/Engine/Source/Runtime/Renderer/Private/MobileBasePass.cpp
+++ b/Engine/Source/Runtime/Renderer/Private/MobileBasePass.cpp
@@ -121,7 +121,7 @@ void SetMobileBasePassDepthState(
 	}
 	else
 	{
-		// TODO: ContactShadows do not work with deferred shading atm
+		// TODO: ContactShadows do not work with deferred shading single pass
 		uint8 CastContactShadows = (PrimitiveSceneProxy && PrimitiveSceneProxy->CastsContactShadow() ? 0x01 : 0x00);
 		StencilValue |= GET_STENCIL_BIT_MASK(MOBILE_CAST_CONTACT_SHADOW, CastContactShadows);
 	}
diff --git a/Engine/Source/Runtime/Renderer/Private/MobileDeferredShadingPass.cpp b/Engine/Source/Runtime/Renderer/Private/MobileDeferredShadingPass.cpp
--- a/Engine/Source/Runtime/Renderer/Private/MobileDeferredShadingPass.cpp
+++ b/Engine/Source/Runtime/Renderer/Private/MobileDeferredShadingPass.cpp
@@ -60,12 +60,6 @@ static FAutoConsoleVariableRef CVarMobileIgnoreDeferredShadingSkyLightChannels(
 	ECVF_RenderThreadSafe
 );
 
-bool IsMobileInlineContactShadowsEnabled(EShaderPlatform ShaderPlatform)
-{
-	static auto* ContactShadowsCVar = IConsoleManager::Get().FindTConsoleVariableDataInt(TEXT("r.ContactShadows"));
-	return !MobileAllowFramebufferFetch(ShaderPlatform) && ContactShadowsCVar && ContactShadowsCVar->GetValueOnRenderThread() != 0;
-}
-
 BEGIN_SHADER_PARAMETER_STRUCT(FMobileDeferredPassParameters, )
 	SHADER_PARAMETER_RDG_UNIFORM_BUFFER(FMobileSceneTextureUniformParameters, MobileSceneTextures)
 	RENDER_TARGET_BINDING_SLOTS()
@@ -76,7 +70,6 @@ BEGIN_SHADER_PARAMETER_STRUCT(FMobileDeferredCommonParameters, )
 	SHADER_PARAMETER(FVector4f, LightFunctionParameters)
 	SHADER_PARAMETER(FVector2f, LightFunctionParameters2)
 	SHADER_PARAMETER(FVector3f, CameraRelativeLightPosition)
-	SHADER_PARAMETER(FVector3f, ContactShadows)
 END_SHADER_PARAMETER_STRUCT()
 
 class FMobileDirectionalLightFunctionPS : public FMaterialShader
@@ -91,7 +84,6 @@ class FMobileDirectionalLightFunctionPS : public FMaterialShader
 	class FEnableCSM					: SHADER_PERMUTATION_BOOL("ENABLE_MOBILE_CSM");
 	class FShadowQuality				: SHADER_PERMUTATION_RANGE_INT("MOBILE_SHADOW_QUALITY", 1, 3); // not using Quality=0
 	class FSkyShadowing					: SHADER_PERMUTATION_BOOL("APPLY_SKY_SHADOWING");
-	class FInlineContactShadows			: SHADER_PERMUTATION_BOOL("SUPPORT_CONTACT_SHADOWS");
 
 	using FPermutationDomain = TShaderPermutationDomain<
 		FEnableShadingModelSupport,
@@ -100,8 +92,7 @@ class FMobileDirectionalLightFunctionPS : public FMaterialShader
 		FEnableScreenSpaceShadowMask,
 		FEnableCSM, 
 		FShadowQuality,
-		FSkyShadowing,
-		FInlineContactShadows>;
+		FSkyShadowing>;
 
 	BEGIN_SHADER_PARAMETER_STRUCT(FParameters, )
 		SHADER_PARAMETER_STRUCT_REF(FMobileDirectionalLightShaderParameters, MobileDirectionalLight)
@@ -146,11 +137,6 @@ class FMobileDirectionalLightFunctionPS : public FMaterialShader
 			PermutationVector.Set<FSkyShadowing>(false);
 		}
 		
-		if (MobileAllowFramebufferFetch(Platform))
-		{
-			PermutationVector.Set<FInlineContactShadows>(false);
-		}
-
 		return PermutationVector;
 	}
 
@@ -195,7 +181,6 @@ class FMobileDirectionalLightFunctionPS : public FMaterialShader
 		PermutationVector.Set<FEnableCSM>(ShadowQuality > 0);
 		PermutationVector.Set<FShadowQuality>(FMath::Clamp(ShadowQuality, 1, 3));
 		PermutationVector.Set<FSkyShadowing>(bInlineReflectionAndSky && bApplySkyShadowing);
-		PermutationVector.Set<FInlineContactShadows>(IsMobileInlineContactShadowsEnabled(ShaderPlatform));
 		return PermutationVector;
 	}
 
@@ -225,8 +210,7 @@ public:
 	class FIESProfileDim			: SHADER_PERMUTATION_BOOL("USE_IES_PROFILE");
 	class FSpotLightShadowDim		: SHADER_PERMUTATION_BOOL("SUPPORT_SPOTLIGHTS_SHADOW");
 	class FSimpleLightDim			: SHADER_PERMUTATION_BOOL("SIMPLE_LIGHT");
-	class FInlineContactShadows		: SHADER_PERMUTATION_BOOL("SUPPORT_CONTACT_SHADOWS");
-	using FPermutationDomain = TShaderPermutationDomain<FEnableShadingModelSupport, FRadialLightTypeDim, FIESProfileDim, FSpotLightShadowDim, FSimpleLightDim, FInlineContactShadows>;
+	using FPermutationDomain = TShaderPermutationDomain<FEnableShadingModelSupport, FRadialLightTypeDim, FIESProfileDim, FSpotLightShadowDim, FSimpleLightDim>;
 
 	BEGIN_SHADER_PARAMETER_STRUCT(FParameters, )
 		SHADER_PARAMETER_STRUCT(FLightShaderParameters, Light)
@@ -271,18 +255,12 @@ public:
 			PermutationVector.Set<FEnableShadingModelSupport>(false);
 		}
 
-		if (MobileAllowFramebufferFetch(Platform))
-		{
-			PermutationVector.Set<FInlineContactShadows>(false);
-		}
-
 		if (PermutationVector.Get<FSimpleLightDim>() == true)
 		{
 			PermutationVector.Set<FEnableShadingModelSupport>(false);
 			PermutationVector.Set<FRadialLightTypeDim>(LIGHT_TYPE_POINT);
 			PermutationVector.Set<FIESProfileDim>(false);
 			PermutationVector.Set<FSpotLightShadowDim>(false);
-			PermutationVector.Set<FInlineContactShadows>(false);
 		}
 
 		return PermutationVector;
@@ -573,25 +551,6 @@ static void SetDirectionalLightDepthStencilState(FGraphicsPipelineStateInitializ
 	}
 }
 
-extern void GetLightContactShadowParameters(const FLightSceneProxy* Proxy, float& OutLength, bool& bOutLengthInWS, float& OutCastingIntensity, float& OutNonCastingIntensity);
-void SetupContactShadowParameters(const FLightSceneProxy* Proxy, const FViewInfo& View, FVector3f& OutContactShadowsParameter)
-{
-	float ContactShadowLength = 0.f;
-	float ContactShadowCastingIntensity = 1.f;
-	float ContactShadowNonCastingIntensity = 0.f;
-	// Only do in-line contact shadows on deferred multi pass
-	if (IsMobileInlineContactShadowsEnabled(View.GetShaderPlatform()) && View.Family->EngineShowFlags.ContactShadows)
-	{
-		bool bContactShadowLengthInWS = false;
-		GetLightContactShadowParameters(Proxy, ContactShadowLength, bContactShadowLengthInWS, ContactShadowCastingIntensity, ContactShadowNonCastingIntensity);
-		// Sign indicates if contact shadow length is in world space or screen space.
-		ContactShadowLength *= bContactShadowLengthInWS ? -1.0f : 1.0f;
-	}
-	OutContactShadowsParameter.X = ContactShadowLength;
-	OutContactShadowsParameter.Y = ContactShadowCastingIntensity;
-	OutContactShadowsParameter.Z = ContactShadowNonCastingIntensity;
-}
-
 static void RenderDirectionalLight(
 	FRHICommandList& RHICmdList, 
 	const FScene& Scene, 
@@ -673,7 +632,6 @@ static void RenderDirectionalLight(
 	const bool bInsetShadows = (CVarMobileSupportInsetShadows != nullptr && CVarMobileSupportInsetShadows->GetInt() != 0);
 	const bool bMobileUsesShadowMaskTexture = (MobileUsesShadowMaskTexture(View.GetShaderPlatform()) || bInsetShadows);
 
-	SetupContactShadowParameters(DirectionalLight.Proxy, View, PassParameters.MobileDeferredCommonParameters.ContactShadows);
 	const uint8 LightingChannelStencilValue = GetLightingChannelStencilValue(LightingChannel);
 	SetDirectionalLightDepthStencilState(GraphicsPSOInit, LightingChannel);
 	
@@ -958,7 +916,6 @@ static void RenderLocalLight(
 	PassParameters.MobileDeferredCommonParameters.TranslatedWorldToLight = FMatrix44f(FTranslationMatrix(-View.ViewMatrices.GetPreViewTranslation()) * WorldToLight);
 	PassParameters.MobileDeferredCommonParameters.CameraRelativeLightPosition = GetCamRelativeLightPosition(View.ViewMatrices, LightSceneInfo);
 	
-	SetupContactShadowParameters(LightSceneInfo.Proxy, View, PassParameters.MobileDeferredCommonParameters.ContactShadows);
 	// Do two passes, first masking DefautLit, second masking all other shading models
 	const bool bOnlyDefaultLitInView = IsOnlyDefaultLitShadingModel(View.ShadingModelMaskInView);
 	int32 NumPasses = !bOnlyDefaultLitInView && MobileUsesGBufferCustomData(Scene.GetShaderPlatform()) ? 2 : 1;
@@ -973,7 +930,6 @@ static void RenderLocalLight(
 		PermutationVector.Set<FMobileRadialLightFunctionPS::FRadialLightTypeDim>(LightType);
 		PermutationVector.Set<FMobileRadialLightFunctionPS::FIESProfileDim>(bUseIESTexture);
 		PermutationVector.Set<FMobileRadialLightFunctionPS::FSpotLightShadowDim>(bShouldCastShadow);
-		PermutationVector.Set<FMobileRadialLightFunctionPS::FInlineContactShadows>(IsMobileInlineContactShadowsEnabled(View.GetShaderPlatform()));
 		FCachedLightMaterial LightMaterial;
 		TShaderRef<FMobileRadialLightFunctionPS> PixelShader;
 		GetLightMaterial(DefaultLightMaterial, LightFunctionMaterialProxy, PermutationVector.ToDimensionValueId(), LightMaterial, PixelShader);
@@ -1064,7 +1020,6 @@ static void RenderSimpleLights(
 		// Render light
 		FMobileRadialLightFunctionPS::FParameters ParametersPS;
 		ParametersPS.Light = GetSimpleDeferredLightParameters(View, SimpleLight, SimpleLightPerViewData).LightParameters;
-		ParametersPS.MobileDeferredCommonParameters.ContactShadows = FVector3f(0.f /*ContactShadowLength*/, 1.f /*ContactShadowCastingIntensity*/, 0.f /*ContactShadowNonCastingIntensity*/);
 		SetShaderParameters(RHICmdList, VertexShader, VertexShader.GetVertexShader(), ParametersVS);
 		SetShaderParametersMixedPS(RHICmdList, PixelShader, ParametersPS, View, DefaultMaterial.MaterialProxy, *DefaultMaterial.Material);
 
diff --git a/Engine/Source/Runtime/Renderer/Private/MobileShadingRenderer.cpp b/Engine/Source/Runtime/Renderer/Private/MobileShadingRenderer.cpp
--- a/Engine/Source/Runtime/Renderer/Private/MobileShadingRenderer.cpp
+++ b/Engine/Source/Runtime/Renderer/Private/MobileShadingRenderer.cpp
@@ -1470,6 +1470,12 @@ void FMobileSceneRenderer::Render(FRDGBuilder& GraphBuilder, const FSceneRenderU
 			RDG_EVENT_SCOPE_STAT(GraphBuilder, ShadowProjection, "ShadowProjection");
 			RDG_GPU_STAT_SCOPE(GraphBuilder, ShadowProjection);
 			BasePassTextures.ScreenSpaceShadowMaskTexture = RenderMobileShadowProjections(GraphBuilder, SceneTextures.Config.Extent, bIsMobileMultiView);
+			// In mobile deferred we always render screen space shadows after the base pass
+			// TODO: ContactShadows do not work with deferred shading single pass
+			if (!bDeferredShading)
+			{
+				RenderMobileScreenSpaceShadows(GraphBuilder, SceneTextures.Config.Extent, BasePassTextures.ScreenSpaceShadowMaskTexture);
+			}
 		}
 
 		// Local Light prepass
@@ -2558,15 +2564,19 @@ void FMobileSceneRenderer::RenderDeferredMultiPass(FRDGBuilder& GraphBuilder, FS
 		RenderDistanceFieldLighting(GraphBuilder, SceneTextures, Parameters, DynamicBentNormalAOTextures, false, false);
 	}
 
-	if (bRequiresShadowProjections && !bIsFullDepthPrepassEnabled)
+	if (bRequiresShadowProjections)
 	{
-		FViewInfo* MainView = Views.Num() > 0 ? &Views[0] : nullptr;
-		const bool bIsMobileMultiView = SceneTextures.Config.bRequireMultiView || (MainView && MainView->Aspects.IsMobileMultiViewEnabled());
-		
-		RDG_CSV_STAT_EXCLUSIVE_SCOPE(GraphBuilder, RenderMobileShadowProjections);
-		RDG_EVENT_SCOPE_STAT(GraphBuilder, ShadowProjection, "ShadowProjection");
-		RDG_GPU_STAT_SCOPE(GraphBuilder, ShadowProjection);
-		BasePassTextures.ScreenSpaceShadowMaskTexture = RenderMobileShadowProjections(GraphBuilder, SceneTextures.Config.Extent, bIsMobileMultiView);
+		if (!bIsFullDepthPrepassEnabled)
+		{
+			FViewInfo* MainView = Views.Num() > 0 ? &Views[0] : nullptr;
+			const bool bIsMobileMultiView = SceneTextures.Config.bRequireMultiView || (MainView && MainView->Aspects.IsMobileMultiViewEnabled());
+
+			RDG_CSV_STAT_EXCLUSIVE_SCOPE(GraphBuilder, RenderMobileShadowProjections);
+			RDG_EVENT_SCOPE_STAT(GraphBuilder, ShadowProjection, "ShadowProjection");
+			RDG_GPU_STAT_SCOPE(GraphBuilder, ShadowProjection);
+			BasePassTextures.ScreenSpaceShadowMaskTexture = RenderMobileShadowProjections(GraphBuilder, SceneTextures.Config.Extent, bIsMobileMultiView);
+		}
+		RenderMobileScreenSpaceShadows(GraphBuilder, SceneTextures.Config.Extent, BasePassTextures.ScreenSpaceShadowMaskTexture);
 	}
 
 	const bool bSupportsSM5Nodes = MobileSupportsSM5MaterialNodes(ShaderPlatform);
diff --git a/Engine/Source/Runtime/Renderer/Private/SceneRendering.h b/Engine/Source/Runtime/Renderer/Private/SceneRendering.h
--- a/Engine/Source/Runtime/Renderer/Private/SceneRendering.h
+++ b/Engine/Source/Runtime/Renderer/Private/SceneRendering.h
@@ -2926,6 +2926,7 @@ protected:
 
 	FRDGTextureRef RenderAmbientOcclusion(FRDGBuilder& GraphBuilder, const FSceneTextures& SceneTextures);
 
+	void RenderMobileScreenSpaceShadows(FRDGBuilder& GraphBuilder, FIntPoint Extent, FRDGTextureRef ScreenShadowMaskTexture);
 	FRDGTextureRef RenderMobileShadowProjections(FRDGBuilder& GraphBuilder, FIntPoint Extent, bool bIsMobileMultiView);
 
 	FRenderTargetBindingSlots InitRenderTargetBindings_Deferred(FSceneTextures& SceneTextures, FColorTargets& ColorTargets);
diff --git a/Engine/Source/Runtime/Renderer/Private/ShadowRendering.cpp b/Engine/Source/Runtime/Renderer/Private/ShadowRendering.cpp
--- a/Engine/Source/Runtime/Renderer/Private/ShadowRendering.cpp
+++ b/Engine/Source/Runtime/Renderer/Private/ShadowRendering.cpp
@@ -2425,6 +2425,31 @@ void FMobileSceneRenderer::RenderModulatedShadowProjections(FRHICommandList& RHI
 	}
 }
 
+void FMobileSceneRenderer::RenderMobileScreenSpaceShadows(FRDGBuilder& GraphBuilder, FIntPoint Extent, FRDGTextureRef ScreenShadowMaskTexture)
+{
+	RDG_RHI_EVENT_SCOPE(GraphBuilder, RenderMobileScreenSpaceShadows);
+
+	const FMinimalSceneTextures& SceneTextures = GetActiveSceneTextures();
+	for (auto LightIt = Scene->Lights.CreateConstIterator(); LightIt; ++LightIt)
+	{
+		const FLightSceneInfoCompact& LightSceneInfoCompact = *LightIt;
+		const FLightSceneInfo* LightSceneInfo = LightSceneInfoCompact.LightSceneInfo;
+
+		const FVisibleLightInfo& VisibleLightInfo = VisibleLightInfos[LightSceneInfo->Id];
+		const FLightSceneProxy* LightSceneProxy = LightSceneInfo->Proxy;
+
+		const bool bProjectingForForwardShading = true;
+
+		if (LightSceneProxy->GetLightType() == LightType_Directional &&
+			LightSceneInfo->GetDynamicShadowMapChannel() != -1)
+		{
+			// Dynamic shadows are projected into channels of the light attenuation texture based on their assigned DynamicShadowMapChannel
+			// Only render screen space shadows if light is assigned to a valid DynamicShadowMapChannel
+			RenderScreenSpaceShadows(GraphBuilder, SceneTextures, Views, LightSceneInfo, bProjectingForForwardShading, ScreenShadowMaskTexture);
+		}
+	}
+}
+
 FRDGTextureRef FMobileSceneRenderer::RenderMobileShadowProjections(FRDGBuilder& GraphBuilder, FIntPoint Extent, bool bIsMobileMultiView)
 {
 	RDG_RHI_EVENT_SCOPE(GraphBuilder, RenderMobileShadowProjections);
@@ -2460,16 +2485,6 @@ FRDGTextureRef FMobileSceneRenderer::RenderMobileShadowProjections(FRDGBuilder&
 				bProjectingForForwardShading);
 		}
 		RenderCapsuleDirectShadows(GraphBuilder, *LightSceneInfo, ScreenShadowMaskTexture, VisibleLightInfo.CapsuleShadowsToProject, bProjectingForForwardShading);
-
-		if (LightSceneProxy->GetLightType() == LightType_Directional && 
-			LightSceneInfo->GetDynamicShadowMapChannel() != -1 &&
-			// For deferred multi-pass contact shadows are done inline
-			!(bRequiresMultiPass && bDeferredShading))
-		{
-			// Dynamic shadows are projected into channels of the light attenuation texture based on their assigned DynamicShadowMapChannel
-			// Only render screen space shadows if light is assigned to a valid DynamicShadowMapChannel
-			RenderScreenSpaceShadows(GraphBuilder, SceneTextures, Views, LightSceneInfo, bProjectingForForwardShading, ScreenShadowMaskTexture);
-		}
 	}
 
 	return ScreenShadowMaskTexture;
diff --git a/Engine/Source/Runtime/Renderer/Private/Shadows/ScreenSpaceShadows.cpp b/Engine/Source/Runtime/Renderer/Private/Shadows/ScreenSpaceShadows.cpp
--- a/Engine/Source/Runtime/Renderer/Private/Shadows/ScreenSpaceShadows.cpp
+++ b/Engine/Source/Runtime/Renderer/Private/Shadows/ScreenSpaceShadows.cpp
@@ -469,7 +469,7 @@ void RenderScreenSpaceShadows(
 {
 	static auto* ContactShadowsCVar = IConsoleManager::Get().FindTConsoleVariableDataInt(TEXT("r.ContactShadows"));
 
-	if (ContactShadowsCVar && ContactShadowsCVar->GetValueOnRenderThread() == 0)
+	if (ScreenShadowMaskTexture == nullptr || (ContactShadowsCVar && ContactShadowsCVar->GetValueOnRenderThread() == 0))
 	{
 		return;
 	}
