Removed unused memory allocations from FMeshDrawCommandPassSetupTaskContext

diff --git a/Engine/Source/Runtime/Renderer/Private/MeshDrawCommands.cpp b/Engine/Source/Runtime/Renderer/Private/MeshDrawCommands.cpp
--- a/Engine/Source/Runtime/Renderer/Private/MeshDrawCommands.cpp
+++ b/Engine/Source/Runtime/Renderer/Private/MeshDrawCommands.cpp
@@ -994,8 +994,6 @@ FMeshDrawCommandPassSetupTaskContext::FMeshDrawCommandPassSetupTaskContext()
 	, NumDynamicMeshElements(0)
 	, NumDynamicMeshCommandBuildRequestElements(0)
 	, NeedsShaderInitialisation(false)
-	, PrimitiveIdBufferData(nullptr)
-	, PrimitiveIdBufferDataSize(0)
 	, PrimitiveBounds(nullptr)
 	, VisibleMeshDrawCommandsNum(0)
 	, NewPassVisibleMeshDrawCommandsNum(0)
@@ -1376,7 +1374,7 @@ void FParallelMeshDrawCommandPass::DispatchPassSetup(
 )
 {
 	TRACE_CPUPROFILER_EVENT_SCOPE(ParallelMdcDispatchPassSetup);
-	check(!TaskEventRef.IsValid() && MeshPassProcessor != nullptr && TaskContext.PrimitiveIdBufferData == nullptr);
+	check(!TaskEventRef.IsValid() && MeshPassProcessor != nullptr);
 	check((PassType == EMeshPass::Num) == (DynamicMeshElementsPassRelevance == nullptr));
 
 	MaxNumDraws = InOutMeshDrawCommands.Num() + NumDynamicMeshElements + NumDynamicMeshCommandBuildRequestElements;
@@ -1443,11 +1441,6 @@ void FParallelMeshDrawCommandPass::DispatchPassSetup(
 	if (MaxNumDraws > 0)
 	{
 		// Preallocate resources on rendering thread based on MaxNumDraws.
-		TaskContext.PrimitiveIdBufferDataSize = TaskContext.InstanceFactor * MaxNumDraws * sizeof(int32);
-		TaskContext.PrimitiveIdBufferData = FMemory::Malloc(TaskContext.PrimitiveIdBufferDataSize);
-#if DO_GUARD_SLOW
-		FMemory::Memzero(TaskContext.PrimitiveIdBufferData, TaskContext.PrimitiveIdBufferDataSize);
-#endif // DO_GUARD_SLOW
 		TaskContext.MeshDrawCommands.Reserve(MaxNumDraws);
 		TaskContext.TempVisibleMeshDrawCommands.Reserve(MaxNumDraws);
 
@@ -1537,7 +1530,6 @@ void FParallelMeshDrawCommandPass::Cleanup()
 		TaskContext.MobileBasePassCSMMeshPassProcessor = nullptr;
 	}
 
-	FMemory::Free(TaskContext.PrimitiveIdBufferData);
 	MaxNumDraws = 0;
 	PassNameForStats.Empty();
 
@@ -1550,8 +1542,6 @@ void FParallelMeshDrawCommandPass::Cleanup()
 	TaskContext.MobileBasePassCSMMeshDrawCommands.Empty();
 	TaskContext.DynamicMeshCommandBuildRequests.Empty();
 	TaskContext.TempVisibleMeshDrawCommands.Empty();
-	TaskContext.PrimitiveIdBufferData = nullptr;
-	TaskContext.PrimitiveIdBufferDataSize = 0;
 }
 
 FParallelMeshDrawCommandPass::~FParallelMeshDrawCommandPass()
diff --git a/Engine/Source/Runtime/Renderer/Private/MeshDrawCommands.h b/Engine/Source/Runtime/Renderer/Private/MeshDrawCommands.h
--- a/Engine/Source/Runtime/Renderer/Private/MeshDrawCommands.h
+++ b/Engine/Source/Runtime/Renderer/Private/MeshDrawCommands.h
@@ -88,8 +88,6 @@ public:
 	bool NeedsShaderInitialisation;
 
 	// Resources preallocated on rendering thread.
-	void* PrimitiveIdBufferData;
-	int32 PrimitiveIdBufferDataSize;
 	FMeshCommandOneFrameArray TempVisibleMeshDrawCommands;
 
 	// For UpdateTranslucentMeshSortKeys.
