Remove TexCreate_RenderTargetable from VolumetricFog textures where unnecessary

Also remove ERDGTextureFlags::MultiFrame from VolumetricFog.LightScattering, which was leftover from the MGPU era.

diff --git a/Engine/Source/Runtime/Renderer/Private/VolumetricFog.cpp b/Engine/Source/Runtime/Renderer/Private/VolumetricFog.cpp
--- a/Engine/Source/Runtime/Renderer/Private/VolumetricFog.cpp
+++ b/Engine/Source/Runtime/Renderer/Private/VolumetricFog.cpp
@@ -1586,16 +1586,22 @@ void FSceneRenderer::ComputeVolumetricFog(FRDGBuilder& GraphBuilder,
 
 		FRDGTextureDesc VolumeDesc = GetVolumetricFogRDGTextureDesc(VolumetricFogResourceGridSize);
 
+		FRDGTextureDesc VolumeDescNoRT = VolumeDesc;
+		VolumeDescNoRT.Flags &= ~TexCreate_RenderTargetable;
+
 		FRDGTextureDesc VolumeDescFastVRAM = VolumeDesc;
 		VolumeDescFastVRAM.Flags |= GFastVRamConfig.VolumetricFog;
 
-		IntegrationData.VBufferA = GraphBuilder.CreateTexture(VolumeDescFastVRAM, TEXT("VolumetricFog.VBufferA"));
+		FRDGTextureDesc VolumeDescNoRTFastVRAM = VolumeDescNoRT;
+		VolumeDescNoRTFastVRAM.Flags |= GFastVRamConfig.VolumetricFog;
+
+		IntegrationData.VBufferA = GraphBuilder.CreateTexture(VolumeDescNoRTFastVRAM, TEXT("VolumetricFog.VBufferA"));
 		IntegrationData.VBufferA_UAV = GraphBuilder.CreateUAV(FRDGTextureUAVDesc(IntegrationData.VBufferA));
 		IntegrationData.VBufferB = nullptr;
 		IntegrationData.VBufferB_UAV = nullptr;
 		if (bUseEmissive)
 		{
-			IntegrationData.VBufferB = GraphBuilder.CreateTexture(VolumeDescFastVRAM, TEXT("VolumetricFog.VBufferB"));
+			IntegrationData.VBufferB = GraphBuilder.CreateTexture(VolumeDescNoRTFastVRAM, TEXT("VolumetricFog.VBufferB"));
 			IntegrationData.VBufferB_UAV = GraphBuilder.CreateUAV(FRDGTextureUAVDesc(IntegrationData.VBufferB));
 		}
 
@@ -1656,7 +1662,7 @@ void FSceneRenderer::ComputeVolumetricFog(FRDGBuilder& GraphBuilder,
 				bUseEmissive);
 		}
 
-		IntegrationData.LightScattering = GraphBuilder.CreateTexture(VolumeDesc, TEXT("VolumetricFog.LightScattering"), ERDGTextureFlags::MultiFrame);
+		IntegrationData.LightScattering = GraphBuilder.CreateTexture(VolumeDescNoRT, TEXT("VolumetricFog.LightScattering"));
 		IntegrationData.LightScatteringUAV = GraphBuilder.CreateUAV(FRDGTextureUAVDesc(IntegrationData.LightScattering));
 
 		{
@@ -1857,7 +1863,7 @@ void FSceneRenderer::ComputeVolumetricFog(FRDGBuilder& GraphBuilder,
 			});
 		}
 
-		FRDGTexture* IntegratedLightScattering = GraphBuilder.CreateTexture(VolumeDesc, TEXT("VolumetricFog.IntegratedLightScattering"));
+		FRDGTexture* IntegratedLightScattering = GraphBuilder.CreateTexture(VolumeDescNoRT, TEXT("VolumetricFog.IntegratedLightScattering"));
 		FRDGTextureUAV* IntegratedLightScatteringUAV = GraphBuilder.CreateUAV(FRDGTextureUAVDesc(IntegratedLightScattering));
 
 		{
