[holdout] Fix Holdout translucency "AddCopyBackgroundVisibilityPass" failure with split screen rendering. Backup the background vsibility per view.

#tests fix the repro scene in the jira

diff --git a/Engine/Source/Runtime/Renderer/Private/TranslucentRendering.cpp b/Engine/Source/Runtime/Renderer/Private/TranslucentRendering.cpp
--- a/Engine/Source/Runtime/Renderer/Private/TranslucentRendering.cpp
+++ b/Engine/Source/Runtime/Renderer/Private/TranslucentRendering.cpp
@@ -1704,7 +1704,6 @@ void FDeferredShadingSceneRenderer::RenderTranslucencyInner(
 			 */
 			FRDGTextureMSAA SeparateTranslucencyColorTexture = SharedColorTexture;
 
-			if (NumProcessedViews == 0 || View.Family->bMultiGPUForkAndJoin)
 			{
 				FTranslucencyPassResources& TranslucencyPassResources = OutTranslucencyResourceMap->Get(ViewIndex, TranslucencyPass);
 				FIntRect CopyRect = TranslucencyPassResources.ViewRect == FIntRect(0,0,0,0) ?  View.ViewRect : TranslucencyPassResources.ViewRect;
