Restore TexCreate_RenderTargetable for VBufferA/VBufferB in VolumetricFog.cpp

They're being used as render targets in VolumetricFogVoxelization.cpp, and I hadn't noticed.

diff --git a/Engine/Source/Runtime/Renderer/Private/VolumetricFog.cpp b/Engine/Source/Runtime/Renderer/Private/VolumetricFog.cpp
--- a/Engine/Source/Runtime/Renderer/Private/VolumetricFog.cpp
+++ b/Engine/Source/Runtime/Renderer/Private/VolumetricFog.cpp
@@ -1592,16 +1592,13 @@ void FSceneRenderer::ComputeVolumetricFog(FRDGBuilder& GraphBuilder,
 		FRDGTextureDesc VolumeDescFastVRAM = VolumeDesc;
 		VolumeDescFastVRAM.Flags |= GFastVRamConfig.VolumetricFog;
 
-		FRDGTextureDesc VolumeDescNoRTFastVRAM = VolumeDescNoRT;
-		VolumeDescNoRTFastVRAM.Flags |= GFastVRamConfig.VolumetricFog;
-
-		IntegrationData.VBufferA = GraphBuilder.CreateTexture(VolumeDescNoRTFastVRAM, TEXT("VolumetricFog.VBufferA"));
+		IntegrationData.VBufferA = GraphBuilder.CreateTexture(VolumeDescFastVRAM, TEXT("VolumetricFog.VBufferA"));
 		IntegrationData.VBufferA_UAV = GraphBuilder.CreateUAV(FRDGTextureUAVDesc(IntegrationData.VBufferA));
 		IntegrationData.VBufferB = nullptr;
 		IntegrationData.VBufferB_UAV = nullptr;
 		if (bUseEmissive)
 		{
-			IntegrationData.VBufferB = GraphBuilder.CreateTexture(VolumeDescNoRTFastVRAM, TEXT("VolumetricFog.VBufferB"));
+			IntegrationData.VBufferB = GraphBuilder.CreateTexture(VolumeDescFastVRAM, TEXT("VolumetricFog.VBufferB"));
 			IntegrationData.VBufferB_UAV = GraphBuilder.CreateUAV(FRDGTextureUAVDesc(IntegrationData.VBufferB));
 		}
 
