Mobile: Disable screen space shadows on GL platforms. Need to modify CS to use RGBA format as writing to RG images is not supported
Make sure we use dynamic inset shadows only if rendering configuration can actually support them

diff --git a/Engine/Source/Runtime/Renderer/Private/ShadowRendering.cpp b/Engine/Source/Runtime/Renderer/Private/ShadowRendering.cpp
--- a/Engine/Source/Runtime/Renderer/Private/ShadowRendering.cpp
+++ b/Engine/Source/Runtime/Renderer/Private/ShadowRendering.cpp
@@ -2427,6 +2427,14 @@ void FMobileSceneRenderer::RenderModulatedShadowProjections(FRHICommandList& RHI
 
 void FMobileSceneRenderer::RenderMobileScreenSpaceShadows(FRDGBuilder& GraphBuilder, FIntPoint Extent, FRDGTextureRef ScreenShadowMaskTexture)
 {
+	static auto* ContactShadowsCVar = IConsoleManager::Get().FindTConsoleVariableDataInt(TEXT("r.ContactShadows"));
+	if (ScreenShadowMaskTexture == nullptr
+		|| IsOpenGLPlatform(ShaderPlatform) // FIXME: GL does not support writing to rg32f images
+		|| (ContactShadowsCVar && ContactShadowsCVar->GetValueOnRenderThread() == 0))
+	{
+		return;
+	}
+		
 	RDG_RHI_EVENT_SCOPE(GraphBuilder, RenderMobileScreenSpaceShadows);
 
 	const FMinimalSceneTextures& SceneTextures = GetActiveSceneTextures();
diff --git a/Engine/Source/Runtime/Renderer/Private/ShadowSetup.cpp b/Engine/Source/Runtime/Renderer/Private/ShadowSetup.cpp
--- a/Engine/Source/Runtime/Renderer/Private/ShadowSetup.cpp
+++ b/Engine/Source/Runtime/Renderer/Private/ShadowSetup.cpp
@@ -4832,6 +4832,7 @@ struct FGatherShadowPrimitivesPacket
 	TArray<FAddSubjectPrimitiveOverflowedIndices> ViewDependentWholeSceneShadowOverflowedIndices;
 	TArray<FShadowSubjectPrimitives> PreShadowSubjectPrimitives;
 	TArray<FShadowSubjectPrimitives> ViewDependentWholeSceneShadowSubjectPrimitives;
+	bool bSupportMobileInsetShadows = false;
 
 	FGatherShadowPrimitivesPacket(
 		FScenePrimitiveOctree::FNodeIndex InNodeIndex,
@@ -4866,6 +4867,12 @@ struct FGatherShadowPrimitivesPacket
 		ViewDependentWholeSceneShadowSubjectPrimitives.Empty(NumVDWSShadows);
 		ViewDependentWholeSceneShadowSubjectPrimitives.AddDefaulted(NumVDWSShadows);
 
+		if (TaskData.FeatureLevel < ERHIFeatureLevel::SM5)
+		{
+			bSupportMobileInsetShadows = CVarMobileInsetShadows.GetValueOnRenderThread() 
+				&& (MobileUsesFullDepthPrepass(TaskData.ShaderPlatform) || (IsMobileDeferredShadingEnabled(TaskData.ShaderPlatform) && !MobileAllowFramebufferFetch(TaskData.ShaderPlatform)));
+		}
+		
 		if (NodeIndex != INDEX_NONE)
 		{
 			// Check all the primitives in this octree node.
@@ -5033,8 +5040,7 @@ struct FGatherShadowPrimitivesPacket
 				// if we are on mobile also check if inset shadows are supported
 				if (bCastsInsetShadow && TaskData.FeatureLevel < ERHIFeatureLevel::SM5)
 				{
-					const bool bMobileInsetShadows = CVarMobileInsetShadows.GetValueOnRenderThread();
-					bCastsInsetShadow &= bMobileInsetShadows;
+					bCastsInsetShadow &= bSupportMobileInsetShadows;
 				}
 
 				if (!bScreenSpaceSizeCulled
