HeterogeneousVolumes: Enabling shadow-casting volumes by default.

Default scalability changes for shadowing volumes.
Fixes shadow mask input type definition.

#tests Dev testing. All Niagara and Rendering Engine Test regressions

diff --git a/Engine/Config/BaseScalability.ini b/Engine/Config/BaseScalability.ini
--- a/Engine/Config/BaseScalability.ini
+++ b/Engine/Config/BaseScalability.ini
@@ -768,8 +768,7 @@ fx.Niagara.QualityLevel=3
 r.Refraction.OffsetQuality=1
 r.HeterogeneousVolumes.DownsampleFactor=2
 r.HeterogeneousVolumes.MaxStepCount=256
-r.HeterogeneousVolumes.Shadows.Resolution=256
-r.HeterogeneousVolumes.Shadows.MaxSampleCount=8
+r.HeterogeneousVolumes.Shadows.Resolution=512
 r.HeterogeneousVolumes.UseExistenceMask=0
 
 [EffectsQuality@Cine]
@@ -802,6 +801,7 @@ r.Refraction.OffsetQuality=1
 r.HeterogeneousVolumes.ApplyFogInscattering=1
 r.HeterogeneousVolumes.DownsampleFactor=1
 r.HeterogeneousVolumes.MaxStepCount=1024
+r.HeterogeneousVolumes.Shadows.Type=0
 r.HeterogeneousVolumes.Shadows.Resolution=1024
 r.HeterogeneousVolumes.Shadows.MaxSampleCount=16
 r.HeterogeneousVolumes.UseExistenceMask=1
diff --git a/Engine/Shaders/Private/HeterogeneousVolumes/HeterogeneousVolumesShadowMask.usf b/Engine/Shaders/Private/HeterogeneousVolumes/HeterogeneousVolumesShadowMask.usf
--- a/Engine/Shaders/Private/HeterogeneousVolumes/HeterogeneousVolumesShadowMask.usf
+++ b/Engine/Shaders/Private/HeterogeneousVolumes/HeterogeneousVolumesShadowMask.usf
@@ -35,7 +35,7 @@ int3 GroupCount;
 int DownsampleFactor;
 
 // Shadow mask
-Texture2D<float> ShadowMaskTexture;
+Texture2D<float4> ShadowMaskTexture;
 
 // Output
 RWTexture2D<float4> RWVolumetricShadowMaskTexture;
diff --git a/Engine/Source/Runtime/Renderer/Private/HeterogeneousVolumes/HeterogeneousVolumes.cpp b/Engine/Source/Runtime/Renderer/Private/HeterogeneousVolumes/HeterogeneousVolumes.cpp
--- a/Engine/Source/Runtime/Renderer/Private/HeterogeneousVolumes/HeterogeneousVolumes.cpp
+++ b/Engine/Source/Runtime/Renderer/Private/HeterogeneousVolumes/HeterogeneousVolumes.cpp
@@ -21,8 +21,8 @@ static TAutoConsoleVariable<int32> CVarHeterogeneousVolumes(
 
 static TAutoConsoleVariable<int32> CVarHeterogeneousVolumesShadows(
 	TEXT("r.HeterogeneousVolumes.Shadows"),
-	0,
-	TEXT("Enables heterogeneous volume-casting shadows (default = 0)"),
+	1,
+	TEXT("Enables heterogeneous volume-casting shadows (default = 1)"),
 	ECVF_RenderThreadSafe
 );
 
@@ -262,8 +262,8 @@ static TAutoConsoleVariable<int32> CVarHeterogeneousVolumesLightingCache(
 
 static TAutoConsoleVariable<int32> CVarHeterogeneousVolumesLightingCacheUseAVSM(
 	TEXT("r.HeterogeneousVolumes.LightingCache.UseAVSM"),
-	1,
-	TEXT("Enables use of AVSMs when evaluating self-shadowing (Default = 1)"),
+	0,
+	TEXT("Enables use of AVSMs when evaluating self-shadowing (Default = 0)"),
 	ECVF_RenderThreadSafe
 );
 
