HeterogeneousVolumes: Fix cutout artifact when ShadowType does not match TranslucencyCompositingMode

diff --git a/Engine/Source/Runtime/Renderer/Private/HeterogeneousVolumes/HeterogeneousVolumesLiveShadingPipeline.cpp b/Engine/Source/Runtime/Renderer/Private/HeterogeneousVolumes/HeterogeneousVolumesLiveShadingPipeline.cpp
--- a/Engine/Source/Runtime/Renderer/Private/HeterogeneousVolumes/HeterogeneousVolumesLiveShadingPipeline.cpp
+++ b/Engine/Source/Runtime/Renderer/Private/HeterogeneousVolumes/HeterogeneousVolumesLiveShadingPipeline.cpp
@@ -2175,6 +2175,7 @@ bool RenderVolumetricShadowMapForLightForHeterogeneousVolumeWithLiveShading(
 	const FLightSceneInfo* LightSceneInfo,
 	const FVisibleLightInfo* VisibleLightInfo,
 	// Shadow data
+	HeterogeneousVolumes::EShadowType ShadowType,
 	HeterogeneousVolumes::FLODInfo LODInfo,
 	const FVector3f* TranslatedWorldOrigin,
 	int32 NumShadowMatrices,
@@ -2291,13 +2292,11 @@ bool RenderVolumetricShadowMapForLightForHeterogeneousVolumeWithLiveShading(
 	}
 #endif // WANTS_DRAW_MESH_EVENTS
 
-	const int32 ShadowmapType = static_cast<int32>(HeterogeneousVolumes::GetShadowType());
-
 	FRenderVolumetricShadowMapForLightWithLiveShadingCS::FPermutationDomain PermutationVector;
 	PermutationVector.Set<FRenderVolumetricShadowMapForLightWithLiveShadingCS::FUseAVSMCompression>(HeterogeneousVolumes::UseAVSMCompression());
 	PermutationVector.Set<FRenderVolumetricShadowMapForLightWithLiveShadingCS::FUseCameraSceneDepth>(bUseCameraSceneDepth && HeterogeneousVolumes::ShadowsUseCameraSceneDepth());
 	PermutationVector.Set<FRenderVolumetricShadowMapForLightWithLiveShadingCS::FUseAnalyticDerivatives>(HeterogeneousVolumes::UseAnalyticDerivatives());
-	PermutationVector.Set<FRenderVolumetricShadowMapForLightWithLiveShadingCS::FShadowMapType>(ShadowmapType);
+	PermutationVector.Set<FRenderVolumetricShadowMapForLightWithLiveShadingCS::FShadowMapType>(static_cast<int32>(ShadowType));
 	TShaderRef<FRenderVolumetricShadowMapForLightWithLiveShadingCS> ComputeShader = Material.GetShader<FRenderVolumetricShadowMapForLightWithLiveShadingCS>(&FLocalVertexFactory::StaticType, PermutationVector, false);
 	if (!ComputeShader.IsNull())
 	{
@@ -2834,6 +2833,7 @@ bool RenderVolumetricShadowMapForLightWithLiveShading(
 			LightSceneInfo,
 			VisibleLightInfo,
 			// Shadow Info
+			HeterogeneousVolumes::GetShadowType(),
 			LODInfo,
 			TranslatedWorldOrigin,
 			NumShadowMatrices,
@@ -2870,6 +2870,7 @@ bool RenderVolumetricShadowMapForLightWithLiveShading(
 				LightSceneInfo,
 				VisibleLightInfo,
 				// Shadow Info
+				HeterogeneousVolumes::GetShadowType(),
 				LODInfo,
 				TranslatedWorldOrigin,
 				NumShadowMatrices,
@@ -3226,6 +3227,7 @@ void RenderAdaptiveVolumetricCameraMapWithLiveShading(
 		LODInfo.DownsampleFactor = HeterogeneousVolumes::GetCameraDownsampleFactor();
 		LODInfo.bIsPerspective = true;
 
+		HeterogeneousVolumes::EShadowType ShadowType = HeterogeneousVolumes::GetTranslucencyCompositingMode() == HeterogeneousVolumes::ETranslucencyCompositingMode::VolumetricShadowMap ? HeterogeneousVolumes::EShadowType::AdaptiveVolumetricShadowMap : HeterogeneousVolumes::EShadowType::BeerShadowMap;
 		// Build a camera shadow for one volume
 		int32 LightType = 0;
 		FLightSceneInfo* LightSceneInfo = nullptr;
@@ -3241,6 +3243,7 @@ void RenderAdaptiveVolumetricCameraMapWithLiveShading(
 			LightSceneInfo,
 			VisibleLightInfo,
 			// Shadow Info
+			ShadowType,
 			LODInfo,
 			TranslatedWorldOrigin,
 			NumShadowMatrices,
@@ -3278,6 +3281,7 @@ void RenderAdaptiveVolumetricCameraMapWithLiveShading(
 				LightSceneInfo,
 				VisibleLightInfo,
 				// Shadow Info
+				ShadowType,
 				LODInfo,
 				TranslatedWorldOrigin,
 				NumShadowMatrices,
