Heterogeneous Volumes: Re-enable LOD factor when creating volumetric shadow maps without cascades

diff --git a/Engine/Source/Runtime/Renderer/Private/HeterogeneousVolumes/HeterogeneousVolumesLiveShadingPipeline.cpp b/Engine/Source/Runtime/Renderer/Private/HeterogeneousVolumes/HeterogeneousVolumesLiveShadingPipeline.cpp
--- a/Engine/Source/Runtime/Renderer/Private/HeterogeneousVolumes/HeterogeneousVolumesLiveShadingPipeline.cpp
+++ b/Engine/Source/Runtime/Renderer/Private/HeterogeneousVolumes/HeterogeneousVolumesLiveShadingPipeline.cpp
@@ -2255,11 +2255,17 @@ bool RenderVolumetricShadowMapForLightForHeterogeneousVolumeWithLiveShading(
 		PassParameters->VoxelResolution = HeterogeneousVolumeInterface->GetVoxelResolution();
 
 		// Ray Data
-		float LODFactor = HeterogeneousVolumes::CalcLODFactor(LODInfo, HeterogeneousVolumeInterface);
 		PassParameters->ShadowStepSize = HeterogeneousVolumes::GetShadowStepSize();
-		// TODO: Assign LOD factor per sub-region shadow for CSM
-		//PassParameters->ShadowStepFactor = HeterogeneousVolumeInterface->GetShadowStepFactor() * LODFactor;
 		PassParameters->ShadowStepFactor = HeterogeneousVolumeInterface->GetShadowStepFactor();
+
+		// TODO: Assign LOD factor per sub-region shadow for CSM
+		bool bCascadeShadowForDirectionalLight = LightType == LightType_Directional && HeterogeneousVolumes::SupportsCascadeShadowsForDirectionalLight();
+		if (!bCascadeShadowForDirectionalLight)
+		{
+			float LODFactor = HeterogeneousVolumes::CalcLODFactor(LODInfo, HeterogeneousVolumeInterface);
+			PassParameters->ShadowStepFactor *= LODFactor;
+		}
+
 		//PassParameters->MaxTraceDistance = HeterogeneousVolumes::GetMaxTraceDistance();
 		//PassParameters->MaxTraceDistance = HeterogeneousVolumes::GetMaxShadowTraceDistance();
 		PassParameters->MaxTraceDistance = MaxShadowTraceDistance;
