Fix crash on Mac when changing number of sample per pixel.

R32Uint texture format is not properly support for MSAA render target. Forcing sample count to 1 on such platform.

diff --git a/Engine/Source/Runtime/Renderer/Private/HairStrands/HairStrandsVisibility.cpp b/Engine/Source/Runtime/Renderer/Private/HairStrands/HairStrandsVisibility.cpp
--- a/Engine/Source/Runtime/Renderer/Private/HairStrands/HairStrandsVisibility.cpp
+++ b/Engine/Source/Runtime/Renderer/Private/HairStrands/HairStrandsVisibility.cpp
@@ -214,6 +214,12 @@ static uint32 GetMaxSamplePerPixel(EShaderPlatform InPlatform)
 		}
 		case HairVisibilityRenderMode_MSAA_Visibility:
 		{
+			// Mac platfoms don't properly support R32_UINT format for for MSAA render target. To alleviate this, force sample count to 1.
+			#if PLATFORM_MAC
+			{
+				return 1;
+			}
+			#else
 			if (GHairVisibilityMSAA_MaxSamplePerPixel <= 1)
 			{
 				return 1;
@@ -230,6 +236,7 @@ static uint32 GetMaxSamplePerPixel(EShaderPlatform InPlatform)
 			{
 				return 8;
 			}
+			#endif
 		}
 		case HairVisibilityRenderMode_PPLL:
 		{
