Fix incorrect enabling advanced debugging shader with Substrate on platform blendable platforms.

This was enabling certain block of code,  causing shader compilation issue on mobile.

#FailedValidations Preflight

diff --git a/Engine/Shaders/Private/BasePassPixelShader.usf b/Engine/Shaders/Private/BasePassPixelShader.usf
--- a/Engine/Shaders/Private/BasePassPixelShader.usf
+++ b/Engine/Shaders/Private/BasePassPixelShader.usf
@@ -1855,7 +1855,7 @@ void FPixelShaderInOut_MainPS(
 		FSubstrateOpaqueRoughRefractionData OpaqueRoughRefractionData = (FSubstrateOpaqueRoughRefractionData)0;
 		FSubstrateIntegrationSettings Settings = InitSubstrateIntegrationSettings(false /*bForceFullyRough*/, SubstrateStruct.bRoughDiffuse, SubstrateStruct.PeelLayersAboveDepth, SubstrateStruct.bRoughnessTracking);
 
-	#if SUBSTRATE_ADVANCED_DEBUG_ENABLED
+	#if SUBSTRATE_ADVANCED_DEBUG_ENABLED && SUBSTRATE_MATERIALCONTAINER_IS_VIEWRESOURCE
 		Settings.SliceStoringDebugSubstrateTreeData = all(uint2(float2(View.CursorPosition) * View.ViewResolutionFraction) == PixelPos) ? SubstrateStruct.SliceStoringDebugSubstrateTreeDataWithoutMRT : Settings.SliceStoringDebugSubstrateTreeData;
 	#endif
 
diff --git a/Engine/Shaders/Private/MobileBasePassPixelShader.usf b/Engine/Shaders/Private/MobileBasePassPixelShader.usf
--- a/Engine/Shaders/Private/MobileBasePassPixelShader.usf
+++ b/Engine/Shaders/Private/MobileBasePassPixelShader.usf
@@ -714,7 +714,7 @@ void Main(
 		FSubstrateTopLayerData TopLayerData = (FSubstrateTopLayerData)0;
 		FSubstrateOpaqueRoughRefractionData OpaqueRoughRefractionData = (FSubstrateOpaqueRoughRefractionData)0;
 		FSubstrateIntegrationSettings Settings = InitSubstrateIntegrationSettings(FULLY_ROUGH, SubstrateStruct.bRoughDiffuse, SubstrateStruct.PeelLayersAboveDepth, SubstrateStruct.bRoughnessTracking);
-	#if SUBSTRATE_ADVANCED_DEBUG_ENABLED
+	#if SUBSTRATE_ADVANCED_DEBUG_ENABLED && SUBSTRATE_MATERIALCONTAINER_IS_VIEWRESOURCE
 		Settings.SliceStoringDebugSubstrateTreeData = all(uint2(float2(View.CursorPosition) * View.ViewResolutionFraction) == PixelPos) ? SubstrateStruct.SliceStoringDebugSubstrateTreeDataWithoutMRT : Settings.SliceStoringDebugSubstrateTreeData;
 	#endif
 
diff --git a/Engine/Source/Runtime/Engine/Private/ShaderCompiler/ShaderCompiler.cpp b/Engine/Source/Runtime/Engine/Private/ShaderCompiler/ShaderCompiler.cpp
--- a/Engine/Source/Runtime/Engine/Private/ShaderCompiler/ShaderCompiler.cpp
+++ b/Engine/Source/Runtime/Engine/Private/ShaderCompiler/ShaderCompiler.cpp
@@ -4000,7 +4000,7 @@ void GlobalBeginCompileShader(
 		const bool bSubstrateOpaqueRoughRefrac = bSubstrate && Substrate::IsOpaqueRoughRefractionEnabled(Target.GetPlatform());
 		SET_SHADER_DEFINE(Input.Environment, SUBSTRATE_OPAQUE_ROUGH_REFRACTION_ENABLED, bSubstrateOpaqueRoughRefrac ? 1 : 0);
 
-		const bool bSubstrateAdvDebug = bSubstrate && Substrate::IsAdvancedVisualizationEnabled();
+		const bool bSubstrateAdvDebug = bSubstrate && Substrate::IsAdvancedVisualizationEnabled(Target.GetPlatform());
 		SET_SHADER_DEFINE(Input.Environment, SUBSTRATE_ADVANCED_DEBUG_ENABLED, bSubstrateAdvDebug ? 1 : 0);
 
 		const bool IsStochasticLightingEnabled = bSubstrate && Substrate::IsStochasticLightingEnabled(Target.GetPlatform());
diff --git a/Engine/Source/Runtime/Engine/Private/SubstrateVisualizationData.cpp b/Engine/Source/Runtime/Engine/Private/SubstrateVisualizationData.cpp
--- a/Engine/Source/Runtime/Engine/Private/SubstrateVisualizationData.cpp
+++ b/Engine/Source/Runtime/Engine/Private/SubstrateVisualizationData.cpp
@@ -95,7 +95,7 @@ void FSubstrateVisualizationData::Initialize()
 			LOCTEXT("AdvancedMaterialPropertiesDesc", "Visualizes Substrate advanced material properties"),
 			FSubstrateViewMode::AdvancedMaterialProperties,
 			true,
-			Substrate::IsAdvancedVisualizationEnabled(),
+			Substrate::IsAdvancedVisualizationEnabled(GMaxRHIShaderPlatform),
 			LOCTEXT("IsSubstrateAdvancedDebugShaderEnabled", "Substrate advanced debugging r.Substrate.Debug.AdvancedVisualizationShaders is disabled"));
 
 		AddVisualizationMode(
diff --git a/Engine/Source/Runtime/RenderCore/Private/RenderUtils.cpp b/Engine/Source/Runtime/RenderCore/Private/RenderUtils.cpp
--- a/Engine/Source/Runtime/RenderCore/Private/RenderUtils.cpp
+++ b/Engine/Source/Runtime/RenderCore/Private/RenderUtils.cpp
@@ -2304,7 +2304,12 @@ namespace Substrate
 
 	bool IsAdvancedVisualizationEnabled()
 	{
-		return CVarSubstrateDebugAdvancedVisualizationShaders.GetValueOnAnyThread() > 0;
+		return IsAdvancedVisualizationEnabled(GMaxRHIShaderPlatform);
+	}
+
+	bool IsAdvancedVisualizationEnabled(EShaderPlatform InPlatform)
+	{
+		return CVarSubstrateDebugAdvancedVisualizationShaders.GetValueOnAnyThread() > 0 && !IsSubstrateBlendableGBufferEnabled(InPlatform);
 	}
 
 	bool IsMaterialLayeringSupportEnabled()
diff --git a/Engine/Source/Runtime/RenderCore/Private/Shader.cpp b/Engine/Source/Runtime/RenderCore/Private/Shader.cpp
--- a/Engine/Source/Runtime/RenderCore/Private/Shader.cpp
+++ b/Engine/Source/Runtime/RenderCore/Private/Shader.cpp
@@ -2496,7 +2496,7 @@ void ShaderMapAppendKey(EShaderPlatform Platform, FShaderKeyGenerator& KeyGen)
 			KeyGen.Append(TEXT("STRTNRMQ"));
 		}
 
-		if (Substrate::IsAdvancedVisualizationEnabled())
+		if (Substrate::IsAdvancedVisualizationEnabled(Platform))
 		{
 			KeyGen.AppendSeparator();
 			KeyGen.Append(TEXT("ADVDEBUG"));
diff --git a/Engine/Source/Runtime/RenderCore/Public/RenderUtils.h b/Engine/Source/Runtime/RenderCore/Public/RenderUtils.h
--- a/Engine/Source/Runtime/RenderCore/Public/RenderUtils.h
+++ b/Engine/Source/Runtime/RenderCore/Public/RenderUtils.h
@@ -633,7 +633,9 @@ namespace Substrate
 	RENDERCORE_API bool IsBackCompatibilityEnabled();
 	RENDERCORE_API bool IsDBufferPassEnabled(EShaderPlatform InPlatform);
 	RENDERCORE_API bool IsOpaqueRoughRefractionEnabled(EShaderPlatform InPlatform);
+	UE_DEPRECATED(5.7, "Plase use IsAdvancedVisualizationEnabled(EShaderPlatform InPlatform)")
 	RENDERCORE_API bool IsAdvancedVisualizationEnabled();
+	RENDERCORE_API bool IsAdvancedVisualizationEnabled(EShaderPlatform InPlatform);
 	RENDERCORE_API bool IsMaterialLayeringSupportEnabled();
 
 	RENDERCORE_API bool IsHiddenMaterialAssetConversionEnabled();
diff --git a/Engine/Source/Runtime/Renderer/Private/Substrate/Substrate.cpp b/Engine/Source/Runtime/Renderer/Private/Substrate/Substrate.cpp
--- a/Engine/Source/Runtime/Renderer/Private/Substrate/Substrate.cpp
+++ b/Engine/Source/Runtime/Renderer/Private/Substrate/Substrate.cpp
@@ -578,7 +578,7 @@ void InitialiseSubstrateFrameSceneData(FRDGBuilder& GraphBuilder, FSceneRenderer
 
 	// Create the material data container
 	const uint32 SliceCountSSS = SUBSTRATE_SSS_DATA_UINT_COUNT;
-	const uint32 SliceCountAdvDebug = IsAdvancedVisualizationEnabled() ? 1 : 0;
+	const uint32 SliceCountAdvDebug = IsAdvancedVisualizationEnabled(SceneRenderer.ShaderPlatform) ? 1 : 0;
 	const uint32 SliceCount = bNeedsMaterialBuffer ? FMath::DivideAndRoundUp(Out.EffectiveMaxBytesPerPixel, 4u) + SliceCountSSS + SliceCountAdvDebug : 1u;
 	if (bNeedsMaterialBuffer && !bBlendableGBufferEnabled)
 	{
diff --git a/Engine/Source/Runtime/Renderer/Private/Substrate/SubstrateVisualize.cpp b/Engine/Source/Runtime/Renderer/Private/Substrate/SubstrateVisualize.cpp
--- a/Engine/Source/Runtime/Renderer/Private/Substrate/SubstrateVisualize.cpp
+++ b/Engine/Source/Runtime/Renderer/Private/Substrate/SubstrateVisualize.cpp
@@ -197,7 +197,7 @@ class FMaterialDebugSubstrateTreeCS : public FGlobalShader
 	static bool IsSupported(EShaderPlatform InPlatform, EShaderPermutationFlags InFlags=EShaderPermutationFlags::None)
 	{
 		return IsSubstrateDebugVisualizationSupported(InPlatform, true, InFlags)
-			&& IsAdvancedVisualizationEnabled()
+			&& IsAdvancedVisualizationEnabled(InPlatform)
 			&& !IsSubstrateBlendableGBufferEnabled(InPlatform)
 			&& !IsVulkanPlatform(InPlatform)  // SUBSTRATE_TODO Move to CPU debug visualisation and it should then work on all platforms
 			&& !IsMetalPlatform(InPlatform);
@@ -249,7 +249,7 @@ class FMaterialDebugSubstrateTreePS : public FGlobalShader
 	static bool IsSupported(EShaderPlatform InPlatform, EShaderPermutationFlags InFlags=EShaderPermutationFlags::None)
 	{
 		return IsSubstrateDebugVisualizationSupported(InPlatform, true, InFlags)
-			&& IsAdvancedVisualizationEnabled()
+			&& IsAdvancedVisualizationEnabled(InPlatform)
 			&& !IsSubstrateBlendableGBufferEnabled(InPlatform)
 			&& !IsVulkanPlatform(InPlatform)  // SUBSTRATE_TODO Move to CPU debug visualisation and it should then work on all platforms
 			&& !IsMetalPlatform(InPlatform);
@@ -1105,7 +1105,7 @@ static void AddVisualizeSystemInfoPasses(FRDGBuilder& GraphBuilder, const FViewI
 	SystemInformation.bCMask = SupportsCMask(View.GetShaderPlatform()) ? 1 : 0;
 	SystemInformation.bAsync = IsClassificationAsync() ? 1 : 0;
 	SystemInformation.b8BitsCoord = Is8bitTileCoordEnabled() ? 1 : 0;
-	SystemInformation.bAdvancedDebugEnabled = IsAdvancedVisualizationEnabled() ? 1 : 0;
+	SystemInformation.bAdvancedDebugEnabled = IsAdvancedVisualizationEnabled(View.GetShaderPlatform()) ? 1 : 0;
 	SystemInformation.LayerPeelIndex = View.SubstrateViewData.SceneData->PeelLayersAboveDepth;
 
 	////////////////////////////////////////////////////////////////////////////////////////////////////
@@ -1364,7 +1364,7 @@ static void AddVisualizeSystemInfoPasses(FRDGBuilder& GraphBuilder, const FViewI
 // Draw each material layer independently
 static void AddVisualizeAdvancedMaterialPasses(FRDGBuilder& GraphBuilder, const FViewInfo& View, FScreenPassTexture& ScreenPassSceneColor, EShaderPlatform Platform)
 {
-	if (!IsAdvancedVisualizationEnabled() ||
+	if (!IsAdvancedVisualizationEnabled(Platform) ||
 		!FMaterialDebugSubstrateTreeCS::IsSupported(Platform) || 
 		!FMaterialDebugSubstrateTreePS::IsSupported(Platform))
 	{
