Fix for incorrect per-view directional light information in VSM FBaseGeneratePageFlagsParameters, leading to artifacts in splitscreen views.

diff --git a/Engine/Source/Runtime/Renderer/Private/VirtualShadowMaps/VirtualShadowMapArray.cpp b/Engine/Source/Runtime/Renderer/Private/VirtualShadowMaps/VirtualShadowMapArray.cpp
--- a/Engine/Source/Runtime/Renderer/Private/VirtualShadowMaps/VirtualShadowMapArray.cpp
+++ b/Engine/Source/Runtime/Renderer/Private/VirtualShadowMaps/VirtualShadowMapArray.cpp
@@ -2613,7 +2613,7 @@ void FVirtualShadowMapArray::BeginMarkPages(
 	// Note: always run this pass such that the distant lights may be marked if need be
 	{
 		FMarkCoarsePagesCS::FParameters* PassParameters = GraphBuilder.AllocParameters< FMarkCoarsePagesCS::FParameters >();
-		PassParameters->MarkingParameters = GetMarkingParameters(GraphBuilder);
+		PassParameters->MarkingParameters = GetMarkingParameters(GraphBuilder, 0);
 		PassParameters->bMarkCoarsePagesLocal = CVarMarkCoarsePagesLocal.GetValueOnRenderThread() != 0 ? 1 : 0;
 		PassParameters->bIncludeNonNaniteGeometry = CVarCoarsePagesIncludeNonNanite.GetValueOnRenderThread();
 
@@ -2641,7 +2641,7 @@ void FVirtualShadowMapArray::BeginMarkPages(
 			//auto PageReceiverMasksUAV = GraphBuilder.CreateUAV(PageReceiverMasksRDG, ERDGUnorderedAccessViewFlags::SkipBarrier);
 
 			FBaseGeneratePageFlagsParameters BaseParameters;
-			BaseParameters.MarkingParameters = GetMarkingParameters(GraphBuilder);
+			BaseParameters.MarkingParameters = GetMarkingParameters(GraphBuilder, ViewIndex);
 			BaseParameters.SceneTexturesStruct = SceneTextures.UniformBuffer;
 			BaseParameters.View = View.ViewUniformBuffer;
 			BaseParameters.ForwardLightStruct = View.ForwardLightingResources.ForwardLightUniformBuffer;
diff --git a/Engine/Source/Runtime/Renderer/Private/VirtualShadowMaps/VirtualShadowMapArray.h b/Engine/Source/Runtime/Renderer/Private/VirtualShadowMaps/VirtualShadowMapArray.h
--- a/Engine/Source/Runtime/Renderer/Private/VirtualShadowMaps/VirtualShadowMapArray.h
+++ b/Engine/Source/Runtime/Renderer/Private/VirtualShadowMaps/VirtualShadowMapArray.h
@@ -422,7 +422,7 @@ public:
 	// This data becomes valid after the shadow depths pass if VSMs are enabled
 	FVirtualShadowMapSamplingParameters GetSamplingParameters(FRDGBuilder& GraphBuilder, int32 ViewIndex) const;
 
-	FVirtualShadowMapMarkingParameters GetMarkingParameters(FRDGBuilder& GraphBuilder, int32 ViewIndex = 0) const;
+	FVirtualShadowMapMarkingParameters GetMarkingParameters(FRDGBuilder& GraphBuilder, int32 ViewIndex) const;
 
 	TRDGUniformBufferRef<FVirtualShadowMapUniformParameters> GetUniformBuffer(int32 ViewIndex) const
 	{
