PR #13529: Use View.Aspects instead of FStereoShaderAspects(..) in PostProcessDeferredDecals.cpp and PostProcessMobile.cpp.

diff --git a/Engine/Source/Runtime/Renderer/Private/CompositionLighting/PostProcessDeferredDecals.cpp b/Engine/Source/Runtime/Renderer/Private/CompositionLighting/PostProcessDeferredDecals.cpp
--- a/Engine/Source/Runtime/Renderer/Private/CompositionLighting/PostProcessDeferredDecals.cpp
+++ b/Engine/Source/Runtime/Renderer/Private/CompositionLighting/PostProcessDeferredDecals.cpp
@@ -255,7 +255,7 @@ void GetDeferredDecalPassParameters(
 
 		const bool bDBufferAProduced = DBufferTextures.DBufferA ? DBufferTextures.DBufferA->HasBeenProduced() : false;
 		const bool bDBufferTexArrayAProduced = DBufferTextures.DBufferATexArray ? DBufferTextures.DBufferATexArray->HasBeenProduced() : false;
-		const bool bUseTextureArrays = View.bIsMobileMultiViewEnabled || UE::StereoRenderUtils::FStereoShaderAspects(View.GetShaderPlatform()).IsMobileMultiViewEnabled();
+		const bool bUseTextureArrays = View.bIsMobileMultiViewEnabled || View.Aspects.IsMobileMultiViewEnabled();
 		const ERenderTargetLoadAction LoadAction = (bUseTextureArrays ? bDBufferTexArrayAProduced : bDBufferAProduced)
 			? ERenderTargetLoadAction::ELoad
 			: ERenderTargetLoadAction::EClear;
diff --git a/Engine/Source/Runtime/Renderer/Private/PostProcess/PostProcessMobile.cpp b/Engine/Source/Runtime/Renderer/Private/PostProcess/PostProcessMobile.cpp
--- a/Engine/Source/Runtime/Renderer/Private/PostProcess/PostProcessMobile.cpp
+++ b/Engine/Source/Runtime/Renderer/Private/PostProcess/PostProcessMobile.cpp
@@ -339,7 +339,7 @@ FMobileBloomSetupOutputs AddMobileBloomSetupPass(FRDGBuilder& GraphBuilder, cons
 		PSShaderParameters->RenderTargets[i] = DestRenderTargets[i].GetRenderTargetBinding();
 	}
 	//if the scenecolor isn't multiview but the app is, need to render as a single-view multiview due to shaders
-	PSShaderParameters->RenderTargets.MultiViewCount = (View.bIsMobileMultiViewEnabled) ? 2 : (UE::StereoRenderUtils::FStereoShaderAspects(View.GetShaderPlatform()).IsMobileMultiViewEnabled() ? 1 : 0);
+	PSShaderParameters->RenderTargets.MultiViewCount = (View.bIsMobileMultiViewEnabled) ? 2 : View.Aspects.IsMobileMultiViewEnabled() ? 1 : 0;
 	
 	PSShaderParameters->EyeAdaptation = EyeAdaptationParameters;
 	PSShaderParameters->BloomThreshold = View.FinalPostProcessSettings.BloomThreshold;
@@ -490,7 +490,7 @@ FScreenPassTexture AddMobileBloomDownPass(FRDGBuilder& GraphBuilder, const FView
 
 	FMobileBloomDownPS::FParameters* PSShaderParameters = GraphBuilder.AllocParameters<FMobileBloomDownPS::FParameters>();
 	PSShaderParameters->RenderTargets[0] = BloomDownOutput.GetRenderTargetBinding();
-	PSShaderParameters->RenderTargets.MultiViewCount = (View.bIsMobileMultiViewEnabled) ? 2 : (UE::StereoRenderUtils::FStereoShaderAspects(View.GetShaderPlatform()).IsMobileMultiViewEnabled() ? 1 : 0);
+	PSShaderParameters->RenderTargets.MultiViewCount = (View.bIsMobileMultiViewEnabled) ? 2 : (View.Aspects.IsMobileMultiViewEnabled() ? 1 : 0);
 	PSShaderParameters->View = View.GetShaderParameters();
 
 	PSShaderParameters->BloomDownSourceTexture = Inputs.BloomDownSource.Texture;
@@ -618,7 +618,7 @@ FScreenPassTexture AddMobileBloomUpPass(FRDGBuilder& GraphBuilder, const FViewIn
 
 	FMobileBloomUpPS::FParameters* PSShaderParameters = GraphBuilder.AllocParameters<FMobileBloomUpPS::FParameters>();
 	PSShaderParameters->RenderTargets[0] = BloomUpOutput.GetRenderTargetBinding();
-	PSShaderParameters->RenderTargets.MultiViewCount = (View.bIsMobileMultiViewEnabled) ? 2 : (UE::StereoRenderUtils::FStereoShaderAspects(View.GetShaderPlatform()).IsMobileMultiViewEnabled() ? 1 : 0);
+	PSShaderParameters->RenderTargets.MultiViewCount = (View.bIsMobileMultiViewEnabled) ? 2 : (View.Aspects.IsMobileMultiViewEnabled() ? 1 : 0);
 	PSShaderParameters->View = View.GetShaderParameters();
 
 	PSShaderParameters->BloomUpSourceATexture = Inputs.BloomUpSourceA.Texture;
@@ -1179,7 +1179,7 @@ FScreenPassTexture AddMobileSunMergePass(FRDGBuilder& GraphBuilder, const FViewI
 
 	FMobileSunMergePS::FParameters* PSShaderParameters = GraphBuilder.AllocParameters<FMobileSunMergePS::FParameters>();
 	PSShaderParameters->RenderTargets[0] = SunMergeOutput.GetRenderTargetBinding();
-	PSShaderParameters->RenderTargets.MultiViewCount = (View.bIsMobileMultiViewEnabled) ? 2 : (UE::StereoRenderUtils::FStereoShaderAspects(View.GetShaderPlatform()).IsMobileMultiViewEnabled() ? 1 : 0);
+	PSShaderParameters->RenderTargets.MultiViewCount = (View.bIsMobileMultiViewEnabled) ? 2 : (View.Aspects.IsMobileMultiViewEnabled() ? 1 : 0);
 	PSShaderParameters->View = View.GetShaderParameters();
 	PSShaderParameters->BloomDirtMaskTint = Settings.BloomDirtMaskTint * Settings.BloomDirtMaskIntensity;
 	PSShaderParameters->SunColorVignetteIntensity = SunColorVignetteIntensityParam;
