Disabling shadow-casting while addressing performance regression

diff --git a/Engine/Source/Runtime/Renderer/Private/HeterogeneousVolumes/HeterogeneousVolumes.cpp b/Engine/Source/Runtime/Renderer/Private/HeterogeneousVolumes/HeterogeneousVolumes.cpp
--- a/Engine/Source/Runtime/Renderer/Private/HeterogeneousVolumes/HeterogeneousVolumes.cpp
+++ b/Engine/Source/Runtime/Renderer/Private/HeterogeneousVolumes/HeterogeneousVolumes.cpp
@@ -21,8 +21,8 @@ static TAutoConsoleVariable<int32> CVarHeterogeneousVolumes(
 
 static TAutoConsoleVariable<int32> CVarHeterogeneousVolumesShadows(
 	TEXT("r.HeterogeneousVolumes.Shadows"),
-	1,
-	TEXT("Enables heterogeneous volume-casting shadows (default = 1)"),
+	0,
+	TEXT("Enables heterogeneous volume-casting shadows (default = 0)"),
 	ECVF_RenderThreadSafe
 );
 
