Rolled back 44168721 causing visual artefacts in TAA on mobile renderer

#rnx
#FailedValidations Preflight

diff --git a/Engine/Source/Runtime/Renderer/Private/PostProcess/PostProcessing.cpp b/Engine/Source/Runtime/Renderer/Private/PostProcess/PostProcessing.cpp
--- a/Engine/Source/Runtime/Renderer/Private/PostProcess/PostProcessing.cpp
+++ b/Engine/Source/Runtime/Renderer/Private/PostProcess/PostProcessing.cpp
@@ -2570,9 +2570,8 @@ void AddMobilePostProcessingPasses(FRDGBuilder& GraphBuilder, FScene* Scene, con
 
 	bool bUseHighResolutionScreenshotMask = IsHighResolutionScreenshotMaskEnabled(View);
 
-	// Set bShouldPrimaryUpscale like Desktop Renderer
-	// compare size with Secondary, not Unscaled
-	bool bShouldPrimaryUpscale = (View.PrimaryScreenPercentageMethod == EPrimaryScreenPercentageMethod::SpatialUpscale && View.GetSecondaryViewRectSize() != View.ViewRect.Size()) || View.LensDistortionLUT.IsEnabled();
+	bool bShouldPrimaryUpscale = (View.PrimaryScreenPercentageMethod == EPrimaryScreenPercentageMethod::SpatialUpscale && View.UnscaledViewRect != View.ViewRect) || View.LensDistortionLUT.IsEnabled();
+	bShouldPrimaryUpscale |= View.Family->GetPrimarySpatialUpscalerInterface() != nullptr;
 
 	PassSequence.SetEnabled(EPass::Tonemap, bUseToneMapper);
 	PassSequence.SetEnabled(EPass::HighResolutionScreenshotMask, bUseHighResolutionScreenshotMask);
@@ -2588,10 +2587,7 @@ void AddMobilePostProcessingPasses(FRDGBuilder& GraphBuilder, FScene* Scene, con
 
 	PassSequence.SetEnabled(EPass::DebugPrimitive, IsDebugPrimitivePassEnabled(View));
 	PassSequence.SetEnabled(EPass::PrimaryUpscale, bShouldPrimaryUpscale);
-	// Match Desktop Renderer, calling View.RequiresSecondaryUpscale()
-	// Both Desktop and Mobile Renderer call Third-party if enabled (no matter what scaling).
-	// Is it really necessary to call the Secondary Spatial if the scaling is 100% ?
-	PassSequence.SetEnabled(EPass::SecondaryUpscale, View.RequiresSecondaryUpscale() || View.Family->GetSecondarySpatialUpscalerInterface() != nullptr);
+	PassSequence.SetEnabled(EPass::SecondaryUpscale, View.Family->GetSecondarySpatialUpscalerInterface() != nullptr);
 
 	PassSequence.SetEnabled(EPass::Visualize, View.Family->EngineShowFlags.ShaderComplexity);
 	PassSequence.SetEnabled(EPass::VisualizeLightGrid, ShouldVisualizeLightGrid(View.Family->GetShaderPlatform()));
@@ -3340,7 +3336,7 @@ void AddMobilePostProcessingPasses(FRDGBuilder& GraphBuilder, FScene* Scene, con
 	{
 		ISpatialUpscaler::FInputs PassInputs;
 		PassSequence.AcceptOverrideIfLastPass(EPass::PrimaryUpscale, PassInputs.OverrideOutput);
-		PassInputs.Stage = PassSequence.IsEnabled(EPass::SecondaryUpscale) ? EUpscaleStage::PrimaryToSecondary : EUpscaleStage::PrimaryToOutput;
+		PassInputs.Stage = EUpscaleStage::PrimaryToOutput;
 		PassInputs.SceneColor = SceneColor;
 
 		if (const ISpatialUpscaler* CustomUpscaler = View.Family->GetPrimarySpatialUpscalerInterface())
@@ -3360,7 +3356,7 @@ void AddMobilePostProcessingPasses(FRDGBuilder& GraphBuilder, FScene* Scene, con
 			}
 			else
 			{
-				check(SceneColor.ViewRect.Size() == View.GetSecondaryViewRectSize());
+				check(SceneColor.ViewRect.Size() == View.UnscaledViewRect.Size());
 			}
 		}
 		else
