Mobile: Fixed case where relfection env shader was not correctly selecting permutations for a specific lighting conditions

diff --git a/Engine/Source/Runtime/Renderer/Private/MobileDeferredShadingPass.cpp b/Engine/Source/Runtime/Renderer/Private/MobileDeferredShadingPass.cpp
--- a/Engine/Source/Runtime/Renderer/Private/MobileDeferredShadingPass.cpp
+++ b/Engine/Source/Runtime/Renderer/Private/MobileDeferredShadingPass.cpp
@@ -483,6 +483,7 @@ void RenderReflectionEnvironmentSkyLighting(
 		GET_STENCIL_MOBILE_SM_MASK(0xff), 0x00>::GetRHI();
 	uint8 StencilRef = PassShadingModelStencilValue(/*bEnableShadingModelSupport=*/false);
 
+	bool bEnableShadingModelSupport = false;
 	if (!IsOnlyDefaultLitShadingModel(View.ShadingModelMaskInView) && MobileUsesGBufferCustomData(Scene.GetShaderPlatform()))
 	{
 		GraphicsPSOInit.DepthStencilState = TStaticDepthStencilState<
@@ -492,10 +493,11 @@ void RenderReflectionEnvironmentSkyLighting(
 			GET_STENCIL_MOBILE_SM_MASK(0xff), 0x00>::GetRHI();
 		// Apply all shading models
 		StencilRef = 0;
+		bEnableShadingModelSupport = true;
 	}
 
 	FMobileReflectionEnvironmentSkyLightingPS::FPermutationDomain PermutationVector;
-	PermutationVector.Set<FMobileReflectionEnvironmentSkyLightingPS::FEnableShadingModelSupport>(!IsOnlyDefaultLitShadingModel(View.ShadingModelMaskInView));
+	PermutationVector.Set<FMobileReflectionEnvironmentSkyLightingPS::FEnableShadingModelSupport>(bEnableShadingModelSupport);
 	PermutationVector.Set<FMobileReflectionEnvironmentSkyLightingPS::FEnableClustredReflection>(bClustredReflection);
 	PermutationVector.Set<FMobileReflectionEnvironmentSkyLightingPS::FEnablePlanarReflection>(bPlanarReflection);
 	PermutationVector.Set<FMobileReflectionEnvironmentSkyLightingPS::FEnableSkyLight>(bEnableSkyLight);
