Fix VolumetricFog for Mobile
- Remove the extra Falloff from IntegrateLight, on mobile we already have it from GetLocalLightAttenuation
- Set GVolumetricFogJitter, GVolumetricFogTemporalReprojection and GVolumetricFogInjectShadowedLightsSeparately to 0 on mobile

diff --git a/Engine/Shaders/Private/VolumetricFog.usf b/Engine/Shaders/Private/VolumetricFog.usf
--- a/Engine/Shaders/Private/VolumetricFog.usf
+++ b/Engine/Shaders/Private/VolumetricFog.usf
@@ -976,8 +976,11 @@ void LightScatteringCS(
 			if (VolumetricScatteringIntensity > 0)
 			{
 				const bool bSoftFadeEnabled = VolumetricFog.LightSoftFading > 0;
-
+#if SHADING_PATH_MOBILE
+				const FDeferredLightData LightData = ConvertToDeferredLight_Mobile(LocalLight);
+#else
 				const FDeferredLightData LightData = ConvertToDeferredLight(LocalLight);
+#endif
 
 				float3 L = 0;
 				float3 ToLight = 0;
@@ -1011,7 +1014,11 @@ void LightScatteringCS(
 					}
 #endif // USE_LIGHT_SOFT_FADING
 
-					Lighting = SofFade * IntegrateLight(Capsule, LightData.bInverseSquared);
+					Lighting = SofFade;
+#if !SHADING_PATH_MOBILE
+					// Mobile is already getting the Falloff from GetLocalLightAttenuation
+					Lighting *= IntegrateLight(Capsule, LightData.bInverseSquared);
+#endif	
 				}
 
 				float CombinedAttenuation = Lighting * LightMask;
diff --git a/Engine/Source/Runtime/Renderer/Private/VolumetricFog.cpp b/Engine/Source/Runtime/Renderer/Private/VolumetricFog.cpp
--- a/Engine/Source/Runtime/Renderer/Private/VolumetricFog.cpp
+++ b/Engine/Source/Runtime/Renderer/Private/VolumetricFog.cpp
@@ -381,6 +381,11 @@ class FInjectShadowedLocalLightPS : public FGlobalShader
 		FVirtualShadowMapArray::SetShaderDefines(OutEnvironment);
 	}
 	
+	static bool ShouldCompilePermutation(const FGlobalShaderPermutationParameters& Parameters)
+	{
+		return !IsMobilePlatform(Parameters.Platform);
+	}
+
 	static EShaderPermutationPrecacheRequest ShouldPrecachePermutation(const FGlobalShaderPermutationParameters& Parameters)
 	{
 		FPermutationDomain PermutationVector(Parameters.PermutationId);
@@ -1386,7 +1391,14 @@ void FSceneRenderer::ComputeVolumetricFog(FRDGBuilder& GraphBuilder,
 	{
 		return;
 	}
-
+	const bool bIsMobilePlatform = IsMobileVolumetricFog(Scene->GetShaderPlatform());
+	// FVolumetricFogLightScatteringCS::FTemporalReprojection permutation is disabled on mobile
+	if (bIsMobilePlatform)
+	{
+		GVolumetricFogJitter = 0;
+		GVolumetricFogTemporalReprojection = 0;
+		GVolumetricFogInjectShadowedLightsSeparately = 0;
+	}
 	const FExponentialHeightFogSceneInfo& FogInfo = Scene->ExponentialFogs[0];
 
 	TRACE_CPUPROFILER_EVENT_SCOPE(FSceneRenderer::ComputeVolumetricFog);
@@ -1498,15 +1510,13 @@ void FSceneRenderer::ComputeVolumetricFog(FRDGBuilder& GraphBuilder,
 			IntegrationData.FrameJitterOffsetValues[FrameOffsetIndex] = VolumetricFogTemporalRandom(View.Family->FrameNumber - FrameOffsetIndex);
 		}
 
-		const bool bIsMobilePlatform = IsMobileVolumetricFog(View.GetShaderPlatform());
 
 		// Mobile has limited capacities with SRV binding so do not enable atlas sampling on there.
 		const bool bUseLightFunctionAtlasEnabledAndSupported = LightFunctionAtlas::IsEnabled(*Scene, ELightFunctionAtlasSystem::VolumetricFog) && !bIsMobilePlatform;
 
 		const bool bUseTemporalReprojection =
 			GVolumetricFogTemporalReprojection
-			&& View.ViewState
-			&& !bIsMobilePlatform;
+			&& View.ViewState;
 
 		IntegrationData.bTemporalHistoryIsValid =
 			bUseTemporalReprojection
