Marked r.UseClusteredDeferredShading and r.Mobile.UseClusteredDeferredShading as deprecated and added notification about future removal. Clustered deferred shading will be removed in a future release due to lack of utility to reduce maintenance burden.

[RN]

diff --git a/Engine/Source/Runtime/Renderer/Private/ClusteredDeferredShadingPass.cpp b/Engine/Source/Runtime/Renderer/Private/ClusteredDeferredShadingPass.cpp
--- a/Engine/Source/Runtime/Renderer/Private/ClusteredDeferredShadingPass.cpp
+++ b/Engine/Source/Runtime/Renderer/Private/ClusteredDeferredShadingPass.cpp
@@ -31,11 +31,14 @@
 #include "LightFunctionAtlas.h"
 #include "AnisotropyRendering.h"
 
+static FAutoConsoleVariableDeprecated CVarUseClusteredDeferredShadingDep(TEXT("r.UseClusteredDeferredShading"), TEXT("r.UseClusteredDeferredShading_ToBeRemoved"), TEXT("5.7"));
+
 // This is used to switch on and off the clustered deferred shading implementation, that uses the light grid to perform shading.
 int32 GUseClusteredDeferredShading = 0;
 static FAutoConsoleVariableRef CVarUseClusteredDeferredShading(
-	TEXT("r.UseClusteredDeferredShading"),
+	TEXT("r.UseClusteredDeferredShading_ToBeRemoved"),
 	GUseClusteredDeferredShading,
+	TEXT("NOTE: The clustered deferred shading implementation will be removed in a future release due to low utility and thus use.\n")
 	TEXT("Toggle use of clustered deferred shading for lights that support it. 0 is off (default), 1 is on (also required is SM5 to actually turn on)."),
 	ECVF_RenderThreadSafe
 );
diff --git a/Engine/Source/Runtime/Renderer/Private/MobileDeferredShadingPass.cpp b/Engine/Source/Runtime/Renderer/Private/MobileDeferredShadingPass.cpp
--- a/Engine/Source/Runtime/Renderer/Private/MobileDeferredShadingPass.cpp
+++ b/Engine/Source/Runtime/Renderer/Private/MobileDeferredShadingPass.cpp
@@ -23,10 +23,13 @@ DECLARE_GPU_STAT(DeferredShading);
 
 static const TCHAR* DeferredMobileLightMaterialPSOCollectorName = TEXT("DeferredMobileLightMaterialPSOCollector");
 
+static FAutoConsoleVariableDeprecated CVarMobileUseClusteredDeferredShadingDep(TEXT("r.Mobile.UseClusteredDeferredShading"), TEXT("r.Mobile.UseClusteredDeferredShading_ToBeRemoved"), TEXT("5.7"));
+
 int32 GMobileUseClusteredDeferredShading = 0;
 static FAutoConsoleVariableRef CVarMobileUseClusteredDeferredShading(
-	TEXT("r.Mobile.UseClusteredDeferredShading"),
+	TEXT("r.Mobile.UseClusteredDeferredShading_ToBeRemoved"),
 	GMobileUseClusteredDeferredShading,
+	TEXT("NOTE: The mobile clustered deferred shading implementation will be removed in a future release due to low utility and use.\n")
 	TEXT("Toggle use of clustered deferred shading for lights that support it. 0 is off (default), 1 is on. (requires LightGrid: r.Mobile.Forward.EnableLocalLights=1)"),
 	ECVF_RenderThreadSafe
 );
