Mobile: Fixed missing translucency with a Lighting Mode set to "Surface"

diff --git a/Engine/Source/Runtime/Renderer/Private/LightMapRendering.cpp b/Engine/Source/Runtime/Renderer/Private/LightMapRendering.cpp
--- a/Engine/Source/Runtime/Renderer/Private/LightMapRendering.cpp
+++ b/Engine/Source/Runtime/Renderer/Private/LightMapRendering.cpp
@@ -158,8 +158,8 @@ void FMobileDirectionalLightAndCSMPolicy::ModifyCompilationEnvironment(const FMa
 
 bool FMobileDirectionalLightAndCSMPolicy::ShouldCompilePermutation(const FMeshMaterialShaderPermutationParameters& Parameters)
 {
-	if (IsMobileDeferredShadingEnabled(Parameters.Platform)
-		&& !Parameters.MaterialParameters.ShadingModels.HasShadingModel(MSM_SingleLayerWater))
+	if (IsMobileDeferredShadingEnabled(Parameters.Platform) && 
+		!(Parameters.MaterialParameters.ShadingModels.HasShadingModel(MSM_SingleLayerWater) || Parameters.MaterialParameters.bIsTranslucentSurfaceLighting))
 	{
 		return false;
 	}
@@ -209,7 +209,8 @@ void FMobileDistanceFieldShadowsLightMapAndCSMLightingPolicy::ModifyCompilationE
 
 bool FMobileDirectionalLightCSMAndLightMapPolicy::ShouldCompilePermutation(const FMeshMaterialShaderPermutationParameters& Parameters)
 {
-	if (IsMobileDeferredShadingEnabled(Parameters.Platform))
+	if (IsMobileDeferredShadingEnabled(Parameters.Platform) && 
+		!Parameters.MaterialParameters.bIsTranslucentSurfaceLighting)
 	{
 		return false;
 	}
diff --git a/Engine/Source/Runtime/Renderer/Private/MobileBasePassRendering.cpp b/Engine/Source/Runtime/Renderer/Private/MobileBasePassRendering.cpp
--- a/Engine/Source/Runtime/Renderer/Private/MobileBasePassRendering.cpp
+++ b/Engine/Source/Runtime/Renderer/Private/MobileBasePassRendering.cpp
@@ -102,7 +102,7 @@ EMobileTranslucentColorTransmittanceMode MobileDefaultTranslucentColorTransmitta
 	{
 		return EMobileTranslucentColorTransmittanceMode::DUAL_SRC_BLENDING;
 	}
-	if (IsMetalMobilePlatform(Platform) || IsAndroidOpenGLESPlatform(Platform))
+	if ((IsMetalMobilePlatform(Platform) || IsAndroidOpenGLESPlatform(Platform)) && MobileAllowFramebufferFetch(Platform))
 	{
 		return EMobileTranslucentColorTransmittanceMode::PROGRAMMABLE_BLENDING;
 	}
