Fix VolumetricFog for MobileSupportsSM5MaterialNodes platforms by not excluding FInjectShadowedLocalLightPS permutations in that case

diff --git a/Engine/Source/Runtime/Renderer/Private/VolumetricFog.cpp b/Engine/Source/Runtime/Renderer/Private/VolumetricFog.cpp
--- a/Engine/Source/Runtime/Renderer/Private/VolumetricFog.cpp
+++ b/Engine/Source/Runtime/Renderer/Private/VolumetricFog.cpp
@@ -383,7 +383,7 @@ class FInjectShadowedLocalLightPS : public FGlobalShader
 	
 	static bool ShouldCompilePermutation(const FGlobalShaderPermutationParameters& Parameters)
 	{
-		return !IsMobilePlatform(Parameters.Platform);
+		return !IsMobileVolumetricFog(Parameters.Platform);
 	}
 
 	static EShaderPermutationPrecacheRequest ShouldPrecachePermutation(const FGlobalShaderPermutationParameters& Parameters)
