Heterogeneous Volumes: Prevent dangling resource bindings when disabling shadow-casting at run-time.

diff --git a/Engine/Source/Runtime/Renderer/Private/HeterogeneousVolumes/HeterogeneousVolumes.cpp b/Engine/Source/Runtime/Renderer/Private/HeterogeneousVolumes/HeterogeneousVolumes.cpp
--- a/Engine/Source/Runtime/Renderer/Private/HeterogeneousVolumes/HeterogeneousVolumes.cpp
+++ b/Engine/Source/Runtime/Renderer/Private/HeterogeneousVolumes/HeterogeneousVolumes.cpp
@@ -1715,6 +1715,7 @@ namespace HeterogeneousVolumes
 					DestroyAdaptiveVolumetricShadowMapUniformBuffer(Pair.Value);
 				}
 				ViewState->AdaptiveVolumetricShadowMapUniformBufferMap.Empty();
+				ViewState->BeerShadowMapUniformBufferMap.Empty();
 			}
 		}
 	}
