Fix bug when matching desktop and mobile upscaler interface where the secondary upscaler is enabled but not upscaled due to interface not available.

#tests fix preview artifact when mobile TAA is enabled

diff --git a/Engine/Source/Runtime/Renderer/Private/PostProcess/PostProcessing.cpp b/Engine/Source/Runtime/Renderer/Private/PostProcess/PostProcessing.cpp
--- a/Engine/Source/Runtime/Renderer/Private/PostProcess/PostProcessing.cpp
+++ b/Engine/Source/Runtime/Renderer/Private/PostProcess/PostProcessing.cpp
@@ -2570,8 +2570,8 @@ void AddMobilePostProcessingPasses(FRDGBuilder& GraphBuilder, FScene* Scene, con
 
 	bool bUseHighResolutionScreenshotMask = IsHighResolutionScreenshotMaskEnabled(View);
 
-	bool bShouldPrimaryUpscale = (View.PrimaryScreenPercentageMethod == EPrimaryScreenPercentageMethod::SpatialUpscale && View.UnscaledViewRect != View.ViewRect) || View.LensDistortionLUT.IsEnabled();
-	bShouldPrimaryUpscale |= View.Family->GetPrimarySpatialUpscalerInterface() != nullptr;
+	//Primary does not require the interface to be available to run as it has a fallback logic already.
+	bool bShouldPrimaryUpscale = (View.PrimaryScreenPercentageMethod == EPrimaryScreenPercentageMethod::SpatialUpscale && View.GetSecondaryViewRectSize() != View.ViewRect.Size()) || View.LensDistortionLUT.IsEnabled();
 
 	PassSequence.SetEnabled(EPass::Tonemap, bUseToneMapper);
 	PassSequence.SetEnabled(EPass::HighResolutionScreenshotMask, bUseHighResolutionScreenshotMask);
@@ -2587,6 +2587,8 @@ void AddMobilePostProcessingPasses(FRDGBuilder& GraphBuilder, FScene* Scene, con
 
 	PassSequence.SetEnabled(EPass::DebugPrimitive, IsDebugPrimitivePassEnabled(View));
 	PassSequence.SetEnabled(EPass::PrimaryUpscale, bShouldPrimaryUpscale);
+	// Mobile is different from the deferred that the secondary upscale logic only runs when the interface is available
+	// While the deferred itself has a fallback when secondary upscaler interface is not present.
 	PassSequence.SetEnabled(EPass::SecondaryUpscale, View.Family->GetSecondarySpatialUpscalerInterface() != nullptr);
 
 	PassSequence.SetEnabled(EPass::Visualize, View.Family->EngineShowFlags.ShaderComplexity);
@@ -3336,7 +3338,7 @@ void AddMobilePostProcessingPasses(FRDGBuilder& GraphBuilder, FScene* Scene, con
 	{
 		ISpatialUpscaler::FInputs PassInputs;
 		PassSequence.AcceptOverrideIfLastPass(EPass::PrimaryUpscale, PassInputs.OverrideOutput);
-		PassInputs.Stage = EUpscaleStage::PrimaryToOutput;
+		PassInputs.Stage = PassSequence.IsEnabled(EPass::SecondaryUpscale) ? EUpscaleStage::PrimaryToSecondary : EUpscaleStage::PrimaryToOutput;
 		PassInputs.SceneColor = SceneColor;
 
 		if (const ISpatialUpscaler* CustomUpscaler = View.Family->GetPrimarySpatialUpscalerInterface())
@@ -3356,7 +3358,7 @@ void AddMobilePostProcessingPasses(FRDGBuilder& GraphBuilder, FScene* Scene, con
 			}
 			else
 			{
-				check(SceneColor.ViewRect.Size() == View.UnscaledViewRect.Size());
+				check(SceneColor.ViewRect.Size() == View.GetSecondaryViewRectSize());
 			}
 		}
 		else
