Fix hair Gbuffer buffer write with Substrate.

diff --git a/Engine/Shaders/Private/HairStrands/HairStrandsComposition.usf b/Engine/Shaders/Private/HairStrands/HairStrandsComposition.usf
--- a/Engine/Shaders/Private/HairStrands/HairStrandsComposition.usf
+++ b/Engine/Shaders/Private/HairStrands/HairStrandsComposition.usf
@@ -303,34 +303,56 @@ void DOFDepthPS(
 
 #include "HairStrandsGBufferCommon.ush"
 
+#if SUBSTRATE_ENABLED && SUBSTRATE_GBUFFER_FORMAT==1
+#define SUBSTRATE_INLINE_SHADING 1
+#define MATERIAL_SUBSTRATE_OPAQUE_PRECOMPUTED_LIGHTING 0
+#include "../Substrate/Substrate.ush"
+#include "../Substrate/SubstrateExport.ush"
+uint MaxBytesPerPixel;
+RWTexture2D<SUBSTRATE_TOP_LAYER_TYPE> RWTopLayerTexture;
+RWTexture2DArray<uint> RWMaterialTextureArray;
+#endif
+
 uint bWriteDummyData;
 
 void MainPS(
 	in float2 UV : TEXCOORD
 	, in float4 SvPosition: SV_Position
+#if SUBSTRATE_ENABLED && SUBSTRATE_GBUFFER_FORMAT==1
+  #if PERMUTATION_OUTPUT_TYPE == 1
+	, out float  OutDepth    : SV_DEPTH
+  #endif
+#else
 	, out float4 OutGBufferA : SV_Target0
 	, out float4 OutGBufferB : SV_Target1
-#if PERMUTATION_OUTPUT_TYPE == 1
+  #if PERMUTATION_OUTPUT_TYPE == 1
 	, out float4 OutGBufferC : SV_Target2
 	, out float4 OutGBufferD : SV_Target3
 	, out float4 OutGBufferE : SV_Target4
 	, out float  OutDepth    : SV_DEPTH
+  #endif
 #endif
 )
 {
 	ResolvedView = ResolveView();
+#if SUBSTRATE_ENABLED && SUBSTRATE_GBUFFER_FORMAT==1
+  #if PERMUTATION_OUTPUT_TYPE == 1
+	OutDepth = 0;
+  #endif
+#else
 	OutGBufferA = 0;
 	OutGBufferB = 0;
-#if PERMUTATION_OUTPUT_TYPE == 0
+  #if PERMUTATION_OUTPUT_TYPE == 0
 	float4 OutGBufferC = 0;
 	float4 OutGBufferD = 0;
 	float4 OutGBufferE = 0;
 	float  OutDepth = 0;
-#elif PERMUTATION_OUTPUT_TYPE == 1
+  #elif PERMUTATION_OUTPUT_TYPE == 1
 	OutGBufferC = 0;
 	OutGBufferD = 0;
 	OutGBufferE = 0;
 	OutDepth = 0;
+  #endif
 #endif
 
 	const uint3 PixelCoord = uint3(floor(SvPosition.xy), 0);
@@ -347,6 +369,7 @@ void MainPS(
 		discard;
 	}
 
+	#if SUBSTRATE_ENABLED == 0 || SUBSTRATE_GBUFFER_FORMAT==0
 	if (bWriteDummyData > 0)
 	{
 		const float3 Tangent = float3(0, 0, 1);
@@ -354,11 +377,12 @@ void MainPS(
 		OutGBufferB = float4(0, 0, 0, EncodeShadingModelIdAndSelectiveOutputMask(SHADINGMODELID_HAIR, 0));
 	}
 	else
+	#endif
 	{
 		float3 Samples_Tangent = 0;
 		float  Samples_Specular = 0;
 		float3 Samples_BaseColor = 0;
-		float  Samples_Roughnesss = 0;
+		float  Samples_Roughness = 0;
 		uint   Samples_LightChannelMask = 0;
 		float  Samples_Backlit = 0;
 		float  Samples_Depth = 0;
@@ -372,7 +396,7 @@ void MainPS(
 			Samples_Tangent		+= Sample.Tangent;
 			Samples_Specular	+= Sample.Specular;
 			Samples_BaseColor	+= Sample.BaseColor;
-			Samples_Roughnesss	+= Sample.Roughness;
+			Samples_Roughness	+= Sample.Roughness;
 			Samples_Backlit		+= Sample.Backlit;
 			if (NodeIt == 0)
 			{
@@ -385,14 +409,47 @@ void MainPS(
 		Samples_Tangent		 = normalize(Samples_Tangent);
 		Samples_Specular	*= InvSampleCount;
 		Samples_BaseColor	*= InvSampleCount;
-		Samples_Roughnesss	*= InvSampleCount;
+		Samples_Roughness	*= InvSampleCount;
 		Samples_Backlit		*= InvSampleCount;
 
+	#if SUBSTRATE_ENABLED && SUBSTRATE_GBUFFER_FORMAT==1
+		const bool bWriteHairData = bWriteDummyData == 0;
+		FSubstrateBSDF HairBSDF = GetSubstrateHairBSDF(Samples_BaseColor, 0.f /*Scatter*/, Samples_Specular, Samples_Roughness, Samples_Backlit, 0 /*Emissive*/, 0/*SharedLocalBasisIndex*/).InlinedBSDF;
+
+		FSubstrateAddressing SubstrateAddressing = GetSubstratePixelDataByteOffset(PixelCoord.xy, uint2(ResolvedView.BufferSizeAndInvSize.xy), MaxBytesPerPixel);
+		FSubstratePixelHeader SubstratePixelHeader = InitialiseSubstratePixelHeader();
+		SubstratePixelHeader.ClosureCount = 1;
+		SubstratePixelHeader.SharedLocalBases.Normals[0] = normalize(Samples_Tangent);
+
+		FSubstrateTopLayerData TopLayerData = (FSubstrateTopLayerData)0;
+		FRWSubstrateMaterialContainer RWSubstrateMaterialContainer = InitialiseRWSubstrateMaterialContainer(RWMaterialTextureArray /*Should normally be WithoutRTs, but here writes are explicit*/, 1/*QuadPixelWriteMask*/, PixelCoord.xy);
+		WriteHair(
+			SubstratePixelHeader,
+			SubstrateAddressing,
+			HairBSDF,
+			bWriteHairData,
+			TopLayerData,
+			RWSubstrateMaterialContainer,
+			RWMaterialTextureArray);
+
+		RWTopLayerTexture[SubstrateAddressing.PixelCoords] = SubstratePackTopLayerData(TopLayerData);
+		RWMaterialTextureArray[uint3(SubstrateAddressing.PixelCoords,0)] = RWSubstrateMaterialContainer.MaterialRenderTargets[0]; // Header
+		if (bWriteHairData)
+		{
+			RWMaterialTextureArray[uint3(SubstrateAddressing.PixelCoords,1)] = RWSubstrateMaterialContainer.MaterialRenderTargets[1];
+			RWMaterialTextureArray[uint3(SubstrateAddressing.PixelCoords,2)] = RWSubstrateMaterialContainer.MaterialRenderTargets[2];
+		}
+
+		#if PERMUTATION_OUTPUT_TYPE == 1
+		OutDepth = Samples_Depth;
+		#endif
+
+	#else
 		WriteGBuffer(
 			Samples_Tangent,
 			Samples_Specular,
 			Samples_BaseColor,
-			Samples_Roughnesss,
+			Samples_Roughness,
 			Samples_LightChannelMask,
 			Samples_Backlit,
 			Samples_Depth,
@@ -402,6 +459,7 @@ void MainPS(
 			OutGBufferD,
 			OutGBufferE,
 			OutDepth);
+	#endif
 	}
 }
 
diff --git a/Engine/Shaders/Private/Substrate/SubstrateExport.ush b/Engine/Shaders/Private/Substrate/SubstrateExport.ush
--- a/Engine/Shaders/Private/Substrate/SubstrateExport.ush
+++ b/Engine/Shaders/Private/Substrate/SubstrateExport.ush
@@ -906,6 +906,53 @@ void PackSubstrateOut(
 
 }
 
+// Trimmed version of PackSubstrateOut() specialized for hair data writting. This is used by hair strands composition. 
+// todo: This should be refactored to avoid duplicate code.
+void WriteHair(
+	FSubstratePixelHeader SubstratePixelHeader,
+	FSubstrateAddressing SubstrateAddressing,
+	FSubstrateBSDF BSDF,
+	bool bWriteHairData,
+	inout FSubstrateTopLayerData TopLayerData,
+	inout FRWSubstrateMaterialContainer SubstrateBuffer,
+	inout RWTexture2DArray<uint> ExtraMaterialDataUAV)
+{
+	// Top layer data
+	{
+		const float3x3 TangentBasis = SubstrateGetBSDFSharedBasis_InlineShading(SubstratePixelHeader, BSDF);
+		const float3 N = TangentBasis[2];
+
+		TopLayerData = (FSubstrateTopLayerData)0;
+		TopLayerData.Material = TOP_LAYER_MATERIAL_VALID;
+		TopLayerData.WorldNormal = N;
+		TopLayerData.Roughness = HAIR_ROUGHNESS(BSDF);
+	}
+	// Header
+	{
+		SubstratePixelHeader.SetHasSubsurface(false);
+		SubstratePixelHeader.SetMaterialMode(HEADER_MATERIALMODE_HAIR);
+		SubstratePixelHeader.ClosureCount = 1;
+	}
+	{
+		// Data0 (Header_State|Header_AO|Data)
+		uint Out = 0;
+		HEADER_SETCOMMONSTATES(Out, SubstratePixelHeader.State);
+		Out = Out & HEADER_HAIRENCODING_MASK;
+		SUBSTRATE_STORE_UINT1(Out);
+
+		#if (HEADER_HAIRENCODING_BIT_COUNT) > 32
+		#error Substrate hair path header is > 32bits
+		#endif
+	}
+	// Data
+	if (bWriteHairData)
+	{
+		SUBSTRATE_STORE_UINT1(PackColorLinearToGamma2AlphaLinear(float4(HAIR_BASECOLOR(BSDF), HAIR_ROUGHNESS(BSDF))));
+		SUBSTRATE_STORE_UINT1(PackRGBA8(float4(HAIR_SCATTER(BSDF), HAIR_SPECULAR(BSDF), HAIR_BACKLIT(BSDF), HAIR_COMPLEXTRANSMITTANCE(BSDF))));
+		// 8 bytes
+	}
+}
+
 #endif // SUBSTRATE_ENABLED && !defined(FMobileBasePassInterpolantsVSToPS)
 
 #ifndef SUBSTRATE_GPU_LIGHTMASS
diff --git a/Engine/Source/Runtime/Renderer/Private/HairStrands/HairStrandsComposition.cpp b/Engine/Source/Runtime/Renderer/Private/HairStrands/HairStrandsComposition.cpp
--- a/Engine/Source/Runtime/Renderer/Private/HairStrands/HairStrandsComposition.cpp
+++ b/Engine/Source/Runtime/Renderer/Private/HairStrands/HairStrandsComposition.cpp
@@ -505,6 +505,10 @@ class FHairVisibilityGBufferWritePS : public FGlobalShader
 		SHADER_PARAMETER_STRUCT_INCLUDE(FHairStrandsTilePassVS::FParameters, TileData)
 		SHADER_PARAMETER(uint32, bWriteDummyData)
 		SHADER_PARAMETER_RDG_UNIFORM_BUFFER(FHairStrandsViewUniformParameters, HairStrands)
+		// Substrate
+		SHADER_PARAMETER_STRUCT_INCLUDE(FSubstrateCommonParameters, Substrate)
+		SHADER_PARAMETER_RDG_TEXTURE_UAV(RWTexture2DARRAY, RWMaterialTextureArray)
+		SHADER_PARAMETER_RDG_TEXTURE_UAV(RWTexture2D, RWTopLayerTexture)
 		RENDER_TARGET_BINDING_SLOTS()
 	END_SHADER_PARAMETER_STRUCT()
 
@@ -530,6 +534,57 @@ public:
 
 IMPLEMENT_GLOBAL_SHADER(FHairVisibilityGBufferWritePS, "/Engine/Private/HairStrands/HairStrandsComposition.usf", "MainPS", SF_Pixel);
 
+namespace Substrate
+{
+	FSubstrateCommonParameters GetSubstrateCommonParameter(const FSubstrateSceneData& In);
+}
+
+static void AddHairVisibilityGBufferWritePass(
+	FRDGBuilder& GraphBuilder,
+	const FViewInfo& View,
+	const bool bWriteDummyData,
+	const FHairStrandsTiles& TileData,
+	FSubstrateSceneData* SubstrateSceneData,
+	FRDGTextureRef OutDepthTexture)
+{
+	check(SubstrateSceneData);
+	check(SubstrateSceneData->MaterialTextureArray);
+	check(Substrate::IsSubstrateEnabled());
+
+	const bool bWriteDepth = OutDepthTexture != nullptr;
+	const FIntPoint Resolution = OutDepthTexture ? OutDepthTexture->Desc.Extent : SubstrateSceneData->MaterialTextureArray->Desc.Extent;
+	FHairVisibilityGBufferWritePS::FParameters* Parameters = GraphBuilder.AllocParameters<FHairVisibilityGBufferWritePS::FParameters>();
+	Parameters->ViewUniformBuffer = View.ViewUniformBuffer;
+	Parameters->bWriteDummyData = bWriteDummyData ? 1 : 0;
+	Parameters->HairStrands = View.HairStrandsViewData.UniformBuffer;
+	Parameters->Substrate = Substrate::GetSubstrateCommonParameter(*SubstrateSceneData);
+	Parameters->RWMaterialTextureArray = SubstrateSceneData->MaterialTextureArrayUAV;
+	Parameters->RWTopLayerTexture = SubstrateSceneData->TopLayerTextureUAV;
+	if (bWriteDepth)
+	{
+		Parameters->RenderTargets.DepthStencil = FDepthStencilBinding(
+			OutDepthTexture,
+			ERenderTargetLoadAction::ELoad,
+			ERenderTargetLoadAction::ELoad,
+			FExclusiveDepthStencil::DepthWrite_StencilNop);
+	}
+
+	FHairVisibilityGBufferWritePS::FPermutationDomain PermutationVector;
+	PermutationVector.Set<FHairVisibilityGBufferWritePS::FOutputType>(bWriteDepth ? 1 : 0);
+	TShaderMapRef<FHairVisibilityGBufferWritePS> PixelShader(View.ShaderMap, PermutationVector);
+	InternalCommonDrawPass(
+		GraphBuilder,
+		RDG_EVENT_NAME("HairStrands::GBufferOverride"),
+		View,
+		Resolution,
+		EHairStrandsCommonPassType::GBuffer,
+		bWriteDepth,
+		false /*bHasHoldout*/,
+		TileData,
+		PixelShader,
+		Parameters);
+}
+
 static void AddHairVisibilityGBufferWritePass(
 	FRDGBuilder& GraphBuilder,
 	const FViewInfo& View,
@@ -676,44 +731,57 @@ static void InternalRenderHairComposition(
 				SceneDepthTexture);
 		}
 
-		const bool bWriteDummyData		= View.Family->ViewMode != VMI_VisualizeBuffer && GHairWriteGBufferData == 1;
-		const bool bWritePartialGBuffer = View.Family->ViewMode != VMI_VisualizeBuffer && (GHairWriteGBufferData == 1 || GHairWriteGBufferData == 2);
-		const bool bWriteFullGBuffer	= View.Family->ViewMode == VMI_VisualizeBuffer || (GHairWriteGBufferData == 3);
+		const bool bWriteDummyData		= !(View.Family->ViewMode == VMI_VisualizeBuffer || View.Family->ViewMode == VMI_VisualizeSubstrate) && GHairWriteGBufferData == 1;
+		const bool bWritePartialGBuffer =  (View.Family->ViewMode == VMI_VisualizeBuffer || View.Family->ViewMode == VMI_VisualizeSubstrate) && (GHairWriteGBufferData == 1 || GHairWriteGBufferData == 2);
+		const bool bWriteFullGBuffer	=  (View.Family->ViewMode == VMI_VisualizeBuffer || View.Family->ViewMode == VMI_VisualizeSubstrate) || (GHairWriteGBufferData == 3);
 		if (bWriteFullGBuffer || bWritePartialGBuffer)
 		{
-			const FSceneTextures& SceneTextures = View.GetSceneTextures();
-			const FRDGTextureRef GBufferATexture = SceneTextures.GBufferA;
-			const FRDGTextureRef GBufferBTexture = SceneTextures.GBufferB;
-			const FRDGTextureRef GBufferCTexture = SceneTextures.GBufferC;
-			const FRDGTextureRef GBufferDTexture = SceneTextures.GBufferD;
-			const FRDGTextureRef GBufferETexture = SceneTextures.GBufferE;
-			if (bWritePartialGBuffer && GBufferATexture && GBufferBTexture)
+			if (Substrate::IsSubstrateEnabled() && !Substrate::IsSubstrateBlendableGBufferEnabled(View.GetShaderPlatform()))
 			{
 				AddHairVisibilityGBufferWritePass(
 					GraphBuilder,
 					View,
 					bWriteDummyData,
 					VisibilityData.TileData,
-					GBufferATexture,
-					GBufferBTexture,
-					nullptr,
-					nullptr,
-					nullptr,
-					nullptr);
+					View.SubstrateViewData.SceneData, 
+					SceneDepthTexture);
 			}
-			else if (bWriteFullGBuffer && GBufferATexture && GBufferBTexture && GBufferCTexture && SceneDepthTexture)
+			else
 			{
-				AddHairVisibilityGBufferWritePass(
-					GraphBuilder,
-					View,
-					bWriteDummyData,
-					VisibilityData.TileData,
-					GBufferATexture,
-					GBufferBTexture,
-					GBufferCTexture,
-					GBufferDTexture,
-					GBufferETexture,
-					SceneDepthTexture);
+				const FSceneTextures& SceneTextures = View.GetSceneTextures();
+				const FRDGTextureRef GBufferATexture = SceneTextures.GBufferA;
+				const FRDGTextureRef GBufferBTexture = SceneTextures.GBufferB;
+				const FRDGTextureRef GBufferCTexture = SceneTextures.GBufferC;
+				const FRDGTextureRef GBufferDTexture = SceneTextures.GBufferD;
+				const FRDGTextureRef GBufferETexture = SceneTextures.GBufferE;
+				if (bWritePartialGBuffer && GBufferATexture && GBufferBTexture)
+				{
+					AddHairVisibilityGBufferWritePass(
+						GraphBuilder,
+						View,
+						bWriteDummyData,
+						VisibilityData.TileData,
+						GBufferATexture,
+						GBufferBTexture,
+						nullptr,
+						nullptr,
+						nullptr,
+						nullptr);
+				}
+				else if (bWriteFullGBuffer && GBufferATexture && GBufferBTexture && GBufferCTexture && SceneDepthTexture)
+				{
+					AddHairVisibilityGBufferWritePass(
+						GraphBuilder,
+						View,
+						bWriteDummyData,
+						VisibilityData.TileData,
+						GBufferATexture,
+						GBufferBTexture,
+						GBufferCTexture,
+						GBufferDTexture,
+						GBufferETexture,
+						SceneDepthTexture);
+				}
 			}
 		}
 	}
diff --git a/Engine/Source/Runtime/Renderer/Private/Substrate/Substrate.cpp b/Engine/Source/Runtime/Renderer/Private/Substrate/Substrate.cpp
--- a/Engine/Source/Runtime/Renderer/Private/Substrate/Substrate.cpp
+++ b/Engine/Source/Runtime/Renderer/Private/Substrate/Substrate.cpp
@@ -532,6 +532,10 @@ void InitialiseSubstrateFrameSceneData(FRDGBuilder& GraphBuilder, FSceneRenderer
 		if (bNeedsMaterialBuffer && !bBlendableGBufferEnabled)
 		{
 			Out.TopLayerTexture = GraphBuilder.CreateTexture(FRDGTextureDesc::Create2D(SceneTextureExtent, GetTopLayerTextureFormat(bUseDBufferPass), FClearValueBinding::Black, TexCreate_RenderTargetable | TexCreate_ShaderResource | TexCreate_FastVRAM | (bNeedsUAV ? TexCreate_UAV : TexCreate_None)), TEXT("Substrate.TopLayerTexture"));
+			if (bNeedsUAV)
+			{
+				Out.TopLayerTextureUAV = GraphBuilder.CreateUAV(Out.TopLayerTexture);
+			}
 		}
 
 		// Separated subsurface and rough refraction textures
@@ -648,7 +652,7 @@ static FSubstrateCommonParameters GetSubstrateCommonParameter()
 	return Out;
 }
 
-static FSubstrateCommonParameters GetSubstrateCommonParameter(const FSubstrateSceneData& In)
+FSubstrateCommonParameters GetSubstrateCommonParameter(const FSubstrateSceneData& In)
 {
 	FSubstrateCommonParameters Out;
 	Out.bRoughDiffuse 		= In.bRoughDiffuse ? 1u : 0u;
