Fix CIS warning.

diff --git a/Engine/Source/Runtime/Renderer/Private/HairStrands/HairStrandsComposition.cpp b/Engine/Source/Runtime/Renderer/Private/HairStrands/HairStrandsComposition.cpp
--- a/Engine/Source/Runtime/Renderer/Private/HairStrands/HairStrandsComposition.cpp
+++ b/Engine/Source/Runtime/Renderer/Private/HairStrands/HairStrandsComposition.cpp
@@ -507,7 +507,7 @@ class FHairVisibilityGBufferWritePS : public FGlobalShader
 		SHADER_PARAMETER_RDG_UNIFORM_BUFFER(FHairStrandsViewUniformParameters, HairStrands)
 		// Substrate
 		SHADER_PARAMETER_STRUCT_INCLUDE(FSubstrateCommonParameters, Substrate)
-		SHADER_PARAMETER_RDG_TEXTURE_UAV(RWTexture2DARRAY, RWMaterialTextureArray)
+		SHADER_PARAMETER_RDG_TEXTURE_UAV(RWTexture2DArray, RWMaterialTextureArray)
 		SHADER_PARAMETER_RDG_TEXTURE_UAV(RWTexture2D, RWTopLayerTexture)
 		RENDER_TARGET_BINDING_SLOTS()
 	END_SHADER_PARAMETER_STRUCT()
