Fix mobile SSR not enabled on single layer water with mobile SM5 enabled

#rnx

diff --git a/Engine/Source/Runtime/Renderer/Private/MobileBasePassRendering.h b/Engine/Source/Runtime/Renderer/Private/MobileBasePassRendering.h
--- a/Engine/Source/Runtime/Renderer/Private/MobileBasePassRendering.h
+++ b/Engine/Source/Runtime/Renderer/Private/MobileBasePassRendering.h
@@ -437,7 +437,7 @@ public:
 		OutEnvironment.SetDefine(TEXT("MOBILE_TRANSLUCENT_COLOR_TRANSMITTANCE_DUAL_SRC_BLENDING"), TranslucentColorTransmittanceMode == EMobileTranslucentColorTransmittanceMode::DUAL_SRC_BLENDING ? 1u : 0u);
 		OutEnvironment.SetDefine(TEXT("MOBILE_TRANSLUCENT_COLOR_TRANSMITTANCE_PROGRAMMABLE_BLENDING"), TranslucentColorTransmittanceMode == EMobileTranslucentColorTransmittanceMode::PROGRAMMABLE_BLENDING ? 1u : 0u);
 		OutEnvironment.SetDefine(TEXT("MOBILE_TRANSLUCENT_COLOR_TRANSMITTANCE_SINGLE_SRC_BLENDING"), TranslucentColorTransmittanceMode == EMobileTranslucentColorTransmittanceMode::SINGLE_SRC_BLENDING ? 1u : 0u);
-		OutEnvironment.SetDefine(TEXT("MOBILE_SSR_ENABLED"), !bDeferredShadingEnabled && AreMobileScreenSpaceReflectionsEnabled(Parameters.Platform) ? 1u : 0u);
+		OutEnvironment.SetDefine(TEXT("MOBILE_SSR_ENABLED"), ((!bDeferredShadingEnabled && AreMobileScreenSpaceReflectionsEnabled(Parameters.Platform)) || (bIsSLW && MobileSupportsSM5MaterialNodes(Parameters.Platform))) ? 1u : 0u);
 	}
 	
 	/** Initialization constructor. */
