Prevent updated shadow mask writes if volumetric shadow map is empty

diff --git a/Engine/Source/Runtime/Renderer/Private/HeterogeneousVolumes/HeterogeneousVolumes.h b/Engine/Source/Runtime/Renderer/Private/HeterogeneousVolumes/HeterogeneousVolumes.h
--- a/Engine/Source/Runtime/Renderer/Private/HeterogeneousVolumes/HeterogeneousVolumes.h
+++ b/Engine/Source/Runtime/Renderer/Private/HeterogeneousVolumes/HeterogeneousVolumes.h
@@ -6,6 +6,7 @@
 #include "ConvexVolume.h"
 #include "RendererInterface.h"
 #include "ShaderParameterMacros.h"
+#include "SceneTexturesConfig.h"
 #include "HeterogeneousVolumeInterface.h"
 
 class FLightSceneInfo;
@@ -703,16 +704,13 @@ void RenderAdaptiveVolumetricShadowMapWithLiveShading(
 void RenderVolumetricShadowMaskForLight(
 	FRDGBuilder& GraphBuilder,
 	// Scene data
-	const FSceneTextures& SceneTextures,
-	FScene* Scene,
-	const FViewInfo& View,
+	const FSceneTexturesConfig& Config,
+	const TConstArrayView<FViewInfo>& Views,
 	// Light data
 	uint32 LightType,
 	const FLightSceneInfo* LightSceneInfo,
-	// Current shadow mask
-	FRDGTextureRef ShadowMaskTexture,
-	// Output
-	FRDGTextureRef& VolumetricShadowMaskTexture
+	// Resultant shadow mask
+	FRDGTextureRef& ScreenShadowMaskTexture
 );
 
 void RenderAdaptiveVolumetricCameraMapWithVoxelGrid(
diff --git a/Engine/Source/Runtime/Renderer/Private/HeterogeneousVolumes/HeterogeneousVolumesLiveShadingPipeline.cpp b/Engine/Source/Runtime/Renderer/Private/HeterogeneousVolumes/HeterogeneousVolumesLiveShadingPipeline.cpp
--- a/Engine/Source/Runtime/Renderer/Private/HeterogeneousVolumes/HeterogeneousVolumesLiveShadingPipeline.cpp
+++ b/Engine/Source/Runtime/Renderer/Private/HeterogeneousVolumes/HeterogeneousVolumesLiveShadingPipeline.cpp
@@ -3451,24 +3451,23 @@ class FRenderVolumetricShadowMaskCS : public FGlobalShader
 
 IMPLEMENT_GLOBAL_SHADER(FRenderVolumetricShadowMaskCS, "/Engine/Private/HeterogeneousVolumes/HeterogeneousVolumesShadowMask.usf", "RenderVolumetricShadowMaskCS", SF_Compute);
 
-void RenderVolumetricShadowMaskForLight(
+bool RenderVolumetricShadowMaskForLight(
 	FRDGBuilder& GraphBuilder,
-	// Scene data
-	const FSceneTextures& SceneTextures,
-	//const FMinimalSceneTextures& SceneTextures,
-	FScene* Scene,
 	const FViewInfo& View,
-	// Light data
-	uint32 LightType,
 	const FLightSceneInfo* LightSceneInfo,
-	// Current shadow mask
 	FRDGTextureRef ShadowMaskTexture,
-	// Output
 	FRDGTextureRef& VolumetricShadowMaskTexture
 )
 {
 	FIntVector GroupCount = FComputeShaderUtils::GetGroupCount(View.ViewRect.Size(), FRenderVolumetricShadowMaskCS::GetThreadGroupSize2D());
 
+	TRDGUniformBufferRef<FAdaptiveVolumetricShadowMapUniformBufferParameters> AVSM = HeterogeneousVolumes::GetAdaptiveVolumetricShadowMapUniformBuffer(GraphBuilder, View.ViewState, LightSceneInfo);
+	TRDGUniformBufferRef<FBeerShadowMapUniformBufferParameters> BeerShadowMap = HeterogeneousVolumes::GetBeerShadowMapUniformBuffer(GraphBuilder, View.ViewState, LightSceneInfo);
+	if (AVSM->GetParameters()->bIsEmpty && BeerShadowMap->GetParameters()->bIsEmpty)
+	{
+		return false;
+	}
+
 	// Note must be done in the same scope as we add the pass otherwise the UB lifetime will not be guaranteed
 	FDeferredLightUniformStruct DeferredLightUniform = GetDeferredLightParameters(View, *LightSceneInfo);
 	TUniformBufferRef<FDeferredLightUniformStruct> DeferredLightUB = CreateUniformBufferImmediate(DeferredLightUniform, UniformBuffer_SingleDraw);
@@ -3477,16 +3476,16 @@ void RenderVolumetricShadowMaskForLight(
 	{
 		// Scene data
 		PassParameters->View = View.ViewUniformBuffer;
-		PassParameters->SceneTextures = GetSceneTextureParameters(GraphBuilder, SceneTextures);
+		PassParameters->SceneTextures = GetSceneTextureParameters(GraphBuilder, View.GetSceneTextures());
 		PassParameters->Scene = View.GetSceneUniforms().GetBuffer(GraphBuilder);
 
 		// Light data
-		PassParameters->LightType = LightType;
+		PassParameters->LightType = LightSceneInfo->Type;
 		PassParameters->DeferredLight = DeferredLightUB;
 
 		// Shadow data
-		PassParameters->AVSM = HeterogeneousVolumes::GetAdaptiveVolumetricShadowMapUniformBuffer(GraphBuilder, View.ViewState, LightSceneInfo);
-		PassParameters->BeerShadowMap = HeterogeneousVolumes::GetBeerShadowMapUniformBuffer(GraphBuilder, View.ViewState, LightSceneInfo);
+		PassParameters->AVSM = AVSM;
+		PassParameters->BeerShadowMap = BeerShadowMap;
 
 		// Dispatch data
 		PassParameters->GroupCount = GroupCount;
@@ -3504,8 +3503,6 @@ void RenderVolumetricShadowMaskForLight(
 	PermutationVector = FRenderVolumetricShadowMaskCS::RemapPermutation(PermutationVector);
 	TShaderRef<FRenderVolumetricShadowMaskCS> ComputeShader = View.ShaderMap->GetShader<FRenderVolumetricShadowMaskCS>(PermutationVector);
 
-
-
 	FComputeShaderUtils::AddPass(
 		GraphBuilder,
 		//RDG_EVENT_NAME("%s", *PassName),
@@ -3513,4 +3510,43 @@ void RenderVolumetricShadowMaskForLight(
 		ComputeShader,
 		PassParameters,
 		GroupCount);
+
+	return true;
 }
+
+DECLARE_GPU_STAT_NAMED(VolumetricShadowMask, TEXT("Volumetric Shadow Mask"));
+
+void RenderVolumetricShadowMaskForLight(
+	FRDGBuilder& GraphBuilder,
+	// Scene data
+	const FSceneTexturesConfig& Config,
+	const TConstArrayView<FViewInfo>& Views,
+	// Light data
+	uint32 LightType,
+	const FLightSceneInfo* LightSceneInfo,
+	// Current shadow mask
+	FRDGTextureRef& ScreenShadowMaskTexture)
+{
+	RDG_EVENT_SCOPE_STAT(GraphBuilder, VolumetricShadowMask, "VolumetricShadowMask");
+	RDG_GPU_STAT_SCOPE(GraphBuilder, VolumetricShadowMask);
+
+	// Render Volumetric Shadow Mask
+	FRDGTextureDesc Desc = FRDGTextureDesc::Create2D(
+		Config.Extent,
+		PF_B8G8R8A8,
+		FClearValueBinding::White,
+		TexCreate_ShaderResource | TexCreate_RenderTargetable | TexCreate_UAV);
+	FRDGTextureRef VolumetricShadowMaskTexture = GraphBuilder.CreateTexture(Desc, TEXT("HeterogeneousVolumes.VolumetricShadowMaskTexture"));
+
+	bool bVolumetricShadowMaskTexture = false;
+	for (const FViewInfo& View : Views)
+	{
+		FRDGTextureRef InputShadowMaskTexture = ScreenShadowMaskTexture ? ScreenShadowMaskTexture : GSystemTextures.GetWhiteDummy(GraphBuilder);
+		bVolumetricShadowMaskTexture |= RenderVolumetricShadowMaskForLight(GraphBuilder, View, LightSceneInfo, InputShadowMaskTexture, VolumetricShadowMaskTexture);
+	}
+
+	if (bVolumetricShadowMaskTexture)
+	{
+		ScreenShadowMaskTexture = VolumetricShadowMaskTexture;
+	}
+}
\ No newline at end of file
diff --git a/Engine/Source/Runtime/Renderer/Private/LightRendering.cpp b/Engine/Source/Runtime/Renderer/Private/LightRendering.cpp
--- a/Engine/Source/Runtime/Renderer/Private/LightRendering.cpp
+++ b/Engine/Source/Runtime/Renderer/Private/LightRendering.cpp
@@ -212,8 +212,6 @@ static FAutoConsoleVariableRef CVarDebugLightDiscardProp(
 );
 #endif // ENABLE_DEBUG_DISCARD_PROP
 
-DECLARE_GPU_STAT_NAMED(VolumetricShadowMask, TEXT("Volumetric Shadow Mask"));
-
 #if RHI_RAYTRACING
 
 static bool ShouldRenderRayTracingShadowsForLightType(ELightComponentType LightType)
@@ -2240,23 +2238,7 @@ void FDeferredShadingSceneRenderer::RenderLights(
 
 					if (ShouldHeterogeneousVolumesCastShadows())
 					{
-						RDG_EVENT_SCOPE_STAT(GraphBuilder, VolumetricShadowMask, "VolumetricShadowMask");
-						RDG_GPU_STAT_SCOPE(GraphBuilder, VolumetricShadowMask);
-
-						// Render Volumetric Shadow Mask
-						FRDGTextureDesc Desc = FRDGTextureDesc::Create2D(
-							SceneTextures.Config.Extent,
-							PF_B8G8R8A8,
-							FClearValueBinding::White,
-							TexCreate_ShaderResource | TexCreate_RenderTargetable | TexCreate_UAV);
-						FRDGTextureRef VolumetricShadowMaskTexture = GraphBuilder.CreateTexture(Desc, TEXT("HeterogeneousVolumes.VolumetricShadowMaskTexture"));
-
-						for (const FViewInfo& View : Views)
-						{
-							FRDGTextureRef InputShadowMaskTexture = ScreenShadowMaskTexture ? ScreenShadowMaskTexture : GSystemTextures.GetWhiteDummy(GraphBuilder);
-							RenderVolumetricShadowMaskForLight(GraphBuilder, View.GetSceneTextures(), Scene, View, LightSceneInfo.Type, &LightSceneInfo, InputShadowMaskTexture, VolumetricShadowMaskTexture);
-						}
-						ScreenShadowMaskTexture = VolumetricShadowMaskTexture;
+						RenderVolumetricShadowMaskForLight(GraphBuilder,	SceneTextures.Config, Views, LightSceneInfo.Type, &LightSceneInfo, ScreenShadowMaskTexture);
 					}
 
 					bUsedShadowMaskTexture = true;
