Add TexCreate_3DTiling flag to volume textures in TranslucentLighting

diff --git a/Engine/Source/Runtime/Renderer/Private/TranslucentLighting.cpp b/Engine/Source/Runtime/Renderer/Private/TranslucentLighting.cpp
--- a/Engine/Source/Runtime/Renderer/Private/TranslucentLighting.cpp
+++ b/Engine/Source/Runtime/Renderer/Private/TranslucentLighting.cpp
@@ -1134,7 +1134,7 @@ void FTranslucencyLightingVolumeTextures::GetTextureFormatAndCreationFlags(EPixe
 	OutPixelFormat = PF_FloatRGBA;
 
 	// TODO: We can skip the and TLV allocations when rendering in forward shading mode
-	OutCreationFlags = TexCreate_ShaderResource | TexCreate_RenderTargetable | TexCreate_ReduceMemoryWithTilingMode | TexCreate_UAV;
+	OutCreationFlags = TexCreate_ShaderResource | TexCreate_RenderTargetable | TexCreate_ReduceMemoryWithTilingMode | TexCreate_3DTiling | TexCreate_UAV;
 }
 
 
