Use compute shader for CombineLUTs if volume rendering is guaranteed

Also reduce dispatch size from 8x8x8 to 4x4x4 to increase work granularity.

#FailedValidations Preflight

diff --git a/Engine/Source/Runtime/Renderer/Private/PostProcess/PostProcessCombineLUTs.cpp b/Engine/Source/Runtime/Renderer/Private/PostProcess/PostProcessCombineLUTs.cpp
--- a/Engine/Source/Runtime/Renderer/Private/PostProcess/PostProcessCombineLUTs.cpp
+++ b/Engine/Source/Runtime/Renderer/Private/PostProcess/PostProcessCombineLUTs.cpp
@@ -184,13 +184,13 @@ struct FCachedLUTSettings
 	FWorkingColorSpaceShaderParameters WorkingColorSpaceShaderParameters;
 	bool bUseCompute = false;
 
-	bool UpdateCachedValues(FRDGBuilder& GraphBuilder, const FViewInfo& View, const FTexture* const* Textures, const float* Weights, uint32 BlendCount)
+	bool UpdateCachedValues(FRDGBuilder& GraphBuilder, const FViewInfo& View, const FTexture* const* Textures, const float* Weights, uint32 BlendCount, bool bInUseCompute)
 	{
 		bool bHasChanged = false;
 		GetCombineLUTParameters(GraphBuilder, View, Textures, Weights, BlendCount, bHasChanged);
 		UPDATE_CACHE_SETTINGS(UniqueID, View.State ? View.State->GetViewKey() : 0, bHasChanged);
 		UPDATE_CACHE_SETTINGS(ShaderPlatform, View.GetShaderPlatform(), bHasChanged);
-		UPDATE_CACHE_SETTINGS(bUseCompute, View.bUseComputePasses, bHasChanged);
+		UPDATE_CACHE_SETTINGS(bUseCompute, bInUseCompute, bHasChanged);
 
 		const FWorkingColorSpaceShaderParameters* InWorkingColorSpaceShaderParameters = reinterpret_cast<const FWorkingColorSpaceShaderParameters*>(GDefaultWorkingColorSpaceUniformBuffer.GetContents());
 		if (InWorkingColorSpaceShaderParameters)
@@ -318,8 +318,6 @@ struct FCachedLUTSettings
 class FLUTBlenderShader : public FGlobalShader
 {
 public:
-	static const int32 GroupSize = 8;
-
 	class FBlendCount : SHADER_PERMUTATION_RANGE_INT("BLENDCOUNT", 1, 5);
 	class FOutputDeviceSRGB : SHADER_PERMUTATION_BOOL("OUTPUT_DEVICE_SRGB");
 	class FUseACES2 : SHADER_PERMUTATION_BOOL("USE_ACES_2");
@@ -333,9 +331,8 @@ public:
 	static void ModifyCompilationEnvironment(const FGlobalShaderPermutationParameters& Parameters, FShaderCompilerEnvironment& OutEnvironment)
 	{
 		FGlobalShader::ModifyCompilationEnvironment(Parameters, OutEnvironment);
-		OutEnvironment.SetDefine(TEXT("THREADGROUP_SIZE"), GroupSize);
 
-		const int UseVolumeLUT = PipelineVolumeTextureLUTSupportGuaranteedAtRuntime(Parameters.Platform) ? 1 : 0;
+		const int32 UseVolumeLUT = PipelineVolumeTextureLUTSupportGuaranteedAtRuntime(Parameters.Platform) ? 1 : 0;
 		OutEnvironment.SetDefine(TEXT("USE_VOLUME_LUT"), UseVolumeLUT);
 	}
 
@@ -367,6 +364,9 @@ IMPLEMENT_GLOBAL_SHADER(FLUTBlenderPS, "/Engine/Private/PostProcessCombineLUTs.u
 class FLUTBlenderCS : public FLUTBlenderShader
 {
 public:
+	static constexpr int32 GroupSize2D = 8;
+	static constexpr int32 GroupSize3D = 4;
+
 	DECLARE_GLOBAL_SHADER(FLUTBlenderCS);
 	SHADER_USE_PARAMETER_STRUCT(FLUTBlenderCS, FLUTBlenderShader);
 
@@ -380,6 +380,14 @@ public:
 	{
 		return IsFeatureLevelSupported(Parameters.Platform, ERHIFeatureLevel::SM5);
 	}
+
+	static void ModifyCompilationEnvironment(const FGlobalShaderPermutationParameters& Parameters, FShaderCompilerEnvironment& OutEnvironment)
+	{
+		FLUTBlenderShader::ModifyCompilationEnvironment(Parameters, OutEnvironment);
+
+		const int32 GroupSize = PipelineVolumeTextureLUTSupportGuaranteedAtRuntime(Parameters.Platform) ? GroupSize3D : GroupSize2D;
+		OutEnvironment.SetDefine(TEXT("THREADGROUP_SIZE"), GroupSize);
+	}
 };
 
 IMPLEMENT_GLOBAL_SHADER(FLUTBlenderCS, "/Engine/Private/PostProcessCombineLUTs.usf", "MainCS", SF_Compute);
@@ -501,10 +509,9 @@ FRDGTextureRef AddCombineLUTPass(FRDGBuilder& GraphBuilder, const FViewInfo& Vie
 		LocalCount = GenerateFinalTable(View.FinalPostProcessSettings, LocalTextures, LocalWeights, GMaxLUTBlendCount);
 	}
 
-	const bool bUseComputePass = View.bUseComputePasses;
-
 	const bool bUseVolumeTextureLUT = PipelineVolumeTextureLUTSupportGuaranteedAtRuntime(View.GetShaderPlatform());
-	
+	const bool bUseComputePass = bUseVolumeTextureLUT || View.bUseComputePasses;
+
 	const EDisplayOutputFormat OutputDevice = ViewFamily.RenderTarget->GetDisplayOutputFormat();
 	// Note: We also enforce floats with HDR ScRGB (editor only) to preserve negative values.
 	const bool bUseFloatOutput = ViewFamily.SceneCaptureSource == SCS_FinalColorHDR ||
@@ -519,7 +526,7 @@ FRDGTextureRef AddCombineLUTPass(FRDGBuilder& GraphBuilder, const FViewInfo& Vie
 
 	View.SetValidTonemappingLUT();
 
-	const bool bHasChanged = CachedLUTSettings.UpdateCachedValues(GraphBuilder, View, LocalTextures, LocalWeights, LocalCount);
+	const bool bHasChanged = CachedLUTSettings.UpdateCachedValues(GraphBuilder, View, LocalTextures, LocalWeights, LocalCount, bUseComputePass);
 	if (!bHasChanged && OutputTexture && CVarLUTUpdateEveryFrame.GetValueOnRenderThread() == 0)
 	{
 		return OutputTexture;
@@ -556,7 +563,8 @@ FRDGTextureRef AddCombineLUTPass(FRDGBuilder& GraphBuilder, const FViewInfo& Vie
 		PermutationVector.Set<FLUTBlenderShader::FOutputDeviceSRGB>(bOutputDeviceSRGB);
 		TShaderMapRef<FLUTBlenderCS> ComputeShader(View.ShaderMap, PermutationVector);
 
-		const uint32 GroupSizeXY = FMath::DivideAndRoundUp(OutputViewSize.X, FLUTBlenderCS::GroupSize);
+		const int32 GroupSize = bUseVolumeTextureLUT ? FLUTBlenderCS::GroupSize3D : FLUTBlenderCS::GroupSize2D;
+		const uint32 GroupSizeXY = FMath::DivideAndRoundUp(OutputViewSize.X, GroupSize);
 		const uint32 GroupSizeZ = bUseVolumeTextureLUT ? GroupSizeXY : 1;
 
 		FComputeShaderUtils::AddPass(
