Fix missing hair shading on Switch2 on MH.

diff --git a/Engine/Source/Runtime/Renderer/Private/LightRendering.cpp b/Engine/Source/Runtime/Renderer/Private/LightRendering.cpp
--- a/Engine/Source/Runtime/Renderer/Private/LightRendering.cpp
+++ b/Engine/Source/Runtime/Renderer/Private/LightRendering.cpp
@@ -960,6 +960,12 @@ class FDeferredLightPS : public FGlobalShader
 		RENDER_TARGET_BINDING_SLOTS()
 	END_SHADER_PARAMETER_STRUCT()
 
+	static bool SupportsHairComplexTransmittance(EShaderPlatform InPlatform)
+	{
+		// For now disable complex transmittance on certain platform as HairLUT is incorrectly populated.
+		return IsHairStrandsSupported(EHairStrandsShaderType::All, InPlatform) && !FDataDrivenShaderPlatformInfo::GetIsLanguageNintendo(InPlatform);
+	}
+
 	static bool ShouldCompilePermutation(const FGlobalShaderPermutationParameters& Parameters)
 	{
 		FPermutationDomain PermutationVector(Parameters.PermutationId);
@@ -1012,8 +1018,7 @@ class FDeferredLightPS : public FGlobalShader
 			return false;
 		}
 
-		const bool bNeedComplexTransmittanceSupport = IsHairStrandsSupported(EHairStrandsShaderType::All, Parameters.Platform);
-		if (PermutationVector.Get<FHairLighting>() == 0 && PermutationVector.Get<FDeferredLightPS::FHairComplexTransmittance>() && !bNeedComplexTransmittanceSupport)
+		if (PermutationVector.Get<FHairLighting>() == 0 && PermutationVector.Get<FDeferredLightPS::FHairComplexTransmittance>() && !SupportsHairComplexTransmittance(Parameters.Platform))
 		{
 			return false;
 		}
@@ -2844,7 +2849,7 @@ static void RenderLight(
 	const ELightComponentType LightType = (ELightComponentType)LightProxy->GetLightType();
 	const bool bIsRadial = LightType != LightType_Directional;
 	const bool bUseVirtualShadowMapMask = VirtualShadowMapId != INDEX_NONE && ShadowMaskBits;
-	const bool bNeedComplexTransmittanceSupport = View.HairCardsMeshElements.Num() && IsHairStrandsSupported(EHairStrandsShaderType::All, View.GetShaderPlatform());
+	const bool bNeedComplexTransmittanceSupport = View.HairCardsMeshElements.Num() && FDeferredLightPS::SupportsHairComplexTransmittance(View.GetShaderPlatform());
 
 	const bool bSupportAnisotropyPermutation = ShouldRenderAnisotropicLighting(View);
 
@@ -3407,7 +3412,7 @@ static void InternalRenderSimpleLightsStandardDeferred(
 		SortedLightSet.SimpleLights.InstanceData[0], // Use a dummy light to create the PassParameter buffer. The light data will be
 		FVector(0, 0, 0));		  // update dynamically with the pass light loop for efficiency purpose
 
-	const bool bNeedComplexTransmittanceSupport = View.HairCardsMeshElements.Num() && IsHairStrandsSupported(EHairStrandsShaderType::All, View.GetShaderPlatform());
+	const bool bNeedComplexTransmittanceSupport = View.HairCardsMeshElements.Num() && FDeferredLightPS::SupportsHairComplexTransmittance(View.GetShaderPlatform());
 
 	FDeferredLightPS::FPermutationDomain PermutationVector;
 	PermutationVector.Set< FDeferredLightPS::FSourceShapeDim >(ELightSourceShape::Capsule);
