Substrate - fixed blending mode of post process material with blendable gbuffer.

#rnx
#tests PC editor dx12

diff --git a/Engine/Source/Runtime/Renderer/Private/PostProcess/PostProcessMaterial.cpp b/Engine/Source/Runtime/Renderer/Private/PostProcess/PostProcessMaterial.cpp
--- a/Engine/Source/Runtime/Renderer/Private/PostProcess/PostProcessMaterial.cpp
+++ b/Engine/Source/Runtime/Renderer/Private/PostProcess/PostProcessMaterial.cpp
@@ -150,7 +150,7 @@ static FRHIBlendState* GetMaterialBlendState(const FMaterial* Material)
 
 	check(Material);
 
-	if (Substrate::IsSubstrateEnabled())
+	if (Material->IsSubstrateMaterial())
 	{
 		switch (Material->GetBlendMode())
 		{
