Revert geometry shader requirement for VolumetricFog.

#test EngineTests dx11 sm5, dx12 sm6, vulkan sm5/sm6, substrate 0/1, substrate format=0.1

diff --git a/Engine/Source/Runtime/Renderer/Private/VolumetricFog.cpp b/Engine/Source/Runtime/Renderer/Private/VolumetricFog.cpp
--- a/Engine/Source/Runtime/Renderer/Private/VolumetricFog.cpp
+++ b/Engine/Source/Runtime/Renderer/Private/VolumetricFog.cpp
@@ -1195,7 +1195,7 @@ IMPLEMENT_GLOBAL_SHADER(FVolumetricFogFinalIntegrationCS, "/Engine/Private/Volum
 
 bool DoesPlatformSupportVolumetricFogVoxelization(const FStaticShaderPlatform Platform)
 {
-	return !IsMobileVolumetricFog(Platform) && RHISupportsGeometryShaders(Platform);
+	return !IsMobileVolumetricFog(Platform);
 }
 bool ShouldRenderVolumetricFog(const FScene* Scene, const FSceneViewFamily& ViewFamily)
 {
