[First Person] Fixed potential nullptr dereference.

diff --git a/Engine/Source/Runtime/Renderer/Private/Shadows/ShadowSceneRenderer.cpp b/Engine/Source/Runtime/Renderer/Private/Shadows/ShadowSceneRenderer.cpp
--- a/Engine/Source/Runtime/Renderer/Private/Shadows/ShadowSceneRenderer.cpp
+++ b/Engine/Source/Runtime/Renderer/Private/Shadows/ShadowSceneRenderer.cpp
@@ -1275,7 +1275,7 @@ void FShadowSceneRenderer::CreateNaniteRenderPasses(
 
 		for (FProjectedShadowInfo* ProjectedShadowInfo : Shadows)
 		{
-			const bool bIsFirstPersonShadowClipmap = ProjectedShadowInfo->VirtualShadowMapClipmap->IsFirstPersonShadow();
+			const bool bIsFirstPersonShadowClipmap = ProjectedShadowInfo->VirtualShadowMapClipmap && ProjectedShadowInfo->VirtualShadowMapClipmap->IsFirstPersonShadow();
 			const bool bHasExplicitInstanceChunks = ProjectedShadowInfo->ExplicitNaniteInstanceChunkStartOffset < ProjectedShadowInfo->ExplicitNaniteInstanceChunkEndOffset;
 			if (ProjectedShadowInfo->bShouldRenderVSM && (!bIsFirstPersonShadowClipmap || bHasExplicitInstanceChunks))
 			{
