[SingleLayerWater] Fixed SLW meshes on the cached (non-GDME) path not working with the new SM5 mobile rendering path.

diff --git a/Engine/Source/Runtime/Renderer/Private/SceneVisibility.cpp b/Engine/Source/Runtime/Renderer/Private/SceneVisibility.cpp
--- a/Engine/Source/Runtime/Renderer/Private/SceneVisibility.cpp
+++ b/Engine/Source/Runtime/Renderer/Private/SceneVisibility.cpp
@@ -1430,6 +1430,7 @@ void FRelevancePacket::ComputeRelevance(FDynamicPrimitiveIndexList& DynamicPrimi
 	const bool bIsTranslucentClippedDepthEnabled = NeedTSRThinGeometryDetectionPass(View);
 	const bool bHLODActive = Scene.SceneLODHierarchy.IsActive();
 	const bool bRenderCustomDepth = (ShadingPath == EShadingPath::Mobile) ? FMobileSceneRenderer::ShouldRenderCustomDepth(View) : true;
+	const bool bMobileSupportsSM5MaterialNodes = MobileSupportsSM5MaterialNodes(View.GetShaderPlatform());
 	const FHLODVisibilityState* const HLODState = bHLODActive && ViewState ? &ViewState->HLODVisibilityState : nullptr;
 	float MaxDrawDistanceScale = GetCachedScalabilityCVars().ViewDistanceScale;
 	MaxDrawDistanceScale *= GetCachedScalabilityCVars().CalculateFieldOfViewDistanceScale(View.DesiredFOV);
@@ -1690,8 +1691,12 @@ void FRelevancePacket::ComputeRelevance(FDynamicPrimitiveIndexList& DynamicPrimi
 								// Specific logic for mobile packets
 								if (ShadingPath == EShadingPath::Mobile)
 								{
+									if (StaticMeshRelevance.bUseSingleLayerWaterMaterial && bMobileSupportsSM5MaterialNodes)
+									{
+										DrawCommandPacket.AddCommandsForMesh(PrimitiveIndex, PrimitiveSceneInfo, StaticMeshRelevance, StaticMesh, CullingPayloadFlags, Scene, bCanCache, EMeshPass::SingleLayerWaterPass);
+									}
 									// Skydome must not be added to base pass bucket
-									if (!StaticMeshRelevance.bUseSkyMaterial)
+									else if (!StaticMeshRelevance.bUseSkyMaterial)
 									{
 										DrawCommandPacket.AddCommandsForMesh(PrimitiveIndex, PrimitiveSceneInfo, StaticMeshRelevance, StaticMesh, CullingPayloadFlags, Scene, bCanCache, EMeshPass::BasePass);
 										if (!bMobileBasePassAlwaysUsesCSM)
