Reduce size of distance field structures.  Rename cvars in platform.ini to match code definition.  Reduce 'MinPrimitiveModifiedBoundsAllocation' to save memory.

#rnx

diff --git a/Engine/Source/Runtime/Renderer/Private/DistanceFieldLightingShared.h b/Engine/Source/Runtime/Renderer/Private/DistanceFieldLightingShared.h
--- a/Engine/Source/Runtime/Renderer/Private/DistanceFieldLightingShared.h
+++ b/Engine/Source/Runtime/Renderer/Private/DistanceFieldLightingShared.h
@@ -100,7 +100,7 @@ END_SHADER_PARAMETER_STRUCT()
 
 namespace DistanceField
 {
-	constexpr uint32 MinPrimitiveModifiedBoundsAllocation = 16 * 1024;
+	constexpr uint32 MinPrimitiveModifiedBoundsAllocation = 1024;
 
 	FDistanceFieldObjectBufferParameters SetupObjectBufferParameters(FRDGBuilder& GraphBuilder, const FDistanceFieldSceneData& DistanceFieldSceneData);
 	FDistanceFieldAtlasParameters SetupAtlasParameters(FRDGBuilder& GraphBuilder, const FDistanceFieldSceneData& DistanceFieldSceneData);
diff --git a/Engine/Source/Runtime/Renderer/Private/ScenePrivate.h b/Engine/Source/Runtime/Renderer/Private/ScenePrivate.h
--- a/Engine/Source/Runtime/Renderer/Private/ScenePrivate.h
+++ b/Engine/Source/Runtime/Renderer/Private/ScenePrivate.h
@@ -64,6 +64,7 @@
 #include "SceneExtensions.h"
 #include "HeterogeneousVolumes/HeterogeneousVolumes.h"
 #include "ScenePrimitiveUpdates.h"
+#include "RenderTransform.h"
 
 /** Factor by which to grow occlusion tests **/
 #define OCCLUSION_SLOP (1.0f)
@@ -1868,7 +1869,6 @@ private:
 class FPrimitiveAndInstance
 {
 public:
-
 	FPrimitiveAndInstance(const FMatrix& InLocalToWorld, const FBox& InWorldBounds, FPrimitiveSceneInfo* InPrimitive, int32 InInstanceIndex)
 	: Primitive(InPrimitive)
 	, InstanceIndex(InInstanceIndex)
@@ -1881,7 +1881,7 @@ public:
 	FVector Origin;
 	FVector3f TransformRows[3];
 
-	FBox3f WorldBoundsRelativeToOrigin;
+	FRenderBounds WorldBoundsRelativeToOrigin;
 
 	int32 InstanceIndex;
 
@@ -1892,7 +1892,7 @@ public:
 		TransformRows[2] = FVector3f((float)InLocalToWorld.M[2][0], (float)InLocalToWorld.M[2][1], (float)InLocalToWorld.M[2][2]);
 		Origin = FVector(InLocalToWorld.M[3][0], InLocalToWorld.M[3][1], InLocalToWorld.M[3][2]);
 
-		WorldBoundsRelativeToOrigin = (FBox3f)(InWorldBounds.ShiftBy(-Origin));
+		WorldBoundsRelativeToOrigin = InWorldBounds.ShiftBy(-Origin);
 	}
 
 	FORCEINLINE FMatrix GetLocalToWorld() const
@@ -1919,7 +1919,7 @@ public:
 
 	FORCEINLINE FBox GetWorldBounds() const
 	{
-		FBox WorldBoundsRelativeToOriginDoublePrecision = (FBox)WorldBoundsRelativeToOrigin;
+		FBox WorldBoundsRelativeToOriginDoublePrecision = WorldBoundsRelativeToOrigin.ToBox();
 		return WorldBoundsRelativeToOriginDoublePrecision.ShiftBy(Origin);
 	}
 };
