Update OpenXR frame synthesis/space warp CVars (r.Velocity.DirectlyRenderOpenXRMotionVectors, xr.OpenXRFrameSynthesis) to be false by default.

diff --git a/Engine/Plugins/Runtime/OpenXR/Source/OpenXRHMD/Private/OpenXRHMD.cpp b/Engine/Plugins/Runtime/OpenXR/Source/OpenXRHMD/Private/OpenXRHMD.cpp
--- a/Engine/Plugins/Runtime/OpenXR/Source/OpenXRHMD/Private/OpenXRHMD.cpp
+++ b/Engine/Plugins/Runtime/OpenXR/Source/OpenXRHMD/Private/OpenXRHMD.cpp
@@ -158,7 +158,7 @@ static TAutoConsoleVariable<bool> CVarOpenXRAlphaInvertPass(
 
 static TAutoConsoleVariable<bool> CVarOpenXRFrameSynthesis(
 	TEXT("xr.OpenXRFrameSynthesis"),
-	1,
+	0,
 	TEXT("If true and supported via XR_EXT_frame_synthesis or XR_FB_space_warp, enable and submit motion vector and motion vector depth swapchains for frame synthesis.\n")
 	TEXT("Currently only supported when using the Vulkan mobile renderer, using mobile multi-view, and r.Velocity.DirectlyRenderOpenXRMotionVectors=True.\n")
 	TEXT("Because normal velocity rendering is disabled when r.Velocity.DirectlyRenderOpenXRMotionVectors=True, temporal anti-aliasing and motion blur will be automatically disabled."),
diff --git a/Engine/Source/Runtime/Renderer/Private/VelocityRendering.cpp b/Engine/Source/Runtime/Renderer/Private/VelocityRendering.cpp
--- a/Engine/Source/Runtime/Renderer/Private/VelocityRendering.cpp
+++ b/Engine/Source/Runtime/Renderer/Private/VelocityRendering.cpp
@@ -83,7 +83,7 @@ static TAutoConsoleVariable<int32> CVarVelocityPixelShaderMotionVectorWorldOffse
 
 static TAutoConsoleVariable<bool> CVarVelocityDirectlyRenderOpenXRMotionVectors(
 	TEXT("r.Velocity.DirectlyRenderOpenXRMotionVectors"),
-	1,
+	0,
 	TEXT("If true and using the Vulkan mobile forward renderer, the engine will render velocity in the OpenXR motion vector format, at the size recommended by FOpenXRHMD::GetRecommendedMotionVectorTextureSize.\n")
 	TEXT("Because the existing scene depth cannot be used to calculate flattened velocity for stationary objects due to the likely size mismatch, this requires including ALL meshes in the velocity pass, even stationary ones which would usually be excluded.\n")
 	TEXT("This setting disables normal velocity rendering and all other features dependent on it, such as Temporal Anti-Aliasing and Motion Blur."),
