Mobile: Fixed vulkan validation warnings related to incorrect PSO pre-caching for a specific materials

diff --git a/Engine/Source/Runtime/Renderer/Private/MobileBasePass.cpp b/Engine/Source/Runtime/Renderer/Private/MobileBasePass.cpp
--- a/Engine/Source/Runtime/Renderer/Private/MobileBasePass.cpp
+++ b/Engine/Source/Runtime/Renderer/Private/MobileBasePass.cpp
@@ -1065,18 +1065,17 @@ void FMobileBasePassMeshProcessor::CollectPSOInitializers(const FSceneTexturesCo
 
 	FGraphicsPipelineRenderTargetsInfo RenderTargetsInfo;
 	RenderTargetsInfo.NumSamples = 1;
-	if (bTranslucentBasePass && 
-		!SceneTexturesConfig.bRequiresDepthAux) // Really need to detect if we use single pass rendering
+	if (MobileAllowFramebufferFetch(ShaderPlatform) || !bTranslucentBasePass)
 	{		
+		SetupGBufferRenderTargetInfo(SceneTexturesConfig, RenderTargetsInfo, false /*bSetupDepthStencil*/);
+	}
+	else
+	{
 		EPixelFormat SceneColorFormat = SceneTexturesConfig.ColorFormat;
 		ETextureCreateFlags SceneColorCreateFlags = SceneTexturesConfig.ColorCreateFlags;
 		AddRenderTargetInfo(SceneColorFormat, SceneColorCreateFlags, RenderTargetsInfo);
 		ExclusiveDepthStencil = FExclusiveDepthStencil::DepthRead_StencilWrite;
 	}
-	else
-	{
-		SetupGBufferRenderTargetInfo(SceneTexturesConfig, RenderTargetsInfo, false /*bSetupDepthStencil*/);
-	}
 	SetupDepthStencilInfo(PF_DepthStencil, SceneTexturesConfig.DepthCreateFlags, ERenderTargetLoadAction::ELoad,
 		ERenderTargetLoadAction::ELoad, ExclusiveDepthStencil, RenderTargetsInfo);
 	SetupMultiViewInfo(RenderTargetsInfo);
diff --git a/Engine/Source/Runtime/Renderer/Private/MobileDeferredShadingPass.cpp b/Engine/Source/Runtime/Renderer/Private/MobileDeferredShadingPass.cpp
--- a/Engine/Source/Runtime/Renderer/Private/MobileDeferredShadingPass.cpp
+++ b/Engine/Source/Runtime/Renderer/Private/MobileDeferredShadingPass.cpp
@@ -1220,7 +1220,14 @@ void FDeferredMobileLightMaterialPSOCollector::CollectPSOInitializersDirectional
 
 			FGraphicsPipelineRenderTargetsInfo RenderTargetsInfo;
 			RenderTargetsInfo.NumSamples = 1;
-			AddRenderTargetInfo(SceneTexturesConfig.ColorFormat, SceneTexturesConfig.ColorCreateFlags, RenderTargetsInfo);
+			if (MobileAllowFramebufferFetch(GMaxRHIShaderPlatform))
+			{
+				SetupGBufferRenderTargetInfo(SceneTexturesConfig, RenderTargetsInfo, false /*bSetupDepthStencil*/);
+			}
+			else
+			{
+				AddRenderTargetInfo(SceneTexturesConfig.ColorFormat, SceneTexturesConfig.ColorCreateFlags, RenderTargetsInfo);
+			}
 
 			SetupDepthStencilInfo(PF_DepthStencil, SceneTexturesConfig.DepthCreateFlags, ERenderTargetLoadAction::ELoad,
 				ERenderTargetLoadAction::ELoad, FExclusiveDepthStencil::DepthRead_StencilWrite, RenderTargetsInfo);
@@ -1330,7 +1337,17 @@ void FDeferredMobileLightMaterialPSOCollector::CollectPSOInitializersLocal(
 
 			FGraphicsPipelineRenderTargetsInfo RenderTargetsInfo;
 			RenderTargetsInfo.NumSamples = 1;
-			AddRenderTargetInfo(SceneTexturesConfig.ColorFormat, SceneTexturesConfig.ColorCreateFlags, RenderTargetsInfo);
+			if (MobileAllowFramebufferFetch(GMaxRHIShaderPlatform))
+			{
+				SetupGBufferRenderTargetInfo(SceneTexturesConfig, RenderTargetsInfo, false /*bSetupDepthStencil*/);
+			}
+			else
+			{
+				AddRenderTargetInfo(SceneTexturesConfig.ColorFormat, SceneTexturesConfig.ColorCreateFlags, RenderTargetsInfo);
+			}
+			
+			SetupDepthStencilInfo(PF_DepthStencil, SceneTexturesConfig.DepthCreateFlags, ERenderTargetLoadAction::ELoad,
+				ERenderTargetLoadAction::ELoad, FExclusiveDepthStencil::DepthRead_StencilWrite, RenderTargetsInfo);
 
 			GraphicsPSOInit.StatePrecachePSOHash = RHIComputeStatePrecachePSOHash(GraphicsPSOInit);
 			ApplyTargetsInfo(GraphicsPSOInit, RenderTargetsInfo);
