[HWRT] Removed ensure from FRayTracingShaderBindingTable::GetMaxAllocatedStaticSegmentCount().

diff --git a/Engine/Source/Runtime/Renderer/Private/RayTracing/RayTracingShaderBindingTable.cpp b/Engine/Source/Runtime/Renderer/Private/RayTracing/RayTracingShaderBindingTable.cpp
--- a/Engine/Source/Runtime/Renderer/Private/RayTracing/RayTracingShaderBindingTable.cpp
+++ b/Engine/Source/Runtime/Renderer/Private/RayTracing/RayTracingShaderBindingTable.cpp
@@ -371,7 +371,6 @@ uint32 FRayTracingShaderBindingTable::GetNumGeometrySegments() const
 
 uint32 FRayTracingShaderBindingTable::GetMaxAllocatedStaticSegmentCount() const
 {
-	ensure(bStaticAllocationsLocked);
 	return StaticRangeAllocator.GetMaxSize() / NumShaderSlotsPerGeometrySegment;
 }
 
