Fix Bloom not working with TSR when motion blur is disabled. Force the invalidation of the low-res scenecolor for translucency after motion blur pass if motion blur is not run.

#tests AfterMotionBlur Translucency affects bloom as expected in test scene.
#DisabledValidations Preflight

diff --git a/Engine/Source/Runtime/Renderer/Private/PostProcess/PostProcessing.cpp b/Engine/Source/Runtime/Renderer/Private/PostProcess/PostProcessing.cpp
--- a/Engine/Source/Runtime/Renderer/Private/PostProcess/PostProcessing.cpp
+++ b/Engine/Source/Runtime/Renderer/Private/PostProcess/PostProcessing.cpp
@@ -1114,8 +1114,12 @@ void AddPostProcessingPasses(
 		{
 			FScreenPassTextureSlice NewSceneColorSlice = AddAfterPassForSceneColorSlice(EPass::MotionBlur, SceneColorSlice);
 
+			// Force invalidation if MotionBlur pass is not enabled and Post motion blur translucency overrides the scene color 
+			// so that the translucency affects bloom.
+			const bool InvalidationRequired = !PassSequence.IsEnabled(EPass::MotionBlur) && PostMotionBlurTranslucencyResources.IsValid();
+
 			// Invalidate half and quarter res.
-			if (NewSceneColorSlice != SceneColorSlice)
+			if (NewSceneColorSlice != SceneColorSlice || InvalidationRequired)
 			{
 				HalfResSceneColor = FScreenPassTextureSlice();
 				QuarterResSceneColor = FScreenPassTextureSlice();
