[GPUScene] Fixed race condition in when re-allocating instances in a single frame by splitting the clear and set (new primitive ID) steps into two kernels with a barrier in between.

#FailedValidations Preflight

diff --git a/Engine/Source/Runtime/Renderer/Private/GPUScene.cpp b/Engine/Source/Runtime/Renderer/Private/GPUScene.cpp
--- a/Engine/Source/Runtime/Renderer/Private/GPUScene.cpp
+++ b/Engine/Source/Runtime/Renderer/Private/GPUScene.cpp
@@ -2409,7 +2409,7 @@ FGPUScenePrimitiveCollector::FUploadData* FGPUSceneDynamicContext::AllocateDynam
 	return UploadData;
 }
 
-TRDGUniformBufferRef<FGPUSceneWriterUniformParameters> FGPUScene::GetWriteParameters(FRDGBuilder& GraphBuilder, FGPUSceneWriterParameters* LegacyParameters)
+TRDGUniformBufferRef<FGPUSceneWriterUniformParameters> FGPUScene::GetWriteParameters(FRDGBuilder& GraphBuilder, FGPUSceneWriterParameters* LegacyParameters, bool bSkipBarriers)
 {
 	FRegisteredBuffers BufferState = CachedRegisteredBuffers;
 	// If we're not in a begin/end block we need to register the buffers here and now.
@@ -2419,9 +2419,9 @@ TRDGUniformBufferRef<FGPUSceneWriterUniformParameters> FGPUScene::GetWriteParame
 	}
 	FGPUSceneWriterUniformParameters* UBParams = GraphBuilder.AllocParameters<FGPUSceneWriterUniformParameters>();
 	UBParams->CommonParameters = CommonParameters;
-	UBParams->GPUSceneInstanceSceneDataRW = GraphBuilder.CreateUAV(BufferState.InstanceSceneDataBuffer, ERDGUnorderedAccessViewFlags::SkipBarrier);
-	UBParams->GPUSceneInstancePayloadDataRW = GraphBuilder.CreateUAV(BufferState.InstancePayloadDataBuffer, ERDGUnorderedAccessViewFlags::SkipBarrier);
-	UBParams->GPUScenePrimitiveSceneDataRW = GraphBuilder.CreateUAV(BufferState.PrimitiveBuffer, ERDGUnorderedAccessViewFlags::SkipBarrier);
+	UBParams->GPUSceneInstanceSceneDataRW = GraphBuilder.CreateUAV(BufferState.InstanceSceneDataBuffer,		bSkipBarriers ? ERDGUnorderedAccessViewFlags::SkipBarrier : ERDGUnorderedAccessViewFlags::None);
+	UBParams->GPUSceneInstancePayloadDataRW = GraphBuilder.CreateUAV(BufferState.InstancePayloadDataBuffer, bSkipBarriers ? ERDGUnorderedAccessViewFlags::SkipBarrier : ERDGUnorderedAccessViewFlags::None);
+	UBParams->GPUScenePrimitiveSceneDataRW = GraphBuilder.CreateUAV(BufferState.PrimitiveBuffer,			bSkipBarriers ? ERDGUnorderedAccessViewFlags::SkipBarrier : ERDGUnorderedAccessViewFlags::None);
 	TRDGUniformBufferRef<FGPUSceneWriterUniformParameters> Result = GraphBuilder.CreateUniformBuffer(UBParams);
 
 	if (LegacyParameters != nullptr)
@@ -2497,10 +2497,10 @@ void FGPUScene::AddUpdatePrimitiveIdsPass(FRDGBuilder& GraphBuilder, FInstanceGP
 		
 #if ENABLE_SCENE_DATA_DX11_UB_ERROR_WORKAROUND
 		PRAGMA_DISABLE_DEPRECATION_WARNINGS
-		GetWriteParameters(GraphBuilder, &PassParameters->GPUSceneWriterParameters);
+		GetWriteParameters(GraphBuilder, &PassParameters->GPUSceneWriterParameters, false);
 		PRAGMA_ENABLE_DEPRECATION_WARNINGS
 #else
-		PassParameters->GPUSceneWriterUB = GetWriteParameters(GraphBuilder);
+		PassParameters->GPUSceneWriterUB = GetWriteParameters(GraphBuilder, nullptr, false);
 #endif
 		auto ComputeShader = GetGlobalShaderMap(FeatureLevel)->GetShader<FGPUSceneSetInstancePrimitiveIdCS>();
 
@@ -2644,9 +2644,12 @@ void FBatchedPrimitiveShaderData::Setup(const FPrimitiveUniformShaderParameters&
 
 void FGPUScene::AddClearInstancesPass(FRDGBuilder& GraphBuilder, FInstanceCullingOcclusionQueryRenderer* OcclusionQueryRenderer)
 {
+	// TODO: pre-size these & probably move work to async task.
 	FInstanceGPULoadBalancer ClearIdData;
+	FInstanceGPULoadBalancer InitializeIdData;
 #if LOG_INSTANCE_ALLOCATIONS
-	FString RangesStr;
+	FString ClearRangesStr;
+	FString InitializeRangesStr;
 #endif
 	for (FInstanceRange Range : InstanceRangesToClear)
 	{
@@ -2658,10 +2661,21 @@ void FGPUScene::AddClearInstancesPass(FRDGBuilder& GraphBuilder, FInstanceCullin
 		if (Range.NumInstanceSceneDataEntries > 0u)
 		{
 			const uint32 PrimitiveID = (Range.Primitive.Index != INDEX_NONE) ? uint32(Range.Primitive.Index) : INVALID_PRIMITIVE_ID;
-			ClearIdData.Add(Range.InstanceSceneDataOffset, Range.NumInstanceSceneDataEntries, PrimitiveID);
+
+			if (PrimitiveID == INVALID_PRIMITIVE_ID)
+			{
+				ClearIdData.Add(Range.InstanceSceneDataOffset, Range.NumInstanceSceneDataEntries, PrimitiveID);
+#if LOG_INSTANCE_ALLOCATIONS
+				ClearRangesStr.Appendf(TEXT("[%6d, %6d), "), Range.InstanceSceneDataOffset, Range.InstanceSceneDataOffset + Range.NumInstanceSceneDataEntries);
+#endif
+			}
+			else
+			{
+				InitializeIdData.Add(Range.InstanceSceneDataOffset, Range.NumInstanceSceneDataEntries, PrimitiveID);
 #if LOG_INSTANCE_ALLOCATIONS
-			RangesStr.Appendf(TEXT("[%6d, %6d), "), Range.InstanceSceneDataOffset, Range.InstanceSceneDataOffset + Range.NumInstanceSceneDataEntries);
+				InitializeRangesStr.Appendf(TEXT("[%6d, %6d), "), Range.InstanceSceneDataOffset, Range.InstanceSceneDataOffset + Range.NumInstanceSceneDataEntries);
 #endif
+			}
 		}
 	}
 
@@ -2672,9 +2686,11 @@ void FGPUScene::AddClearInstancesPass(FRDGBuilder& GraphBuilder, FInstanceCullin
 	}
 
 #if LOG_INSTANCE_ALLOCATIONS
-	UE_LOG(LogTemp, Warning, TEXT("AddClearInstancesPass: \n%s"), *RangesStr);
+	UE_LOG(LogTemp, Warning, TEXT("AddClearInstancesPass(ClearIdData): \n%s"),  *ClearRangesStr);
+	UE_LOG(LogTemp, Warning, TEXT("AddClearInstancesPass(InitializeIdData): \n%s"),  *InitializeRangesStr);
 #endif
 	AddUpdatePrimitiveIdsPass(GraphBuilder, ClearIdData);
+	AddUpdatePrimitiveIdsPass(GraphBuilder, InitializeIdData);
 	InstanceRangesToClear.Empty();
 }
 
diff --git a/Engine/Source/Runtime/Renderer/Private/GPUScene.h b/Engine/Source/Runtime/Renderer/Private/GPUScene.h
--- a/Engine/Source/Runtime/Renderer/Private/GPUScene.h
+++ b/Engine/Source/Runtime/Renderer/Private/GPUScene.h
@@ -60,7 +60,7 @@ BEGIN_SHADER_PARAMETER_STRUCT(FGPUSceneResourceParameters, RENDERER_API)
 	SHADER_PARAMETER_RDG_BUFFER_SRV(ByteAddressBuffer, GPUSceneLightData)
 END_SHADER_PARAMETER_STRUCT()
 
-//Non-RDG paramters, untracked, use with care (or preferably, don't)
+//Non-RDG parameters, untracked, use with care (or preferably, don't)
 BEGIN_SHADER_PARAMETER_STRUCT(FGPUSceneResourceParametersRHI, RENDERER_API)
 	SHADER_PARAMETER_STRUCT_INCLUDE(FGPUSceneCommonParameters, CommonParameters)
 	SHADER_PARAMETER_SRV(StructuredBuffer<float4>, GPUSceneInstanceSceneData)
@@ -365,7 +365,7 @@ private:
 	/**
 	 * Fills in the FGPUSceneWriterParameters to use for read/write access to the GPU Scene.
 	 */
-	TRDGUniformBufferRef<FGPUSceneWriterUniformParameters> GetWriteParameters(FRDGBuilder& GraphBuilder, FGPUSceneWriterParameters* LegacyParameters = nullptr);
+	TRDGUniformBufferRef<FGPUSceneWriterUniformParameters> GetWriteParameters(FRDGBuilder& GraphBuilder, FGPUSceneWriterParameters* LegacyParameters = nullptr, bool bSkipBarriers = true);
 
 	// set to true if the requested (cvar) tile size is >0, otherwise the tile size is re-derived from the instance count
 	bool bUseTiledInstanceDataLayout = true;
