[TSR] Add exposure offset (r.TSR.ShadingRejection.ExposureOffset)  back so global darkening ghosting can be improved. The value behaves different and is now used to adjust the exposure offset for the reject shading moire luma and guide history. Any legacy project before UE5.5 using the old CVar should consider adjusting the value.

#tests Test on projects provided in UE-243547

diff --git a/Engine/Shaders/Private/TemporalSuperResolution/TSRMeasureFlickeringLuma.usf b/Engine/Shaders/Private/TemporalSuperResolution/TSRMeasureFlickeringLuma.usf
--- a/Engine/Shaders/Private/TemporalSuperResolution/TSRMeasureFlickeringLuma.usf
+++ b/Engine/Shaders/Private/TemporalSuperResolution/TSRMeasureFlickeringLuma.usf
@@ -11,7 +11,7 @@
 
 
 //------------------------------------------------------- PARAMETERS
-
+float ExposureOffsetFactor;
 Texture2D<tsr_half3> SceneColorTexture;
 RWTexture2DArray<tsr_half> FlickeringLumaOutput;
 
@@ -55,7 +55,7 @@ void MainCS(
 				const uint2 PixelOffset = uint2(x, y) * TILE_SIZE;
 				uint2 InputPixelPos = OutputPixelOrigin + PixelOffset;
 			
-				tsr_half3 LinearColor = ColorArray[x + y * TILE_ROWS];
+				tsr_half3 LinearColor = ColorArray[x + y * TILE_ROWS] * tsr_half(ExposureOffsetFactor);
 				tsr_half3 SMColor = LinearToSMCS(LinearColor);
 				tsr_half SMLuma = dot(SMColor, kMoireSMCSWeights);
 				tsr_half GLuma = SMCSToGCS(SMLuma.rrr).r;
diff --git a/Engine/Shaders/Private/TemporalSuperResolution/TSRRejectShading.usf b/Engine/Shaders/Private/TemporalSuperResolution/TSRRejectShading.usf
--- a/Engine/Shaders/Private/TemporalSuperResolution/TSRRejectShading.usf
+++ b/Engine/Shaders/Private/TemporalSuperResolution/TSRRejectShading.usf
@@ -648,7 +648,7 @@ void MainCS(
 		#endif
 
 		// Convert input to shading LDR
-		tsr_tensor_halfC ExposedInput = LinearToSMCS(Input);
+		tsr_tensor_halfC ExposedInput = LinearToSMCS(Input * tsr_half(ExposureOffsetFactor));
 		tsr_tensor_halfC ExposedHistory = GCSToSMCS(History);
 			
 		// Check whether the history resurrection is closer to input, to discard and no pay the register pressure of OriginalResurrectedHistory
@@ -727,7 +727,7 @@ void MainCS(
 		#endif
 		
 		// Convert input to shading LDR
-		tsr_tensor_halfC ExposedInput = LinearToSMCS(Input);
+		tsr_tensor_halfC ExposedInput = LinearToSMCS(Input * tsr_half(ExposureOffsetFactor));
 		tsr_tensor_halfC ExposedResurrectedHistory = GCSToSMCS(ResurrectedHistory);
 			
 		// Mutual Annihilation
@@ -815,7 +815,7 @@ void MainCS(
 		#endif // CONFIG_HISTORY_RESURRECTION
 
 		// Convert input to same color space as the accumulation
-		tsr_tensor_halfC ExposedInput   = AddPerceptionAdd(Input,                /* DestPerceptionAdd = */ kHistoryAccumulationPreceptionAdd);
+		tsr_tensor_halfC ExposedInput   = AddPerceptionAdd(Input * tsr_half(ExposureOffsetFactor),                /* DestPerceptionAdd = */ kHistoryAccumulationPreceptionAdd);
 		tsr_tensor_halfC ExposedHistory = AddPerceptionAdd(GCSToLinear(History), /* DestPerceptionAdd = */ kHistoryAccumulationPreceptionAdd);
 		
 		tsr_tensor_halfC ExposedFinalGuide = UpdateGuide(ExposedInput, ExposedHistory, bIsDisoccluded, RejectionBlendFinal);
diff --git a/Engine/Shaders/Private/TemporalSuperResolution/TSRShadingAnalysis.ush b/Engine/Shaders/Private/TemporalSuperResolution/TSRShadingAnalysis.ush
--- a/Engine/Shaders/Private/TemporalSuperResolution/TSRShadingAnalysis.ush
+++ b/Engine/Shaders/Private/TemporalSuperResolution/TSRShadingAnalysis.ush
@@ -23,7 +23,7 @@ static const tsr_half FilteringWeight = rcp(tsr_half(1.0 + 4 * 0.5 + 4 * 0.25));
 
 
 //------------------------------------------------------- PARAMETERS
-
+float ExposureOffsetFactor;
 float FlickeringFramePeriod;
 uint bPassthroughAlpha;
 
@@ -232,8 +232,8 @@ void ComputeMoireError(
 			
 		tsr_tensor_halfC OriginalInput = ComposeTranslucency(OriginalOpaqueInput, OriginalTranslucencyInput);
 
-		OriginalInput       = LinearToSMCS(OriginalInput);
-		OriginalOpaqueInput = LinearToSMCS(OriginalOpaqueInput);
+		OriginalInput       = LinearToSMCS(OriginalInput * tsr_half(ExposureOffsetFactor));
+		OriginalOpaqueInput = LinearToSMCS(OriginalOpaqueInput * tsr_half(ExposureOffsetFactor));
 
 		tsr_tensor_half SceneColorLuma;
 		tsr_tensor_half TranslucencyChange;
diff --git a/Engine/Shaders/Private/TemporalSuperResolution/TSRVisualize.usf b/Engine/Shaders/Private/TemporalSuperResolution/TSRVisualize.usf
--- a/Engine/Shaders/Private/TemporalSuperResolution/TSRVisualize.usf
+++ b/Engine/Shaders/Private/TemporalSuperResolution/TSRVisualize.usf
@@ -43,6 +43,7 @@ uint bReprojectionField;
 float MaxHistorySampleCount;
 float OutputToHistoryResolutionFractionSquare;
 float FlickeringFramePeriod;
+float ExposureOffsetFactor;
 
 Texture2D<float2>    PrevDistortingDisplacementTexture;
 Texture2D<float2>    ResurrectedDistortingDisplacementTexture;
@@ -349,8 +350,8 @@ tsr_half3 DisplayAntiFlickering(FContext Ctx)
 		
 		tsr_half3 OriginalInput = OriginalOpaqueInput * OriginalTranslucencyInput.a + OriginalTranslucencyInput.rgb;
 		
-		OriginalInput = LinearToSMCS(OriginalInput);
-		OriginalOpaqueInput = LinearToSMCS(OriginalOpaqueInput);
+		OriginalInput = LinearToSMCS(OriginalInput * tsr_half(ExposureOffsetFactor));
+		OriginalOpaqueInput = LinearToSMCS(OriginalOpaqueInput * tsr_half(ExposureOffsetFactor));
 
 		tsr_half SceneColorLuma = dot(OriginalOpaqueInput.rgb, kMoireSMCSWeights);
 		tsr_half TranslucencyChange = dot(abs(OriginalOpaqueInput.rgb - OriginalInput.rgb), kMoireSMCSWeights);
diff --git a/Engine/Source/Runtime/Renderer/Private/PostProcess/TemporalAA.h b/Engine/Source/Runtime/Renderer/Private/PostProcess/TemporalAA.h
--- a/Engine/Source/Runtime/Renderer/Private/PostProcess/TemporalAA.h
+++ b/Engine/Source/Runtime/Renderer/Private/PostProcess/TemporalAA.h
@@ -251,6 +251,7 @@ struct FTSRPassConfig
 	float ShadingRejectionFlickeringFrameRateCap = 60;
 	float ShadingRejectionFlickeringPeriod = 2.0f;
 	float ShadingRejectionFlickeringMaxParallaxVelocity = 10.0f;
+	float ShadingRejectionExposureOffsetFactor = 0.0f;
 
 	bool  ThinGeometryDetectionEnable = false;
 	float ThinGeometryErrorMultiplier = 200.0f;
diff --git a/Engine/Source/Runtime/Renderer/Private/PostProcess/TemporalSuperResolution.cpp b/Engine/Source/Runtime/Renderer/Private/PostProcess/TemporalSuperResolution.cpp
--- a/Engine/Source/Runtime/Renderer/Private/PostProcess/TemporalSuperResolution.cpp
+++ b/Engine/Source/Runtime/Renderer/Private/PostProcess/TemporalSuperResolution.cpp
@@ -202,6 +202,22 @@ TAutoConsoleVariable<int32> CVarTSRShadingTileOverscan(
 	TEXT("and therefor need to overlap the tile by couple of padding to hide it. Higher means less prones to tiling artifacts, but performance loss."),
 	ECVF_RenderThreadSafe);
 
+TAutoConsoleVariable<float> CVarTSRShadingExposureOffset(
+	TEXT("r.TSR.ShadingRejection.ExposureOffset"), 0,
+	TEXT("The shading rejection needs to have a representative idea how bright a linear color pixel ends up displayed to the user. ")
+	TEXT("And the shading rejection detect if a color become to changed to be visible in the back buffer by comparing to MeasureBackbufferLDRQuantizationError().\n")
+	TEXT("\n")
+	TEXT("It is important to have TSR's MeasureBackbufferLDRQuantizationError() ends up distributed uniformly across ")
+	TEXT("the range of color intensity or it could otherwise disregard some subtle VFX causing ghosting.\n")
+	TEXT("\n")
+	TEXT("This controls adjusts the exposure of the linear color space solely in the TSR's rejection heuristic, such that higher value ")
+	TEXT("lifts the shadows's LDR intensity, meaning MeasureBackbufferLDRQuantizationError() is decreased in these shadows and increased in ")
+	TEXT("the highlights, control directly.\n")
+	TEXT("\n")
+	TEXT("The best TSR internal buffer to verify this is TSR.Flickering.Luminance, either with the \"show VisualizeTemporalUpscaler\" command or in DumpGPU ")
+	TEXT("with the RGB Linear[0;1] source color space against the Tonemaper's output in sRGB source color space.\n"),
+	ECVF_RenderThreadSafe);
+
 TAutoConsoleVariable<int32> CVarTSRLensDistortion(
 	TEXT("r.TSR.LensDistortion"), 1,
 	TEXT("Whether to apply lens distortion in TSR at runtime (enabled by default, requires r.TSR.Support.LensDistortion enabled at cook time)."),
@@ -707,6 +723,7 @@ class FTSRMeasureFlickeringLumaCS : public FTSRShader
 
 	BEGIN_SHADER_PARAMETER_STRUCT(FParameters, )
 		SHADER_PARAMETER_STRUCT(FScreenPassTextureViewportParameters, InputInfo)
+		SHADER_PARAMETER(float, ExposureOffsetFactor)
 		SHADER_PARAMETER_RDG_TEXTURE(Texture2D, SceneColorTexture)
 		SHADER_PARAMETER_RDG_TEXTURE_UAV(RWTexture2DArray, FlickeringLumaOutput)
 	END_SHADER_PARAMETER_STRUCT()
@@ -865,6 +882,7 @@ class FTSRDecimateHistoryCS : public FTSRShader
 		SHADER_PARAMETER(FVector2f, ResurrectionGuideUVViewportBilinearMin)
 		SHADER_PARAMETER(FVector2f, ResurrectionGuideUVViewportBilinearMax)
 		SHADER_PARAMETER(FVector3f, HistoryGuideQuantizationError)
+		SHADER_PARAMETER(float, ExposureOffsetFactor)
 		SHADER_PARAMETER(float, ResurrectionFrameIndex)
 		SHADER_PARAMETER(float, PrevFrameIndex)
 		SHADER_PARAMETER(FMatrix44f, ClipToResurrectionClip)
@@ -940,6 +958,7 @@ class FTSRRejectShadingCS : public FTSRConvolutionNetworkShader
 		SHADER_PARAMETER(int32, bEnableResurrection)
 		SHADER_PARAMETER(int32, bEnableFlickeringHeuristic)
 		SHADER_PARAMETER(int32, bPassthroughAlpha)
+		SHADER_PARAMETER(float, ExposureOffsetFactor)
 
 		SHADER_PARAMETER_RDG_TEXTURE(Texture2D, InputTexture)
 		SHADER_PARAMETER_RDG_TEXTURE_SRV(Texture2D, InputMoireLumaTexture)
@@ -1411,6 +1430,7 @@ class FTSRVisualizeCS : public FTSRShader
 		SHADER_PARAMETER(float, MaxHistorySampleCount)
 		SHADER_PARAMETER(float, OutputToHistoryResolutionFractionSquare)
 		SHADER_PARAMETER(float, FlickeringFramePeriod)
+		SHADER_PARAMETER(float, ExposureOffsetFactor)
 
 		SHADER_PARAMETER_RDG_TEXTURE(Texture2D, PrevDistortingDisplacementTexture)
 		SHADER_PARAMETER_RDG_TEXTURE(Texture2D, ResurrectedDistortingDisplacementTexture)
@@ -1459,6 +1479,12 @@ bool ComposeSeparateTranslucencyInTSR(const FViewInfo& View)
 	return true;
 }
 
+float GetTSRExposureOffsetFactor()
+{
+	float TSRExposureOffsetFactor = FMath::Pow(2.0f, CVarTSRShadingExposureOffset.GetValueOnRenderThread());
+	return TSRExposureOffsetFactor;
+}
+
 static FRDGTextureUAVRef CreateDummyUAV(FRDGBuilder& GraphBuilder, EPixelFormat PixelFormat)
 {
 	FRDGTextureDesc Desc = FRDGTextureDesc::Create2D(
@@ -1686,6 +1712,7 @@ FScreenPassTexture AddTSRMeasureFlickeringLuma(FRDGBuilder& GraphBuilder, FGloba
 		FTSRMeasureFlickeringLumaCS::FParameters* PassParameters = GraphBuilder.AllocParameters<FTSRMeasureFlickeringLumaCS::FParameters>();
 		PassParameters->InputInfo = GetScreenPassTextureViewportParameters(FScreenPassTextureViewport(
 			SceneColor.Texture->Desc.Extent, SceneColor.ViewRect));
+		PassParameters->ExposureOffsetFactor = GetTSRExposureOffsetFactor();
 		PassParameters->SceneColorTexture = SceneColor.Texture;
 		PassParameters->FlickeringLumaOutput = GraphBuilder.CreateUAV(FlickeringLuma.Texture);
 
@@ -1786,6 +1813,7 @@ FTSRPassConfig GetTSRMainPassConfig(const FViewInfo& View)
 	PassConfig.ShadingRejectionFlickeringFrameRateCap = CVarTSRFlickeringFrameRateCap.GetValueOnRenderThread();
 	PassConfig.ShadingRejectionFlickeringPeriod = CVarTSRFlickeringPeriod.GetValueOnRenderThread();
 	PassConfig.ShadingRejectionFlickeringMaxParallaxVelocity = CVarTSRFlickeringMaxParralaxVelocity.GetValueOnRenderThread();
+	PassConfig.ShadingRejectionExposureOffsetFactor = GetTSRExposureOffsetFactor();
 
 	PassConfig.ThinGeometryDetectionEnable =  ShouldEnableThinGeometryDetection();
 	PassConfig.ThinGeometryErrorMultiplier = CVarTSRThinGeometryDetectionErrorMultiplier.GetValueOnRenderThread();
@@ -2591,6 +2619,7 @@ FDefaultTemporalUpscaler::FOutputs AddTemporalSuperResolutionPasses(
 			PassParameters->ResurrectionGuideUVViewportBilinearMin = GetScreenPassTextureViewportParameters(ResurrectionGuideViewport).UVViewportBilinearMin;
 			PassParameters->ResurrectionGuideUVViewportBilinearMax = GetScreenPassTextureViewportParameters(ResurrectionGuideViewport).UVViewportBilinearMax;
 			PassParameters->HistoryGuideQuantizationError = ComputePixelFormatQuantizationError(ReprojectedHistoryGuideTexture->Desc.Format);
+			PassParameters->ExposureOffsetFactor = PassConfig.ShadingRejectionExposureOffsetFactor;
 		}
 
 		PassParameters->ResurrectionFrameIndex = ResurrectionFrameSliceIndex;
@@ -2913,6 +2942,7 @@ FDefaultTemporalUpscaler::FOutputs AddTemporalSuperResolutionPasses(
 		PassParameters->FlickeringFramePeriod = FlickeringFramePeriod;
 		PassParameters->TheoricBlendFactor = 1.0f / (1.0f + MaxHistorySampleCount / OutputToInputResolutionFractionSquare);
 		PassParameters->TileOverscan = TileOverscan;
+		PassParameters->ExposureOffsetFactor = PassConfig.ShadingRejectionExposureOffsetFactor;
 		PassParameters->bEnableResurrection = bCanResurrectHistory;
 		PassParameters->bEnableFlickeringHeuristic = FlickeringFramePeriod > 0.0f;
 		PassParameters->bPassthroughAlpha = IsPrimitiveAlphaHoldoutEnabled(View);
@@ -3530,6 +3560,7 @@ FDefaultTemporalUpscaler::FOutputs AddTemporalSuperResolutionPasses(
 			PassParameters->MaxHistorySampleCount = MaxHistorySampleCount;
 			PassParameters->OutputToHistoryResolutionFractionSquare = OutputToHistoryResolutionFractionSquare;
 			PassParameters->FlickeringFramePeriod = FlickeringFramePeriod;
+			PassParameters->ExposureOffsetFactor = PassConfig.ShadingRejectionExposureOffsetFactor;
 
 			PassParameters->PrevDistortingDisplacementTexture = PrevDistortingDisplacementTexture;
 			PassParameters->ResurrectedDistortingDisplacementTexture = ResurrectedDistortingDisplacementTexture;
