[TSR] Limit the range of the exposureoffset setting for TSR reject shading pass.

diff --git a/Engine/Source/Runtime/Renderer/Private/PostProcess/TemporalSuperResolution.cpp b/Engine/Source/Runtime/Renderer/Private/PostProcess/TemporalSuperResolution.cpp
--- a/Engine/Source/Runtime/Renderer/Private/PostProcess/TemporalSuperResolution.cpp
+++ b/Engine/Source/Runtime/Renderer/Private/PostProcess/TemporalSuperResolution.cpp
@@ -1481,7 +1481,8 @@ bool ComposeSeparateTranslucencyInTSR(const FViewInfo& View)
 
 float GetTSRExposureOffsetFactor()
 {
-	float TSRExposureOffsetFactor = FMath::Pow(2.0f, CVarTSRShadingExposureOffset.GetValueOnRenderThread());
+	float TSRExposureOffsetFactor = FMath::Pow(2.0f, 
+		FMath::Clamp(CVarTSRShadingExposureOffset.GetValueOnRenderThread(), -20.0f, 20.0f));
 	return TSRExposureOffsetFactor;
 }
 
