Apply correct stride to FRayTracingDebugParameters InstanceExtraDataBuffer and InstanceDebugBuffer when falling back to default structured buffers.
Fixes an RHI validation failure on some console platforms when using the ray tracing debug view mode.

diff --git a/Engine/Source/Runtime/Renderer/Private/RayTracing/RayTracingDebug.cpp b/Engine/Source/Runtime/Renderer/Private/RayTracing/RayTracingDebug.cpp
--- a/Engine/Source/Runtime/Renderer/Private/RayTracing/RayTracingDebug.cpp
+++ b/Engine/Source/Runtime/Renderer/Private/RayTracing/RayTracingDebug.cpp
@@ -1721,13 +1721,13 @@ void RenderRayTracingDebug(FRDGBuilder& GraphBuilder, const FScene& Scene, const
 	FRDGBufferRef InstanceExtraDataBuffer = RayTracingScene.GetInstanceExtraDataBuffer(ERayTracingSceneLayer::Base, View.GetRayTracingSceneViewHandle());
 	if (InstanceExtraDataBuffer == nullptr)
 	{
-		InstanceExtraDataBuffer = GSystemTextures.GetDefaultStructuredBuffer(GraphBuilder, sizeof(uint32));
+		InstanceExtraDataBuffer = GSystemTextures.GetDefaultStructuredBuffer(GraphBuilder, sizeof(FRayTracingInstanceExtraData));
 	}
 
 	FRDGBufferRef InstanceDebugBuffer = RayTracingScene.GetInstanceDebugBuffer(ERayTracingSceneLayer::Base);
 	if (InstanceDebugBuffer == nullptr)
 	{
-		InstanceDebugBuffer = GSystemTextures.GetDefaultStructuredBuffer(GraphBuilder, sizeof(uint32));
+		InstanceDebugBuffer = GSystemTextures.GetDefaultStructuredBuffer(GraphBuilder, sizeof(FRayTracingInstanceDebugData));
 	}
 
 	FRDGBufferRef InstanceBuffer = RayTracingScene.GetInstanceBuffer(ERayTracingSceneLayer::Base, View.GetRayTracingSceneViewHandle());
