Substrate - added option to disable tiling of lighting passes on some platform.
This helps when rect primitives are not supported to make those passes faster.

#tests PC editor dx12 and console
#rnx

#DisabledValidations CheckCopyright, CheckBalancedMacros, CheckUnacceptableWords

diff --git a/Engine/Source/Runtime/RenderCore/Private/RenderUtils.cpp b/Engine/Source/Runtime/RenderCore/Private/RenderUtils.cpp
--- a/Engine/Source/Runtime/RenderCore/Private/RenderUtils.cpp
+++ b/Engine/Source/Runtime/RenderCore/Private/RenderUtils.cpp
@@ -2069,6 +2069,12 @@ static TAutoConsoleVariable<int32> CVarSubstrateStochasticLighting(
 	TEXT("Enable stochastic lighting to get better performance."),
 	ECVF_ReadOnly | ECVF_RenderThreadSafe);
 
+static TAutoConsoleVariable<int32> CVarSubstrateTileLightingPasses(
+	TEXT("r.Substrate.TileLightingPasses"),
+	1,
+	TEXT("Enable tiling of full screen directional light or environment lighting passes. Disabling this can help performance on low end platform not supporting Rect primitives."),
+	ECVF_RenderThreadSafe);
+
 namespace Substrate
 {
 	bool IsSubstrateEnabled()
@@ -2333,6 +2339,11 @@ namespace Substrate
 
 		return IsSubstrateEnabled() && bSubstrateHiddenMaterialAssetConversion;
 	}
+
+	bool AreLightingPassesTiled()
+	{
+		return CVarSubstrateTileLightingPasses.GetValueOnAnyThread() > 0;
+	}
 }
 
 bool HasFirstPersonGBufferBit(EShaderPlatform ShaderPlatform)
diff --git a/Engine/Source/Runtime/RenderCore/Public/RenderUtils.h b/Engine/Source/Runtime/RenderCore/Public/RenderUtils.h
--- a/Engine/Source/Runtime/RenderCore/Public/RenderUtils.h
+++ b/Engine/Source/Runtime/RenderCore/Public/RenderUtils.h
@@ -644,6 +644,8 @@ namespace Substrate
 
 	RENDERCORE_API bool IsHiddenMaterialAssetConversionEnabled();
 
+	RENDERCORE_API bool AreLightingPassesTiled();
+
 	RENDERCORE_API uint32 GetRayTracingMaterialPayloadSizeInBytes(bool bFullySimplifiedMaterial);
 
 	RENDERCORE_API uint32 GetBytePerPixel(EShaderPlatform InPlatform);
diff --git a/Engine/Source/Runtime/Renderer/Private/IndirectLightRendering.cpp b/Engine/Source/Runtime/Renderer/Private/IndirectLightRendering.cpp
--- a/Engine/Source/Runtime/Renderer/Private/IndirectLightRendering.cpp
+++ b/Engine/Source/Runtime/Renderer/Private/IndirectLightRendering.cpp
@@ -1730,6 +1730,8 @@ void FDeferredShadingSceneRenderer::RenderDiffuseIndirectAndAmbientOcclusion(
 	}
 }
 
+void GetSubstrateTiledDepthStencilStateNoDepthTest(ESubstrateTileType InTileType, FGraphicsPipelineStateInitializer& GraphicsPSOInit, uint32& OutStencilRef);
+
 static void AddSkyReflectionPass(
 	FRDGBuilder& GraphBuilder, 
 	FViewInfo& View, 
@@ -1873,7 +1875,8 @@ static void AddSkyReflectionPass(
 	PassParameters->VS.TileIndirectBuffer = nullptr;
 	PassParameters->VS.TileListBuffer = nullptr;
 	const bool bSubstrateEnabled = Substrate::IsSubstrateEnabled();
-	if (bSubstrateEnabled)
+	const bool bSubstrateLightingPassesTiled = Substrate::AreLightingPassesTiled();
+	if (bSubstrateEnabled && bSubstrateLightingPassesTiled)
 	{
 		check(SubstrateTileMaterialType <= ESubstrateTileType::EComplexSpecial);
 		PassParameters->VS = Substrate::SetTileParameters(GraphBuilder, View, SubstrateTileMaterialType, SubstrateTilePrimitiveType);
@@ -1887,7 +1890,7 @@ static void AddSkyReflectionPass(
 		PassParameters,
 		ERDGPassFlags::Raster,
 		[PassParameters, &View, PixelShader, bCheckerboardSubsurfaceRendering, SubstrateTileMaterialType, bUseDepthBounds,
-		SubstrateTilePassVertexShader, bSubstrateEnabled, SubstrateTilePrimitiveType](FRDGAsyncTask, FRHICommandList& InRHICmdList)
+		SubstrateTilePassVertexShader, bSubstrateEnabled, SubstrateTilePrimitiveType, bSubstrateLightingPassesTiled](FRDGAsyncTask, FRHICommandList& InRHICmdList)
 	{
 		InRHICmdList.SetViewport(View.ViewRect.Min.X, View.ViewRect.Min.Y, 0.0f, View.ViewRect.Max.X, View.ViewRect.Max.Y, 1.0f);
 
@@ -1931,16 +1934,23 @@ static void AddSkyReflectionPass(
 			InRHICmdList.SetDepthBounds(Values.MinDepth, Values.MaxDepth);
 		}
 
-		if (bSubstrateEnabled)
+
+		uint32 StencilRef = 0u;
+		if (bSubstrateEnabled && !bSubstrateLightingPassesTiled)
+		{
+			GetSubstrateTiledDepthStencilStateNoDepthTest(SubstrateTileMaterialType, GraphicsPSOInit, StencilRef);
+		}
+
+		if (bSubstrateEnabled && bSubstrateLightingPassesTiled)
 		{
 			GraphicsPSOInit.BoundShaderState.VertexShaderRHI = SubstrateTilePassVertexShader.GetVertexShader();
 			GraphicsPSOInit.PrimitiveType = SubstrateTilePrimitiveType;
 		}
 
-		SetGraphicsPipelineState(InRHICmdList, GraphicsPSOInit, 0x0);
+		SetGraphicsPipelineState(InRHICmdList, GraphicsPSOInit, StencilRef);
 		SetShaderParameters(InRHICmdList, PixelShader, PixelShader.GetPixelShader(), PassParameters->PS);
 
-		if (bSubstrateEnabled)
+		if (bSubstrateEnabled && bSubstrateLightingPassesTiled)
 		{
 			SetShaderParameters(InRHICmdList, SubstrateTilePassVertexShader, SubstrateTilePassVertexShader.GetVertexShader(), PassParameters->VS);
 			InRHICmdList.DrawPrimitiveIndirect(PassParameters->VS.TileIndirectBuffer->GetIndirectRHICallBuffer(), Substrate::TileTypeDrawIndirectArgOffset(SubstrateTileMaterialType));
diff --git a/Engine/Source/Runtime/Renderer/Private/LightRendering.cpp b/Engine/Source/Runtime/Renderer/Private/LightRendering.cpp
--- a/Engine/Source/Runtime/Renderer/Private/LightRendering.cpp
+++ b/Engine/Source/Runtime/Renderer/Private/LightRendering.cpp
@@ -2443,6 +2443,26 @@ static void InternalSetBoundingGeometryRasterizerState(FGraphicsPipelineStateIni
 	}
 }
 
+template<ECompareFunction CompareFunction>
+void GetSubstrateTiledDepthStencilState(ESubstrateTileType InTileType, FGraphicsPipelineStateInitializer& GraphicsPSOInit, uint32& OutStencilRef)
+{
+	OutStencilRef = 0;
+	check(Substrate::IsSubstrateEnabled());
+	switch (InTileType)
+	{
+	case ESubstrateTileType::ESimple:			OutStencilRef = Substrate::StencilBit_Fast;				GraphicsPSOInit.DepthStencilState = TStaticDepthStencilState<false, CompareFunction, true, CF_Equal, SO_Keep, SO_Keep, SO_Keep, true, CF_Equal, SO_Keep, SO_Keep, SO_Keep, Substrate::StencilBit_Fast, 0x0>::GetRHI(); break;
+	case ESubstrateTileType::ESingle:			OutStencilRef = Substrate::StencilBit_Single;			GraphicsPSOInit.DepthStencilState = TStaticDepthStencilState<false, CompareFunction, true, CF_Equal, SO_Keep, SO_Keep, SO_Keep, true, CF_Equal, SO_Keep, SO_Keep, SO_Keep, Substrate::StencilBit_Single, 0x0>::GetRHI(); break;
+	case ESubstrateTileType::EComplex:			OutStencilRef = Substrate::StencilBit_Complex;			GraphicsPSOInit.DepthStencilState = TStaticDepthStencilState<false, CompareFunction, true, CF_Equal, SO_Keep, SO_Keep, SO_Keep, true, CF_Equal, SO_Keep, SO_Keep, SO_Keep, Substrate::StencilBit_Complex, 0x0>::GetRHI(); break;
+	case ESubstrateTileType::EComplexSpecial:	OutStencilRef = Substrate::StencilBit_ComplexSpecial; 	GraphicsPSOInit.DepthStencilState = TStaticDepthStencilState<false, CompareFunction, true, CF_Equal, SO_Keep, SO_Keep, SO_Keep, true, CF_Equal, SO_Keep, SO_Keep, SO_Keep, Substrate::StencilBit_ComplexSpecial, 0x0>::GetRHI(); break;
+	default: check(false);
+	}
+}
+
+void GetSubstrateTiledDepthStencilStateNoDepthTest(ESubstrateTileType InTileType, FGraphicsPipelineStateInitializer& GraphicsPSOInit, uint32& OutStencilRef)
+{
+	return GetSubstrateTiledDepthStencilState<CF_Always>(InTileType, GraphicsPSOInit, OutStencilRef);
+}
+
 template<ECompareFunction CompareFunction>
 static uint32 InternalSetBoundingGeometryDepthState(FGraphicsPipelineStateInitializer& GraphicsPSOInit, ESubstrateTileType TileType)
 {
@@ -2451,15 +2471,7 @@ static uint32 InternalSetBoundingGeometryDepthState(FGraphicsPipelineStateInitia
 	uint32 StencilRef = 0u;
 	if (TileType != ESubstrateTileType::ECount)
 	{
-		check(Substrate::IsSubstrateEnabled());
-		switch (TileType)
-		{
-		case ESubstrateTileType::ESimple :			StencilRef = Substrate::StencilBit_Fast;			GraphicsPSOInit.DepthStencilState = TStaticDepthStencilState<false, CompareFunction, true, CF_Equal, SO_Keep, SO_Keep, SO_Keep, true, CF_Equal, SO_Keep, SO_Keep, SO_Keep, Substrate::StencilBit_Fast, 0x0>::GetRHI(); break;
-		case ESubstrateTileType::ESingle :			StencilRef = Substrate::StencilBit_Single;			GraphicsPSOInit.DepthStencilState = TStaticDepthStencilState<false, CompareFunction, true, CF_Equal, SO_Keep, SO_Keep, SO_Keep, true, CF_Equal, SO_Keep, SO_Keep, SO_Keep, Substrate::StencilBit_Single, 0x0>::GetRHI(); break;
-		case ESubstrateTileType::EComplex:			StencilRef = Substrate::StencilBit_Complex;			GraphicsPSOInit.DepthStencilState = TStaticDepthStencilState<false, CompareFunction, true, CF_Equal, SO_Keep, SO_Keep, SO_Keep, true, CF_Equal, SO_Keep, SO_Keep, SO_Keep, Substrate::StencilBit_Complex, 0x0>::GetRHI(); break;
-		case ESubstrateTileType::EComplexSpecial:	StencilRef = Substrate::StencilBit_ComplexSpecial; 	GraphicsPSOInit.DepthStencilState = TStaticDepthStencilState<false, CompareFunction, true, CF_Equal, SO_Keep, SO_Keep, SO_Keep, true, CF_Equal, SO_Keep, SO_Keep, SO_Keep, Substrate::StencilBit_ComplexSpecial, 0x0>::GetRHI(); break;
-		default: check(false);
-		}
+		GetSubstrateTiledDepthStencilState<CompareFunction>(TileType, GraphicsPSOInit, StencilRef);
 	}
 	else
 	{
@@ -2626,7 +2638,16 @@ uint32 SetupLightGraphicsPSOState(
 		GraphicsPSOInit.bDepthBounds = false;
 		GraphicsPSOInit.RasterizerState = TStaticRasterizerState<FM_Solid, CM_None>::GetRHI();
 		GraphicsPSOInit.BoundShaderState.VertexDeclarationRHI = GFilterVertexDeclaration.VertexDeclarationRHI;
-		GraphicsPSOInit.DepthStencilState = TStaticDepthStencilState<false, CF_Always>::GetRHI();
+
+		if(SubstrateTileMaterialType != ESubstrateTileType::ECount && Substrate::IsSubstrateEnabled() && !Substrate::AreLightingPassesTiled())
+		{
+			// Directional light only: when not tiled, we rely on the stencil test to choose which pixel to process in a pass.
+			StencilRef = InternalSetBoundingGeometryDepthState<CF_Always>(GraphicsPSOInit, SubstrateTileMaterialType);
+		}
+		else
+		{
+			GraphicsPSOInit.DepthStencilState = TStaticDepthStencilState<false, CF_Always>::GetRHI();
+		}
 	}
 	else
 	{
@@ -2666,7 +2687,6 @@ static void InternalRenderLight(
 	{
 
 		const bool bIsRadial = LightType != LightType_Directional;
-		const bool bEnableSubstrateTiledPass   = SubstrateTileMaterialType != ESubstrateTileType::ECount;
 		const bool bEnableSubstrateStencilTest = SubstrateTileMaterialType != ESubstrateTileType::ECount && bIsRadial;
 
 		FGraphicsPipelineStateInitializer GraphicsPSOInit;
@@ -2685,6 +2705,8 @@ static void InternalRenderLight(
 
 		if (LightType == LightType_Directional)
 		{
+			const bool bEnableSubstrateTiledPass = SubstrateTileMaterialType != ESubstrateTileType::ECount && Substrate::IsSubstrateEnabled() && Substrate::AreLightingPassesTiled();
+
 			FDeferredLightVS::FPermutationDomain PermutationVectorVS;
 			PermutationVectorVS.Set<FDeferredLightVS::FRadialLight>(false);
 			TShaderMapRef<FDeferredLightVS> VertexShader(View.ShaderMap, PermutationVectorVS);
@@ -2714,7 +2736,7 @@ static void InternalRenderLight(
 			SetGraphicsPipelineState(RHICmdList, GraphicsPSOInit, StencilRef);
 			SetShaderParameters(RHICmdList, PixelShader, PixelShader.GetPixelShader(), PassParameters->PS);
 
-			if (SubstrateTileMaterialType != ECount)
+			if (bEnableSubstrateTiledPass)
 			{
 				check(Substrate::IsSubstrateEnabled());
 				SetShaderParameters(RHICmdList, TileVertexShader, TileVertexShader.GetVertexShader(), VSParameters);
@@ -3687,7 +3709,7 @@ void DeferredLightGlobalPSOCollector(const FSceneTexturesConfig& SceneTexturesCo
 				VSPermutationVector.Set< Substrate::FSubstrateTilePassVS::FEnableTexCoordScreenVector >(true);
 				TShaderMapRef<Substrate::FSubstrateTilePassVS> TileVertexShader(GlobalShaderMap, VSPermutationVector);
 
-				const bool bEnableSubstrateTiledPass = SubstrateTileMaterialType != ESubstrateTileType::ECount;
+				const bool bEnableSubstrateTiledPass = SubstrateTileMaterialType != ESubstrateTileType::ECount && Substrate::IsSubstrateEnabled() && Substrate::AreLightingPassesTiled();
 				GraphicsPSOInit.BoundShaderState.VertexShaderRHI = bEnableSubstrateTiledPass ? TileVertexShader.GetVertexShader() : VertexShader.GetVertexShader();
 				GraphicsPSOInit.BoundShaderState.PixelShaderRHI = PixelShaderRHI;
 			}
