[Stochastic lighting] Fix a typo causing assignment of HistoryGatherUVMinMax to HistoryUVMinMax

#tests tested in editor

diff --git a/Engine/Source/Runtime/Renderer/Private/StochasticLighting/StochasticLighting.cpp b/Engine/Source/Runtime/Renderer/Private/StochasticLighting/StochasticLighting.cpp
--- a/Engine/Source/Runtime/Renderer/Private/StochasticLighting/StochasticLighting.cpp
+++ b/Engine/Source/Runtime/Renderer/Private/StochasticLighting/StochasticLighting.cpp
@@ -267,7 +267,7 @@ void StochasticLighting::FContext::Run(const FViewInfo& View, EReflectionsMethod
 		if (!View.bCameraCut && !View.bPrevTransformsReset)
 		{
 			HistoryScreenPositionScaleBias = StochasticLightingViewState.HistoryScreenPositionScaleBias;
-			HistoryUVMinMax = StochasticLightingViewState.HistoryGatherUVMinMax;
+			HistoryUVMinMax = StochasticLightingViewState.HistoryUVMinMax;
 			HistoryGatherUVMinMax = StochasticLightingViewState.HistoryGatherUVMinMax;
 			HistoryBufferSizeAndInvSize = StochasticLightingViewState.HistoryBufferSizeAndInvSize;
 
