Add shadow-casting heterogeneous volumes to the list of conditionals requiring creation of a screen space shadow mask.

#tests Content Examples project

diff --git a/Engine/Source/Runtime/Renderer/Private/LightRendering.cpp b/Engine/Source/Runtime/Renderer/Private/LightRendering.cpp
--- a/Engine/Source/Runtime/Renderer/Private/LightRendering.cpp
+++ b/Engine/Source/Runtime/Renderer/Private/LightRendering.cpp
@@ -1738,6 +1738,7 @@ void FDeferredShadingSceneRenderer::RenderLights(
 				const bool bUseHairDeepShadow = bDrawShadows && bUseHairLighting && LightSceneProxy.CastsHairStrandsDeepShadow();
 				const bool bRunHairLighting = bUseHairLighting && (bIsHairPass || OcclusionType == ELightOcclusionType::Raytraced);
 				const bool bDrawFirstPersonSelfShadow = bDrawShadows && !bIsHairPass && ShouldRenderFirstPersonSelfShadowForLight(*this, ViewFamily, Views, LightSceneInfo);
+				const bool bShouldHeterogeneousVolumesCastShadows = ShouldHeterogeneousVolumesCastShadows();
 
 				bool bUsedShadowMaskTexture = false;
 
@@ -1769,7 +1770,8 @@ void FDeferredShadingSceneRenderer::RenderLights(
 						!bDrawPreviewIndicator &&
 						SortedLightInfo.SortKey.Fields.LightType != LightType_Directional &&
 						VisibleLightInfo.ContainsOnlyVirtualShadowMaps() &&
-						!bDrawFirstPersonSelfShadow;
+						!bDrawFirstPersonSelfShadow &&
+						!bShouldHeterogeneousVolumesCastShadows;
 					bElideScreenShadowMaskSubPixel = bElideScreenShadowMask && CVarHairStrandsAllowOneTransmittancePass.GetValueOnRenderThread() > 0;
 
 					if (!SharedScreenShadowMaskTexture || !SharedScreenShadowMaskSubPixelTexture)
@@ -2241,7 +2243,7 @@ void FDeferredShadingSceneRenderer::RenderLights(
 						RenderFirstPersonSelfShadow(GraphBuilder, *this, Views, ScreenShadowMaskTexture, FirstPersonSelfShadowInputs, LightSceneInfo);
 					}
 
-					if (ShouldHeterogeneousVolumesCastShadows())
+					if (bShouldHeterogeneousVolumesCastShadows)
 					{
 						RenderVolumetricShadowMaskForLight(GraphBuilder,	SceneTextures.Config, Views, LightSceneInfo.Type, &LightSceneInfo, ScreenShadowMaskTexture);
 					}
