Changed the "stat sceneculling" counters to increase clarity and correctness.

diff --git a/Engine/Source/Runtime/Renderer/Private/SceneCulling/SceneCulling.cpp b/Engine/Source/Runtime/Renderer/Private/SceneCulling/SceneCulling.cpp
--- a/Engine/Source/Runtime/Renderer/Private/SceneCulling/SceneCulling.cpp
+++ b/Engine/Source/Runtime/Renderer/Private/SceneCulling/SceneCulling.cpp
@@ -71,11 +71,15 @@ DECLARE_DWORD_COUNTER_STAT(TEXT("Update Uploaded Blocks"), STAT_SceneCulling_Upl
 
 DECLARE_DWORD_COUNTER_STAT(TEXT("Block Count"), STAT_SceneCulling_BlockCount, STATGROUP_SceneCulling);
 DECLARE_DWORD_COUNTER_STAT(TEXT("Cell Count"), STAT_SceneCulling_CellCount, STATGROUP_SceneCulling);
-DECLARE_DWORD_COUNTER_STAT(TEXT("Item Chunk Count"), STAT_SceneCulling_ItemChunkCount, STATGROUP_SceneCulling);
-DECLARE_DWORD_COUNTER_STAT(TEXT("Used Explicit Item Count"), STAT_SceneCulling_UsedExplicitItemCount, STATGROUP_SceneCulling);
-DECLARE_DWORD_COUNTER_STAT(TEXT("Compressed Item Count"), STAT_SceneCulling_CompressedItemCount, STATGROUP_SceneCulling);
-DECLARE_DWORD_COUNTER_STAT(TEXT("Item Buffer Count"), STAT_SceneCulling_ItemBufferCount, STATGROUP_SceneCulling);
+// Number of chunk IDs that reference explicit chunks OR represents RLE an compressed range
+DECLARE_DWORD_COUNTER_STAT(TEXT("Chunk ID Count"), STAT_SceneCulling_ItemChunkCount, STATGROUP_SceneCulling);
+DECLARE_DWORD_COUNTER_STAT(TEXT("Explicit Chunk Count"), STAT_SceneCulling_ExplicitChunkCount, STATGROUP_SceneCulling);
+DECLARE_DWORD_COUNTER_STAT(TEXT("Compressed Chunk Count"), STAT_SceneCulling_CompressedChunkCount, STATGROUP_SceneCulling);
+// Size of the item buffer (where explicit chunk instance IDs are stored), not all are used.
+DECLARE_DWORD_COUNTER_STAT(TEXT("Item Buffer Size"), STAT_SceneCulling_ItemBufferSize, STATGROUP_SceneCulling);
+DECLARE_FLOAT_COUNTER_STAT(TEXT("Chunk Occupancy"), STAT_SceneCulling_ChunkOccupancy, STATGROUP_SceneCulling);
 DECLARE_DWORD_COUNTER_STAT(TEXT("Total Id Cache Size"), STAT_SceneCulling_IdCacheSize, STATGROUP_SceneCulling);
+DECLARE_DWORD_COUNTER_STAT(TEXT("Scene Instance Count"), STAT_SceneCulling_SceneInstanceCount, STATGROUP_SceneCulling);
 
 #if SC_ENABLE_DETAILED_BUILDER_STATS
 
@@ -2591,11 +2595,15 @@ void FSceneCulling::PublishStats()
 {
 	SET_DWORD_STAT(STAT_SceneCulling_BlockCount, CellBlockData.Num());
 	SET_DWORD_STAT(STAT_SceneCulling_CellCount, CellHeaders.Num());
-	SET_DWORD_STAT(STAT_SceneCulling_ItemChunkCount, PackedCellChunkData.Num());
-	int32 NumUsedItems = PackedCellData.Num() - FreeChunks.Num() * INSTANCE_HIERARCHY_MAX_CHUNK_SIZE;
-	SET_DWORD_STAT(STAT_SceneCulling_UsedExplicitItemCount, NumUsedItems);
-	SET_DWORD_STAT(STAT_SceneCulling_CompressedItemCount, INSTANCE_HIERARCHY_MAX_CHUNK_SIZE * PackedCellChunkData.Num() - NumUsedItems);
-	SET_DWORD_STAT(STAT_SceneCulling_ItemBufferCount, PackedCellData.Num());
+	int32 NumChunks = CellChunkIdAllocator.GetSparselyAllocatedSize();
+	SET_DWORD_STAT(STAT_SceneCulling_ItemChunkCount, NumChunks);
+	int32 NumExplicitChunks = PackedCellData.Num() / INSTANCE_HIERARCHY_MAX_CHUNK_SIZE - FreeChunks.Num();
+	SET_DWORD_STAT(STAT_SceneCulling_ExplicitChunkCount, NumExplicitChunks);
+	SET_DWORD_STAT(STAT_SceneCulling_CompressedChunkCount, NumChunks - NumExplicitChunks);
+	SET_DWORD_STAT(STAT_SceneCulling_ItemBufferSize, PackedCellData.Num());
+	float ChunkOccupancy = float(Scene.GPUScene.GetInstanceSceneDataAllocator().GetSparselyAllocatedSize()) /  float(NumChunks * INSTANCE_HIERARCHY_MAX_CHUNK_SIZE);
+	SET_FLOAT_STAT(STAT_SceneCulling_ChunkOccupancy, ChunkOccupancy);
+	SET_DWORD_STAT(STAT_SceneCulling_SceneInstanceCount, Scene.GPUScene.GetInstanceSceneDataAllocator().GetSparselyAllocatedSize());
 
 	SET_DWORD_STAT(STAT_SceneCulling_IdCacheSize, TotalCellIndexCacheItems);
 
