[Substrate] Workaround for tile shaped artifacts when Substrate and r.Shadow.StencilCulling are enabled. Adds a dummy pixel shader forcing late-Z to the stencil mask write pass.

diff --git a/Engine/Shaders/Private/LateZDummyPS.usf b/Engine/Shaders/Private/LateZDummyPS.usf
new file mode 100644
--- /dev/null
+++ b/Engine/Shaders/Private/LateZDummyPS.usf
@@ -0,0 +1,9 @@
+// Copyright Epic Games, Inc. All Rights Reserved.
+
+#pragma once
+#include "Common.ush"
+
+void MainPS(float4 SVPos : SV_Position, out float OutDepth : SV_Depth)
+{
+	OutDepth = SVPos.z;
+}
\ No newline at end of file
diff --git a/Engine/Source/Runtime/Renderer/Private/ShadowRendering.cpp b/Engine/Source/Runtime/Renderer/Private/ShadowRendering.cpp
--- a/Engine/Source/Runtime/Renderer/Private/ShadowRendering.cpp
+++ b/Engine/Source/Runtime/Renderer/Private/ShadowRendering.cpp
@@ -177,6 +177,13 @@ static FAutoConsoleVariableRef CVarGShadowStencilCulling(
 	ECVF_RenderThreadSafe
 );
 
+// Workaround for UE-300339
+static TAutoConsoleVariable<int32> CVarShadowStencilCullingArtifactWorkaround(
+	TEXT("r.Shadow.StencilCullingArtifactWorkaround"),
+	1,
+	TEXT("Workaround for a seemingly platform specific artifact (tile shaped artifacts) when both Substrate and r.Shadow.StencilCulling are enabled. The workaround forces a dummy pixel shader with late-Z in the stencil masking pass, which seems to fix the bug."),
+	ECVF_RenderThreadSafe);
+
 static TAutoConsoleVariable<int32> CVarFilterMethod(
 	TEXT("r.Shadow.FilterMethod"),
 	0,
@@ -934,6 +941,16 @@ class FWholeSceneDirectionalShadowStencilVS : public FGlobalShader
 };
 IMPLEMENT_GLOBAL_SHADER(FWholeSceneDirectionalShadowStencilVS, "/Engine/Private/ShadowProjectionVertexShader.usf", "WholeSceneDirectionalShadowStencilVS", SF_Vertex);
 
+class FLateZDummyPS : public FGlobalShader
+{
+	DECLARE_GLOBAL_SHADER(FLateZDummyPS);
+	SHADER_USE_PARAMETER_STRUCT(FLateZDummyPS, FGlobalShader)
+
+	BEGIN_SHADER_PARAMETER_STRUCT(FParameters, )
+	END_SHADER_PARAMETER_STRUCT()
+};
+IMPLEMENT_GLOBAL_SHADER(FLateZDummyPS, "/Engine/Private/LateZDummyPS.usf", "MainPS", SF_Pixel);
+
 void FProjectedShadowInfo::SetupProjectionStencilMask(
 	FRHICommandList& RHICmdList, 
 	const FViewInfo* View, 
@@ -1093,9 +1110,19 @@ void FProjectedShadowInfo::SetupProjectionStencilMask(
 
 		// Find the projection shaders.
 		TShaderMapRef<FShadowVolumeBoundProjectionVS> VertexShader(View->ShaderMap);
+		TShaderRef<FLateZDummyPS> PixelShader;
 
 		GraphicsPSOInit.BoundShaderState.VertexDeclarationRHI = GetVertexDeclarationFVector4();
 		GraphicsPSOInit.BoundShaderState.VertexShaderRHI = VertexShader.GetVertexShader();
+		if (CVarShadowStencilCullingArtifactWorkaround.GetValueOnRenderThread() != 0 && IsRHIDeviceNVIDIA())
+		{
+			PixelShader = View->ShaderMap->GetShader<FLateZDummyPS>();
+			GraphicsPSOInit.BoundShaderState.PixelShaderRHI = PixelShader.GetPixelShader();
+		}
+		else
+		{
+			GraphicsPSOInit.BoundShaderState.PixelShaderRHI = nullptr;
+		}
 		GraphicsPSOInit.PrimitiveType = PT_TriangleList;
 
 		SetGraphicsPipelineState(RHICmdList, GraphicsPSOInit, 0);
