[MegaLights translucent lighting volume] Fix a crash caused by dereferencing null LightSceneInfo. Also fix double injection of simple lights handled by MegaLights

#tests tested in editor
#FailedValidations Preflight

diff --git a/Engine/Source/Runtime/Renderer/Private/TranslucentLighting.cpp b/Engine/Source/Runtime/Renderer/Private/TranslucentLighting.cpp
--- a/Engine/Source/Runtime/Renderer/Private/TranslucentLighting.cpp
+++ b/Engine/Source/Runtime/Renderer/Private/TranslucentLighting.cpp
@@ -2531,19 +2531,53 @@ void FilterTranslucencyLightingVolume(
 	}
 }
 
-static void SplitSimpleLightsByView(TConstArrayView<FViewInfo> Views, const FSimpleLightArray& SimpleLights, TArrayView<FSimpleLightArray> OutSimpleLightsByView)
+static void SplitSimpleLightsByView(TConstArrayView<FViewInfo> Views, const FSortedLightSetSceneInfo& SortedLightSet, TArrayView<FSimpleLightArray> OutSimpleLightsByView, bool bMegaLightsTranslucencyVolume)
 {
 	check(OutSimpleLightsByView.Num() == Views.Num());
 
-	for (int32 LightIndex = 0; LightIndex < SimpleLights.InstanceData.Num(); ++LightIndex)
+	const FSimpleLightArray& SimpleLights = SortedLightSet.SimpleLights;
+	const TConstArrayView<FSortedLightSceneInfo> SortedLights = SortedLightSet.SortedLights;
+
+	for (int32 SortedIndex = 0; SortedIndex < SortedLightSet.SimpleLightsEnd; ++SortedIndex)
 	{
+		const FSortedLightSceneInfo& SortedLightInfo = SortedLights[SortedIndex];
+		const int32 SimpleLightsIndex = SortedLightInfo.SimpleLightIndex;
+
+		check(!SortedLightInfo.SortKey.Fields.bIsNotSimpleLight && !SortedLightInfo.SortKey.Fields.bHandledByMegaLights && !SortedLightInfo.LightSceneInfo);
+
 		for (int32 ViewIndex = 0; ViewIndex < Views.Num(); ++ViewIndex)
 		{
-			FSimpleLightPerViewEntry PerViewEntry = SimpleLights.GetViewDependentData(LightIndex, ViewIndex, Views.Num());
-			OutSimpleLightsByView[ViewIndex].InstanceData.Add(SimpleLights.InstanceData[LightIndex]);
+			const FSimpleLightEntry& InstanceEntry = SimpleLights.InstanceData[SimpleLightsIndex];
+			const FSimpleLightPerViewEntry& PerViewEntry = SimpleLights.GetViewDependentData(SimpleLightsIndex, ViewIndex, Views.Num());
+
+			OutSimpleLightsByView[ViewIndex].InstanceData.Add(InstanceEntry);
 			OutSimpleLightsByView[ViewIndex].PerViewData.Add(PerViewEntry);
 		}
 	}
+
+	// Inject via the regular path if MegaLights is not enabled for TLV
+	if (!bMegaLightsTranslucencyVolume)
+	{
+		for (int32 SortedIndex = SortedLightSet.MegaLightsLightStart; SortedIndex < SortedLightSet.MegaLightsSimpleLightsEnd; ++SortedIndex)
+		{
+			const FSortedLightSceneInfo& SortedLightInfo = SortedLights[SortedIndex];
+			const int32 SimpleLightsIndex = SortedLightInfo.SimpleLightIndex;
+
+			if (!SortedLightInfo.SortKey.Fields.bIsNotSimpleLight)
+			{
+				check(SortedLightInfo.SortKey.Fields.bHandledByMegaLights && !SortedLightInfo.LightSceneInfo);
+
+				for (int32 ViewIndex = 0; ViewIndex < Views.Num(); ++ViewIndex)
+				{
+					const FSimpleLightEntry& InstanceEntry = SimpleLights.InstanceData[SimpleLightsIndex];
+					const FSimpleLightPerViewEntry& PerViewEntry = SimpleLights.GetViewDependentData(SimpleLightsIndex, ViewIndex, Views.Num());
+
+					OutSimpleLightsByView[ViewIndex].InstanceData.Add(InstanceEntry);
+					OutSimpleLightsByView[ViewIndex].PerViewData.Add(PerViewEntry);
+				}
+			}
+		}
+	}
 }
 
 void FDeferredShadingSceneRenderer::GatherTranslucencyVolumeMarkedVoxels(FRDGBuilder& GraphBuilder)
@@ -2644,7 +2678,7 @@ void FDeferredShadingSceneRenderer::RenderTranslucencyLightingVolume(FRDGBuilder
 		auto& SimpleLightsByView = *GraphBuilder.AllocObject<TArray<FSimpleLightArray, SceneRenderingAllocator>>();
 		SimpleLightsByView.SetNum(Views.Num());
 
-		SplitSimpleLightsByView(Views, SimpleLights, SimpleLightsByView);
+		SplitSimpleLightsByView(Views, SortedLightSet, SimpleLightsByView, bMegaLightsTranslucencyVolume);
 
 		for (int32 ViewIndex = 0; ViewIndex < Views.Num(); ++ViewIndex)
 		{
@@ -2669,7 +2703,10 @@ void FDeferredShadingSceneRenderer::RenderTranslucencyLightingVolume(FRDGBuilder
 		{
 			const FSortedLightSceneInfo& SortedLightInfo = SortedLights[LightIndex];
 
-			if (bMegaLightsTranslucencyVolume && SortedLightInfo.SortKey.Fields.bHandledByMegaLights)
+			if ((bMegaLightsTranslucencyVolume && SortedLightInfo.SortKey.Fields.bHandledByMegaLights)
+				// This is a simple light handled by MegaLights. If MegaLights is enabled for TLV, injection is handled later by
+				// InjectTranslucencyLightingVolumeMegaLights(). Otherwise, it will be injected above along with regular simple lights.
+				|| !SortedLightInfo.SortKey.Fields.bIsNotSimpleLight)
 			{
 				continue;
 			}
