Fix Stochastic lighting classication pass incorrectly running on SM5 platform when Substrate is enabled with Format=0.

diff --git a/Engine/Source/Runtime/Renderer/Private/Substrate/Substrate.cpp b/Engine/Source/Runtime/Renderer/Private/Substrate/Substrate.cpp
--- a/Engine/Source/Runtime/Renderer/Private/Substrate/Substrate.cpp
+++ b/Engine/Source/Runtime/Renderer/Private/Substrate/Substrate.cpp
@@ -244,7 +244,7 @@ bool IsClassificationAsync()
 
 bool UsesStochasticLightingClassification(EShaderPlatform InPlatform)
 {
-	return Substrate::IsSubstrateEnabled() && Substrate::IsSubstrateBlendableGBufferEnabled(InPlatform) && CVarSubstrateUseStochasticLightingClassification.GetValueOnRenderThread();
+	return Substrate::IsSubstrateEnabled() && Substrate::IsSubstrateBlendableGBufferEnabled(InPlatform) && GetMaxSupportedFeatureLevel(InPlatform) > ERHIFeatureLevel::SM5 && CVarSubstrateUseStochasticLightingClassification.GetValueOnRenderThread();
 }
 
 static EPixelFormat GetClassificationTileFormat(const FIntPoint& InResolution, uint32 InTileEncoding)
