Path Tracer: Avoid crash from missing shaders is path tracing codepath is reached while path tracing is not enabled

#rnx

diff --git a/Engine/Source/Runtime/Renderer/Private/PathTracing.cpp b/Engine/Source/Runtime/Renderer/Private/PathTracing.cpp
--- a/Engine/Source/Runtime/Renderer/Private/PathTracing.cpp
+++ b/Engine/Source/Runtime/Renderer/Private/PathTracing.cpp
@@ -3499,20 +3499,17 @@ void FDeferredShadingSceneRenderer::RenderPathTracing(
 	// a scope for the whole path tracer (which would create the same problem).
 	RDG_GPU_MASK_SCOPE(GraphBuilder, AllViewsGPUMask);
 
-	if (!ensureMsgf(FDataDrivenShaderPlatformInfo::GetSupportsPathTracing(View.GetShaderPlatform()),
-		TEXT("Attempting to use path tracing on unsupported platform.")))
+	if (!PathTracing::ShouldCompilePathTracingShadersForProject(View.GetShaderPlatform()))
 	{
+		AddClearUAVPass(GraphBuilder, GraphBuilder.CreateUAV(SceneColorOutputTexture), FLinearColor::Black);
+		OnGetOnScreenMessages.AddLambda([](FScreenMessageWriter& ScreenMessageWriter)->void
+		{
+			static const FText Message = NSLOCTEXT("Renderer", "PathTracingNotEnabled", "Path Tracing is not enabled or supported in the current configuration.");
+			ScreenMessageWriter.DrawLine(Message);
+		});
 		return;
 	}
 
-	if (CVarPathTracing.GetValueOnRenderThread() == 0)
-	{
-		// Path tracing is not enabled on this project (should not be seen by end-users since the menu entry to pick path tracing should be hidden)
-		// If they reach this code through ShowFlag manipulation, they may observe an incomplete image. Is there a way to inform the user here?
-		return;
-	}
-
-
 	if (int DebugMode = GetPathTracingVisualizationMode(); DebugMode >= 0)
 	{
 		const int32 DispatchResX = View.ViewRect.Size().X;
