Path Tracer: Clear render target with proper utility function in case it is not a UAV

#okforgitub public
#rnx

diff --git a/Engine/Source/Runtime/Renderer/Private/PathTracing.cpp b/Engine/Source/Runtime/Renderer/Private/PathTracing.cpp
--- a/Engine/Source/Runtime/Renderer/Private/PathTracing.cpp
+++ b/Engine/Source/Runtime/Renderer/Private/PathTracing.cpp
@@ -3501,7 +3501,7 @@ void FDeferredShadingSceneRenderer::RenderPathTracing(
 
 	if (!PathTracing::ShouldCompilePathTracingShadersForProject(View.GetShaderPlatform()))
 	{
-		AddClearUAVPass(GraphBuilder, GraphBuilder.CreateUAV(SceneColorOutputTexture), FLinearColor::Black);
+		AddClearRenderTargetPass(GraphBuilder, SceneColorOutputTexture);
 		OnGetOnScreenMessages.AddLambda([](FScreenMessageWriter& ScreenMessageWriter)->void
 		{
 			static const FText Message = NSLOCTEXT("Renderer", "PathTracingNotEnabled", "Path Tracing is not enabled or supported in the current configuration.");
