Mobile: Do not allocate ScreenSpaceShadowMask texture if nothing will be rendered into it

diff --git a/Engine/Source/Runtime/Renderer/Private/CapsuleShadowRendering.cpp b/Engine/Source/Runtime/Renderer/Private/CapsuleShadowRendering.cpp
--- a/Engine/Source/Runtime/Renderer/Private/CapsuleShadowRendering.cpp
+++ b/Engine/Source/Runtime/Renderer/Private/CapsuleShadowRendering.cpp
@@ -618,14 +618,14 @@ bool FSceneRenderer::RenderCapsuleDirectShadows(
 
 				if (GRHISupportsBindingTexArrayPerSlice) // Native mobile multi view
 				{
-					PassParameters->RenderTargets[0] = FRenderTargetBinding(ScreenShadowMaskTexture, ERenderTargetLoadAction::ELoad, /*MipIndex*/ 0, /*ArraySlice*/ ViewIndex);
+					PassParameters->RenderTargets[0] = FRenderTargetBinding(ScreenShadowMaskTexture, GetLoadActionIfProduced(ScreenShadowMaskTexture, ERenderTargetLoadAction::EClear), /*MipIndex*/ 0, /*ArraySlice*/ ViewIndex);
 				}
 				else // Mobile multi view fallback or not stereo
 				{
 					// D3D12 RHI does not support setting slice 1 of a texture 2d array without creating it with TexCreate_TargetArraySlicesIndependently.
 					// Unfortunately, this flag does not allow to target the array with SV_RenderTargetArrayIndex anymore so this is not viable for mobile multi view fallback.
 					// In this case, the PS draw calls are rerouted to the correct slice with SV_RenderTargetArrayIndex as a workaround.
-					PassParameters->RenderTargets[0] = FRenderTargetBinding(ScreenShadowMaskTexture, ERenderTargetLoadAction::ELoad);
+					PassParameters->RenderTargets[0] = FRenderTargetBinding(ScreenShadowMaskTexture, GetLoadActionIfProduced(ScreenShadowMaskTexture, ERenderTargetLoadAction::EClear));
 				}
 
 				PassParameters->SceneTextures = GetSceneTextureShaderParameters(View);
diff --git a/Engine/Source/Runtime/Renderer/Private/DistanceFieldShadowing.cpp b/Engine/Source/Runtime/Renderer/Private/DistanceFieldShadowing.cpp
--- a/Engine/Source/Runtime/Renderer/Private/DistanceFieldShadowing.cpp
+++ b/Engine/Source/Runtime/Renderer/Private/DistanceFieldShadowing.cpp
@@ -1137,7 +1137,7 @@ void FProjectedShadowInfo::RenderRayTracedDistanceFieldProjection(
 		const FVector2f RayTracedShadowsTextureViewportMax(RayTracedShadowsTexture.ViewRect.Max);
 
 		FDistanceFieldShadowingUpsample* PassParameters = GraphBuilder.AllocParameters<FDistanceFieldShadowingUpsample>();
-		PassParameters->RenderTargets[0] = FRenderTargetBinding(ScreenShadowMaskTexture, ERenderTargetLoadAction::ELoad);
+		PassParameters->RenderTargets[0] = FRenderTargetBinding(ScreenShadowMaskTexture, GetLoadActionIfProduced(ScreenShadowMaskTexture, ERenderTargetLoadAction::EClear));
 		PassParameters->RenderTargets.DepthStencil = FDepthStencilBinding(SceneTextures.Depth.Target, ERenderTargetLoadAction::ELoad, ERenderTargetLoadAction::ELoad, FExclusiveDepthStencil::DepthRead_StencilRead);
 		
 		PassParameters->PS.View = GetShaderBinding(View.ViewUniformBuffer);
diff --git a/Engine/Source/Runtime/Renderer/Private/ShadowRendering.cpp b/Engine/Source/Runtime/Renderer/Private/ShadowRendering.cpp
--- a/Engine/Source/Runtime/Renderer/Private/ShadowRendering.cpp
+++ b/Engine/Source/Runtime/Renderer/Private/ShadowRendering.cpp
@@ -1285,12 +1285,15 @@ void FProjectedShadowInfo::RenderProjection(
 #endif
 
 	// Interpret null render targets as skipping the render pass; this is currently used by mobile.
-	const ERDGPassFlags PassFlags = CommonPassParameters.RenderTargets[0].GetTexture() == nullptr ? ERDGPassFlags::SkipRenderPass : ERDGPassFlags::None;
+	FRDGTextureRef OutputTexture0 = CommonPassParameters.RenderTargets[0].GetTexture();
+	const ERDGPassFlags PassFlags = OutputTexture0 == nullptr ? ERDGPassFlags::SkipRenderPass : ERDGPassFlags::None;
 
 	auto* PassParameters = GraphBuilder.AllocParameters<FShadowProjectionPassParameters>();
 	*PassParameters = CommonPassParameters;
 	PassParameters->View = View->GetShaderParameters();
 	PassParameters->RenderTargets.MultiViewCount = (View->bIsMobileMultiViewEnabled) ? 2 : (View->Aspects.IsMobileMultiViewEnabled() ? 1 : 0);
+	// LoadAction from CommonPassParameters have to be adjusted depending on whether this is a first write to a target
+	PassParameters->RenderTargets[0].SetLoadAction(GetLoadActionIfProduced(OutputTexture0, ERenderTargetLoadAction::EClear));
 
 	if (RenderTargets.DepthTarget)
 	{
@@ -1556,11 +1559,14 @@ void FProjectedShadowInfo::RenderOnePassPointLightProjection(
 	const bool bCameraInsideLightGeometry = ((FVector)View.ViewMatrices.GetViewOrigin() - LightBounds.Center).SizeSquared() < FMath::Square(LightBounds.W * 1.05f + View.NearClippingDistance * 2.0f);
 
 	// Interpret null render targets as skipping the render pass; this is currently used by mobile.
-	const ERDGPassFlags PassFlags = CommonPassParameters.RenderTargets[0].GetTexture() == nullptr ? ERDGPassFlags::SkipRenderPass : ERDGPassFlags::None;
+	FRDGTextureRef OutputTexture0 = CommonPassParameters.RenderTargets[0].GetTexture();
+	const ERDGPassFlags PassFlags = OutputTexture0 == nullptr ? ERDGPassFlags::SkipRenderPass : ERDGPassFlags::None;
 
 	auto* PassParameters = GraphBuilder.AllocParameters<FShadowProjectionPassParameters>();
 	*PassParameters = CommonPassParameters;
 	PassParameters->View = View.GetShaderParameters();
+	// LoadAction from CommonPassParameters have to be adjusted depending on whether this is a first write to a target
+	PassParameters->RenderTargets[0].SetLoadAction(GetLoadActionIfProduced(OutputTexture0, ERenderTargetLoadAction::EClear));
 
 	if (RenderTargets.DepthTarget)
 	{
@@ -2120,7 +2126,7 @@ void FSceneRenderer::RenderShadowProjections(
 		{
 			DepthStencilTexture = SceneTextures.Depth.Target;
 		}
-		CommonPassParameters.RenderTargets[0] = FRenderTargetBinding(OutputTexture, ERenderTargetLoadAction::ELoad);
+		CommonPassParameters.RenderTargets[0] = FRenderTargetBinding(OutputTexture, GetLoadActionIfProduced(OutputTexture, ERenderTargetLoadAction::EClear));
 		CommonPassParameters.RenderTargets.DepthStencil =
 			bSubPixelShadow ?
 			FDepthStencilBinding(View.HairStrandsViewData.VisibilityData.HairOnlyDepthTexture, ERenderTargetLoadAction::ELoad, ERenderTargetLoadAction::ELoad, ExclusiveDepthStencil) :
@@ -2492,13 +2498,6 @@ FRDGTextureRef FMobileSceneRenderer::RenderMobileShadowProjections(FRDGBuilder&
 	FRDGTextureDesc TargetDesc = FRDGTextureDesc::CreateRenderTargetTextureDesc(Extent, PF_B8G8R8A8, FClearValueBinding::White, TexCreate_RenderTargetable | TexCreate_ShaderResource, bIsMobileMultiView);
 	FRDGTextureRef ScreenShadowMaskTexture = GraphBuilder.CreateTexture(TargetDesc, TEXT("ScreenSpaceShadowMaskTextureMobile"));
 	
-	FRDGTextureClearInfo TextureClearInfo;
-	if (GRHISupportsBindingTexArrayPerSlice)
-	{
-		TextureClearInfo.NumSlices = ScreenShadowMaskTexture->Desc.ArraySize;
-	}
-	AddClearRenderTargetPass(GraphBuilder, ScreenShadowMaskTexture, TextureClearInfo);
-
 	const FMinimalSceneTextures& SceneTextures = GetActiveSceneTextures();
 	for (auto LightIt = Scene->Lights.CreateConstIterator(); LightIt; ++LightIt)
 	{
diff --git a/Engine/Source/Runtime/Renderer/Private/Shadows/ScreenSpaceShadows.cpp b/Engine/Source/Runtime/Renderer/Private/Shadows/ScreenSpaceShadows.cpp
--- a/Engine/Source/Runtime/Renderer/Private/Shadows/ScreenSpaceShadows.cpp
+++ b/Engine/Source/Runtime/Renderer/Private/Shadows/ScreenSpaceShadows.cpp
@@ -231,7 +231,7 @@ void UpsampleScreenSpaceShadows(
 	const FVector2f ShadowsTextureViewportMax(ShadowsTexture.ViewRect.Max);
 
 	FScreenSpaceShadowsUpsample* PassParameters = GraphBuilder.AllocParameters<FScreenSpaceShadowsUpsample>();
-	PassParameters->RenderTargets[0] = FRenderTargetBinding(ScreenShadowMaskTexture, ERenderTargetLoadAction::ELoad);
+	PassParameters->RenderTargets[0] = FRenderTargetBinding(ScreenShadowMaskTexture, GetLoadActionIfProduced(ScreenShadowMaskTexture, ERenderTargetLoadAction::EClear));
 	PassParameters->RenderTargets.DepthStencil = FDepthStencilBinding(SceneTextures.Depth.Resolve, ERenderTargetLoadAction::ELoad, ERenderTargetLoadAction::ELoad, FExclusiveDepthStencil::DepthRead_StencilRead);
 
 	PassParameters->PS.View = GetShaderBinding(View.ViewUniformBuffer);
