Fix crash with Substrate Format=0 on SM6 platform not supporting Lumen nor Megalights.

The merged classificaiton should only be used when Lumen or Megalight is supported.

diff --git a/Engine/Source/Runtime/Renderer/Private/Substrate/Substrate.cpp b/Engine/Source/Runtime/Renderer/Private/Substrate/Substrate.cpp
--- a/Engine/Source/Runtime/Renderer/Private/Substrate/Substrate.cpp
+++ b/Engine/Source/Runtime/Renderer/Private/Substrate/Substrate.cpp
@@ -244,7 +244,11 @@ bool IsClassificationAsync()
 
 bool UsesStochasticLightingClassification(EShaderPlatform InPlatform)
 {
-	return Substrate::IsSubstrateEnabled() && Substrate::IsSubstrateBlendableGBufferEnabled(InPlatform) && GetMaxSupportedFeatureLevel(InPlatform) > ERHIFeatureLevel::SM5 && CVarSubstrateUseStochasticLightingClassification.GetValueOnRenderThread();
+	return 
+		Substrate::IsSubstrateEnabled() && 
+		Substrate::IsSubstrateBlendableGBufferEnabled(InPlatform) && 
+		(DoesPlatformSupportLumenGI(InPlatform) || MegaLights::ShouldCompileShaders(InPlatform)) &&
+		CVarSubstrateUseStochasticLightingClassification.GetValueOnRenderThread();
 }
 
 static EPixelFormat GetClassificationTileFormat(const FIntPoint& InResolution, uint32 InTileEncoding)
