Add StereoSetupCS(0) to FRescaleVariableRateShadingCS since it calls ViewportUVToBufferUV, which now uses ResolvedView when INSTANCED_STEREO && !MOBILE_MULTI_VIEW.
This fixes Contrast Adaptive Shading reprojection being calculated incorrectly on projects where instanced stereo is enabled, such as VR Template.

diff --git a/Engine/Shaders/Private/VariableRateShading/VRSShadingRateReproject.usf b/Engine/Shaders/Private/VariableRateShading/VRSShadingRateReproject.usf
--- a/Engine/Shaders/Private/VariableRateShading/VRSShadingRateReproject.usf
+++ b/Engine/Shaders/Private/VariableRateShading/VRSShadingRateReproject.usf
@@ -63,6 +63,8 @@ void RescaleVariableRateShading(
 	uint3 DispatchThreadId : SV_DispatchThreadID,
 	uint3 GroupThreadId : SV_GroupThreadID)
 {
+	StereoSetupCS(0); // Needed to support ViewportUVToBufferUV, since that function uses ResolvedView when INSTANCED_STEREO && !MOBILE_MULTI_VIEW
+
 	bool bValidThread = DispatchThreadId.x < ScaledSRIDimensions.x && DispatchThreadId.y < ScaledSRIDimensions.y;
 	
 	uint shadingRate = D3D12_SHADING_RATE_2X2 | (D3D12_SHADING_RATE_2X2 << CONSERVATIVE_SHADING_RATE_SHIFT);
diff --git a/Engine/Source/Runtime/Renderer/Private/VariableRateShading/ContrastAdaptiveImageGenerator.cpp b/Engine/Source/Runtime/Renderer/Private/VariableRateShading/ContrastAdaptiveImageGenerator.cpp
--- a/Engine/Source/Runtime/Renderer/Private/VariableRateShading/ContrastAdaptiveImageGenerator.cpp
+++ b/Engine/Source/Runtime/Renderer/Private/VariableRateShading/ContrastAdaptiveImageGenerator.cpp
@@ -142,7 +142,7 @@ class FRescaleVariableRateShadingCS : public FGlobalShader
 
 	BEGIN_SHADER_PARAMETER_STRUCT(FParameters, )
 		SHADER_PARAMETER_RDG_UNIFORM_BUFFER(FSceneTextureUniformParameters, SceneTextures)
-		SHADER_PARAMETER_STRUCT_REF(FViewUniformShaderParameters, View)
+		SHADER_PARAMETER_STRUCT_INCLUDE(FViewShaderParameters, View)
 		SHADER_PARAMETER_RDG_TEXTURE(Texture2D<uint>, InputSRITexture)
 		SHADER_PARAMETER(FVector2f, TextureDimensions)
 		SHADER_PARAMETER(FVector2f, InvTextureDimensions)
@@ -174,7 +174,7 @@ class FRescaleVariableRateShadingCS : public FGlobalShader
 		FRDGTextureUAVRef ScaledConservateSRIUAV)
 	{
 		Parameters.SceneTextures = SceneTextures.UniformBuffer;
-		Parameters.View = ViewInfo.ViewUniformBuffer;
+		Parameters.View = ViewInfo.GetShaderParameters();
 		Parameters.InputSRITexture = InputSRITexture;
 		Parameters.TextureDimensions = TextureDimensions;
 		Parameters.InvTextureDimensions = FVector2f(1.0f / TextureDimensions.X, 1.0f / TextureDimensions.Y);
