Remove unnecessary depth render target binding from UpsampleScreenSpaceShadows pass.

diff --git a/Engine/Source/Runtime/Renderer/Private/Shadows/ScreenSpaceShadows.cpp b/Engine/Source/Runtime/Renderer/Private/Shadows/ScreenSpaceShadows.cpp
--- a/Engine/Source/Runtime/Renderer/Private/Shadows/ScreenSpaceShadows.cpp
+++ b/Engine/Source/Runtime/Renderer/Private/Shadows/ScreenSpaceShadows.cpp
@@ -232,7 +232,6 @@ void UpsampleScreenSpaceShadows(
 
 	FScreenSpaceShadowsUpsample* PassParameters = GraphBuilder.AllocParameters<FScreenSpaceShadowsUpsample>();
 	PassParameters->RenderTargets[0] = FRenderTargetBinding(ScreenShadowMaskTexture, GetLoadActionIfProduced(ScreenShadowMaskTexture, ERenderTargetLoadAction::EClear));
-	PassParameters->RenderTargets.DepthStencil = FDepthStencilBinding(SceneTextures.Depth.Resolve, ERenderTargetLoadAction::ELoad, ERenderTargetLoadAction::ELoad, FExclusiveDepthStencil::DepthRead_StencilRead);
 
 	PassParameters->PS.View = GetShaderBinding(View.ViewUniformBuffer);
 	PassParameters->PS.SceneTextures = SceneTextures.GetSceneTextureShaderParameters(View.GetFeatureLevel());
