Mitigation for GPU hang that can happen when Nanite streamer update is skipped (r.Nanite.ForceStreamingUpdate, default off)

#rnx

diff --git a/Engine/Source/Runtime/Renderer/Private/DeferredShadingRenderer.cpp b/Engine/Source/Runtime/Renderer/Private/DeferredShadingRenderer.cpp
--- a/Engine/Source/Runtime/Renderer/Private/DeferredShadingRenderer.cpp
+++ b/Engine/Source/Runtime/Renderer/Private/DeferredShadingRenderer.cpp
@@ -312,6 +312,12 @@ static TAutoConsoleVariable<float> CVarNaniteViewMeshLODBiasMin(
 	TEXT("Minimum LOD offset for rasterizing Nanite meshes for the main viewport (Default = -2)."),
 	ECVF_RenderThreadSafe);
 
+static TAutoConsoleVariable<bool> CVarNaniteForceStreamingUpdate(
+	TEXT("r.Nanite.ForceStreamingUpdate"), false,
+	TEXT("Temporary mitigation for corruption issues that can occur when Nanite Streaming update isn't called after meshes have been added. This can happen when UpdateSceneCaptureContents gets the SceneUpdateInputs instead."),
+	ECVF_RenderThreadSafe);
+
+
 namespace Lumen
 {
 	extern bool AnyLumenHardwareRayTracingPassEnabled();
@@ -1903,7 +1909,7 @@ void FDeferredShadingSceneRenderer::Render(FRDGBuilder& GraphBuilder, const FSce
 	bool bVisualizeNanite = false;
 	if (bNaniteEnabled)
 	{
-		if (SceneUpdateInputs)
+		if (SceneUpdateInputs || CVarNaniteForceStreamingUpdate.GetValueOnRenderThread())
 		{
 			Nanite::GGlobalResources.Update(GraphBuilder);
 			Nanite::GStreamingManager.BeginAsyncUpdate(GraphBuilder);
