Modify Dummy texture transition in LensFlareBlue to empty structured buffer

diff --git a/Engine/Source/Runtime/Renderer/Private/PostProcess/PostProcessLensFlares.cpp b/Engine/Source/Runtime/Renderer/Private/PostProcess/PostProcessLensFlares.cpp
--- a/Engine/Source/Runtime/Renderer/Private/PostProcess/PostProcessLensFlares.cpp
+++ b/Engine/Source/Runtime/Renderer/Private/PostProcess/PostProcessLensFlares.cpp
@@ -225,12 +225,11 @@ FScreenPassTexture AddLensFlaresPass(
 					FlareViewport.Rect.Max.Y,
 					1.0f);
 
-				// Force a transition to SRVGraphicsNonPixel on a dummy texture to account for VS accessing the texture (using NON_PIXEL_SRV in struct gets merged)
-				if(GRHINeedsSRVGraphicsNonPixelWorkaround)
+				// Force a transition to SRVGraphicsNonPixel on an Empty Structured Buffer to account for VS accessing the texture (using NON_PIXEL_SRV in struct gets merged)
+				if (GRHINeedsSRVGraphicsNonPixelWorkaround)
 				{
-					RHICmdList.Transition(FRHITransitionInfo(GDummyTransitionTexture->GetTextureRHI(), ERHIAccess::Unknown, ERHIAccess::SRVGraphicsNonPixel));
-					
-					
+					RHICmdList.Transition(FRHITransitionInfo(GEmptyStructuredBufferWithUAV->GetRHI(), ERHIAccess::UAVCompute | ERHIAccess::UAVGraphics, ERHIAccess::VertexOrIndexBuffer | ERHIAccess::SRVGraphicsNonPixel));
+
 					FRHIRenderPassInfo RPInfo(PassParameters->RenderTargets[0].GetTexture()->GetRHI(), ERenderTargetActions::Clear_Store);
 					RHICmdList.BeginRenderPass(RPInfo, TEXT("LensFlareBlur"));
 				}
