Fix Lumen with Substrate Format=1 when materials have several closures per pixel.

diff --git a/Engine/Source/Runtime/Renderer/Private/StochasticLighting/StochasticLighting.cpp b/Engine/Source/Runtime/Renderer/Private/StochasticLighting/StochasticLighting.cpp
--- a/Engine/Source/Runtime/Renderer/Private/StochasticLighting/StochasticLighting.cpp
+++ b/Engine/Source/Runtime/Renderer/Private/StochasticLighting/StochasticLighting.cpp
@@ -507,6 +507,8 @@ void FDeferredShadingSceneRenderer::StochasticLightingTileClassificationMark(FRD
 {
 	const ERDGPassFlags ComputePassFlags = ERDGPassFlags::Compute;
 
+	bool bNeedsClear = true;
+
 	for (int32 ViewIndex = 0; ViewIndex < Views.Num(); ++ViewIndex)
 	{
 		FViewInfo& View = Views[ViewIndex];
@@ -667,6 +669,24 @@ void FDeferredShadingSceneRenderer::StochasticLightingTileClassificationMark(FRD
 			StochasticLightingContext.LumenPackedPixelDataUAV = LumenPackedPixelDataUAV;
 			StochasticLightingContext.MegaLightsPackedPixelDataUAV = MegaLightsPackedPixelDataUAV;
 
+			if (Lumen::SupportsMultipleClosureEvaluation(View))
+			{
+				if (LumenPackedPixelData && ClosureCount > 1 && bNeedsClear)
+				{
+					const uint32 LumenInvalidPackedPixelData = 0x30;
+
+					// Initialize LumenPackedPixelData to invalid value for all pixels belonging to slice>0, i.e. closure with index > 0. This is necessary because:
+					// 1) The classification is dispatched only on valid tiles. For closure>0, the LumenPackedPixelData won't be initialized otherwise
+					// 2) The temporal reprojection pass uses LumenPackedPixelData to update history value (in particular NumFamesAccumulated), which are used next frame to prune invalid history data.
+					// Without this, LumenScreenProbeGather will fetch invalid/uninitialized history data for closure>0, causing visual artifacts
+					FRDGTextureUAVRef LumenPackedPixelDataOverflowUAV = GraphBuilder.CreateUAV(FRDGTextureUAVDesc(LumenPackedPixelData, 0/*InMipLevel*/, LumenPackedPixelData->Desc.Format, 1/*InFirstArraySlice*/, ClosureCount - 1u/*InNumArraySlices*/));
+
+					// This value needs to be kept in sync with StochasticLightingCommon.ush
+					AddClearUAVPass(GraphBuilder, LumenPackedPixelDataOverflowUAV, LumenInvalidPackedPixelData, RunConfig.ComputePassFlags);
+					bNeedsClear = false;
+				}
+			}
+			
 			StochasticLightingContext.Run(View, ViewPipelineState.ReflectionsMethod, RunConfig);
 
 			if (Lumen::SupportsMultipleClosureEvaluation(View))
