VulkanRHI: Remove precaching on some debug shaders

#rnx

diff --git a/Engine/Source/Runtime/Renderer/Private/GPUScene.cpp b/Engine/Source/Runtime/Renderer/Private/GPUScene.cpp
--- a/Engine/Source/Runtime/Renderer/Private/GPUScene.cpp
+++ b/Engine/Source/Runtime/Renderer/Private/GPUScene.cpp
@@ -2110,6 +2110,11 @@ public:
 		return UseGPUScene(Parameters.Platform) && ShaderPrint::IsSupported(Parameters.Platform);;
 	}
 
+	static EShaderPermutationPrecacheRequest ShouldPrecachePermutation(const FGlobalShaderPermutationParameters& Parameters)
+	{
+		return EShaderPermutationPrecacheRequest::NotPrecached;
+	}
+
 	static void ModifyCompilationEnvironment(const FGlobalShaderPermutationParameters& Parameters, FShaderCompilerEnvironment& OutEnvironment)
 	{
 		FGlobalShader::ModifyCompilationEnvironment(Parameters, OutEnvironment);
diff --git a/Engine/Source/Runtime/Renderer/Private/GlobalDistanceField.cpp b/Engine/Source/Runtime/Renderer/Private/GlobalDistanceField.cpp
--- a/Engine/Source/Runtime/Renderer/Private/GlobalDistanceField.cpp
+++ b/Engine/Source/Runtime/Renderer/Private/GlobalDistanceField.cpp
@@ -1960,6 +1960,11 @@ class FGlobalDistanceFieldDebugCS : public FGlobalShader
 		return ShouldCompileDistanceFieldShaders(Parameters.Platform) && ShaderPrint::IsSupported(Parameters.Platform);
 	}
 
+	static EShaderPermutationPrecacheRequest ShouldPrecachePermutation(const FGlobalShaderPermutationParameters& Parameters)
+	{
+		return EShaderPermutationPrecacheRequest::NotPrecached;
+	}
+
 	static int32 GetGroupSize()
 	{
 		return 64;
diff --git a/Engine/Source/Runtime/Renderer/Private/VirtualShadowMaps/VirtualShadowMapArray.cpp b/Engine/Source/Runtime/Renderer/Private/VirtualShadowMaps/VirtualShadowMapArray.cpp
--- a/Engine/Source/Runtime/Renderer/Private/VirtualShadowMaps/VirtualShadowMapArray.cpp
+++ b/Engine/Source/Runtime/Renderer/Private/VirtualShadowMaps/VirtualShadowMapArray.cpp
@@ -1276,6 +1276,16 @@ class FGeneratePageFlagsFromFroxelsCS : public FVirtualShadowMapPageManagementSh
 	class FThrottling : SHADER_PERMUTATION_BOOL("PERMUTATION_THROTTLING"); 
 	using FPermutationDomain = TShaderPermutationDomain<FDebugRenderDim, FThrottling>;
 
+	static EShaderPermutationPrecacheRequest ShouldPrecachePermutation(const FGlobalShaderPermutationParameters& Parameters)
+	{
+		FPermutationDomain PermutationVector(Parameters.PermutationId);
+		if (PermutationVector.Get<FDebugRenderDim>())
+		{
+			return EShaderPermutationPrecacheRequest::NotPrecached;
+		}
+		return FVirtualShadowMapPageManagementShader::ShouldPrecachePermutation(Parameters);
+	}
+
 	BEGIN_SHADER_PARAMETER_STRUCT(FParameters, )
 		SHADER_PARAMETER_STRUCT_INCLUDE(FBaseGeneratePageFlagsParameters, Base)
 		SHADER_PARAMETER(uint32, bShouldMarkLocaLights)
