Prevented crash on instance upload for scenes with too many instances by skipping upload for instances outside the instance data buffer.

diff --git a/Engine/Source/Runtime/Renderer/Private/GPUScene.cpp b/Engine/Source/Runtime/Renderer/Private/GPUScene.cpp
--- a/Engine/Source/Runtime/Renderer/Private/GPUScene.cpp
+++ b/Engine/Source/Runtime/Renderer/Private/GPUScene.cpp
@@ -34,7 +34,7 @@
 
 // Useful for debugging
 #define FORCEINLINE_GPUSCENE FORCEINLINE
-//#define FORCEINLINE_GPUSCENE FORCENOINLINE
+// #define FORCEINLINE_GPUSCENE FORCENOINLINE
 
 static_assert(
 	sizeof(FLightSceneData) == SIZEOF_LIGHT_SCENE_DATA,
@@ -302,6 +302,7 @@ struct FPrimitiveUploadInfoHeader
 	uint32 PrimitiveID = INVALID_PRIMITIVE_ID;
 
 	/** Optional */
+	int32 InstanceSceneDataOffset = INDEX_NONE;
 	int32 NumInstanceUploads = 0;
 	int32 InstancePayloadDataOffset = 0;
 	int32 InstancePayloadDataStride = 0;
@@ -320,7 +321,6 @@ struct FInstanceUploadInfo
 {
 	const FInstanceSceneDataBuffers *InstanceSceneDataBuffers = nullptr;
 	TConstArrayView<FInstanceSceneData> PrimitiveInstances;
-	int32 InstanceSceneDataOffset = INDEX_NONE;
 	int32 InstancePayloadDataOffset = INDEX_NONE;
 	int32 InstancePayloadDataStride = 0;
 	int32 InstancePayloadExtensionCount = 0;
@@ -478,6 +478,8 @@ struct FUploadDataSourceAdapterScenePrimitives
 
 		PrimitiveUploadInfo.InstancePayloadDataOffset = PrimitiveSceneInfo->GetInstancePayloadDataOffset();
 		PrimitiveUploadInfo.InstancePayloadDataStride = PrimitiveSceneInfo->GetInstancePayloadDataStride();
+
+		PrimitiveUploadInfo.InstanceSceneDataOffset = PrimitiveSceneInfo->GetInstanceSceneDataOffset();
 		// TODO: Somehow associate the procedural update with the update entry such that we can cross validate the existence of a writer here?
 		if (PrimitiveSceneProxy->IsInstanceDataGPUOnly())
 		{
@@ -487,6 +489,7 @@ struct FUploadDataSourceAdapterScenePrimitives
 		{
 			PrimitiveUploadInfo.NumInstanceUploads = PrimitiveSceneInfo->GetNumInstanceSceneDataEntries();
 		}
+		check(PrimitiveUploadInfo.NumInstanceUploads == 0 || PrimitiveUploadInfo.InstanceSceneDataOffset >= 0);
 	}
 
 	FORCEINLINE_GPUSCENE uint32 PackFlags(FInstanceDataFlags Flags) const
@@ -536,8 +539,6 @@ struct FUploadDataSourceAdapterScenePrimitives
 		InstanceUploadInfo.PrimitiveSceneInfo = PrimitiveSceneInfo;
 #endif
 
-		InstanceUploadInfo.InstanceSceneDataOffset = PrimitiveSceneInfo->GetInstanceSceneDataOffset();
-		check(InstanceUploadInfo.InstanceSceneDataOffset >= 0);
 		InstanceUploadInfo.InstancePayloadDataOffset = PrimitiveSceneInfo->GetInstancePayloadDataOffset();
 		InstanceUploadInfo.InstancePayloadDataStride = PrimitiveSceneInfo->GetInstancePayloadDataStride();
 		InstanceUploadInfo.LastUpdateSceneFrameNumber = SceneFrameNumber;
@@ -1297,6 +1298,13 @@ void FGPUScene::UploadGeneral(FRDGBuilder& GraphBuilder, const FRegisteredBuffer
 	{
 		FPrimitiveUploadInfoHeader UploadInfo;
 		UploadDataSourceAdapter.GetPrimitiveInfoHeader(ItemIndex, UploadInfo);
+
+		if (UploadInfo.InstanceSceneDataOffset + UploadInfo.NumInstanceUploads >= int32(MAX_INSTANCE_ID))
+		{
+			// clamp the uploads such that we don't try to scatter any that will land outside the buffer anyway.
+			UploadInfo.NumInstanceUploads = FMath::Max(0, int32(MAX_INSTANCE_ID) - UploadInfo.InstanceSceneDataOffset);
+		}
+
 		UploadInfo.NumPayloadUploadInstances = UploadInfo.InstancePayloadDataStride ? UploadInfo.NumInstanceUploads : 0;
 		// Figure out how much space we actually have left.
 		int32 PayloadSpaceLeft = FMath::Max(int32(0), int32(MaxInstancePayloadOffset) - int32(UploadInfo.InstancePayloadDataOffset));
@@ -1456,7 +1464,7 @@ void FGPUScene::UploadGeneral(FRDGBuilder& GraphBuilder, const FRegisteredBuffer
 
 						for (uint32 RefIndex = 0; RefIndex < InstanceDataStrideInFloat4s; ++RefIndex)
 						{
-							uint32 ScatterOffset = TaskContext.CalcInstanceDataIndex(UploadInfo.InstanceSceneDataOffset + InstanceIndex, RefIndex);
+							uint32 ScatterOffset = TaskContext.CalcInstanceDataIndex(PrimitiveHeader.InstanceSceneDataOffset + InstanceIndex, RefIndex);
 							FVector4f* DstVector = static_cast<FVector4f*>(TaskContext.InstanceSceneUploader->Set_GetRef(UploadInstanceItemOffset + RefIndex, ScatterOffset));
 							*DstVector = InstanceSceneData.Data[RefIndex];
 						}
@@ -1687,6 +1695,7 @@ struct FUploadDataSourceAdapterDynamicPrimitives
 			PrimitiveUploadInfo.InstancePayloadDataOffset = INDEX_NONE;
 			PrimitiveUploadInfo.InstancePayloadDataStride = 0;
 		}
+		PrimitiveUploadInfo.InstanceSceneDataOffset		= InstanceIDStartOffset + PrimData.LocalInstanceSceneDataOffset;
 		if (PrimData.SourceData.DataWriterGPU.IsBound())
 		{
 			// Only upload if we have data, otherwise expect the delegate to handle missing data
@@ -1729,7 +1738,6 @@ struct FUploadDataSourceAdapterDynamicPrimitives
 			InstanceUploadInfo.PrimitiveID					= PrimitiveIDStartOffset + ItemIndex;
 			InstanceUploadInfo.PrimitiveToWorld				= ShaderParams.LocalToRelativeWorld;
 			InstanceUploadInfo.PrevPrimitiveToWorld			= ShaderParams.PreviousLocalToRelativeWorld;
-			InstanceUploadInfo.InstanceSceneDataOffset		= InstanceIDStartOffset + PrimData.LocalInstanceSceneDataOffset;
 			InstanceUploadInfo.InstancePayloadDataOffset	= PrimData.LocalPayloadDataOffset == INDEX_NONE ? INDEX_NONE : PayloadStartOffset + PrimData.LocalPayloadDataOffset;
 			InstanceUploadInfo.InstancePayloadDataStride	= PrimData.SourceData.GetPayloadFloat4Stride();
 			InstanceUploadInfo.InstanceCustomDataCount		= PrimData.SourceData.NumInstanceCustomDataFloats;
