Take batch size into account when resolving visible instance count from indirect args in the mesh draw command stats

diff --git a/Engine/Source/Runtime/Renderer/Private/MeshDrawCommandStats.cpp b/Engine/Source/Runtime/Renderer/Private/MeshDrawCommandStats.cpp
--- a/Engine/Source/Runtime/Renderer/Private/MeshDrawCommandStats.cpp
+++ b/Engine/Source/Runtime/Renderer/Private/MeshDrawCommandStats.cpp
@@ -384,9 +384,14 @@ void FMeshDrawCommandStatsManager::Update()
 						}
 						else if (IndirectArgsPtr && DrawData.UseInstantCullingIndirectBuffer > 0 && InstanceCullingReadBackData)
 						{
-							const FRHIDrawIndexedIndirectParameters& IndirectArgs = IndirectArgsPtr[PassStats->IndirectArgParameterOffset + IndirectCommandIndex];
-							ensure(DrawData.PrimitiveCount == IndirectArgs.IndexCountPerInstance / 3);
-							DrawData.VisibleInstanceCount = IndirectArgs.InstanceCount;
+							ensure(DrawData.NumBatches > 0);
+							DrawData.VisibleInstanceCount = 0;
+							for (int32 BatchIndex = 0; BatchIndex < DrawData.NumBatches; BatchIndex++)
+							{
+								const FRHIDrawIndexedIndirectParameters& IndirectArgs = IndirectArgsPtr[PassStats->IndirectArgParameterOffset + IndirectCommandIndex + BatchIndex];
+								ensure(DrawData.PrimitiveCount == IndirectArgs.IndexCountPerInstance / 3);
+								DrawData.VisibleInstanceCount += IndirectArgs.InstanceCount;								
+							}
 							ensure(DrawData.VisibleInstanceCount <= DrawData.TotalInstanceCount);
 							Stats.InstanceCullingIndirectInstances += DrawData.VisibleInstanceCount;
 							Stats.InstanceCullingIndirectDrawData.VertexCount += DrawData.VisibleInstanceCount * DrawData.StatsData.NumVertices;
diff --git a/Engine/Source/Runtime/Renderer/Private/MeshDrawCommandStats.h b/Engine/Source/Runtime/Renderer/Private/MeshDrawCommandStats.h
--- a/Engine/Source/Runtime/Renderer/Private/MeshDrawCommandStats.h
+++ b/Engine/Source/Runtime/Renderer/Private/MeshDrawCommandStats.h
@@ -18,7 +18,8 @@ struct FVisibleMeshDrawCommandStatsData
 	FMeshDrawCommandStatsData StatsData;				//< Extracted stats data from the MDC
 	int32 PrimitiveCount = 0;							//< Primitice count of a single instance
 	int32 TotalInstanceCount = 0;						//< Total instance count if no per instance culling would be done
-	int32 VisibleInstanceCount = 0;						//< Visisble instances (read back from indirect args if needed)
+	int16 VisibleInstanceCount = 0;						//< Visisble instances (read back from indirect args if needed)
+	int16 NumBatches = 0;								//< Total number of batched in this specific draw command
 
 	uint32 IndirectArgsOffset : 31;						//< Offset in indirect arg buffer where results can be read (when using instance culling indirect args the pass offset needs to be applied on top of this to get the correct offset)
 	uint32 UseInstantCullingIndirectBuffer : 1;			//< Is the draw command using the shared GPU scene instance culling indirect arg buffer?
diff --git a/Engine/Source/Runtime/Renderer/Private/MeshDrawCommands.cpp b/Engine/Source/Runtime/Renderer/Private/MeshDrawCommands.cpp
--- a/Engine/Source/Runtime/Renderer/Private/MeshDrawCommands.cpp
+++ b/Engine/Source/Runtime/Renderer/Private/MeshDrawCommands.cpp
@@ -939,6 +939,7 @@ void CollectMeshDrawCommandPassStats(
 			{
 				DrawData.UseInstantCullingIndirectBuffer = 1;
 				DrawData.IndirectArgsOffset = MeshDrawCommandInfo.IndirectArgsOffsetOrNumInstances;
+				DrawData.NumBatches = MeshDrawCommandInfo.NumBatches;
 			}
 
 			// Find out the total instance count
