Adding more guards around shadow map generation for Heterogeneous Volumes

1. Preventing Shadow Mask allocation by first checking if any view has bound shadow maps.
2. Adding shadow-type guard before rendering AVSM shadows

diff --git a/Engine/Source/Runtime/Renderer/Private/HeterogeneousVolumes/HeterogeneousVolumes.h b/Engine/Source/Runtime/Renderer/Private/HeterogeneousVolumes/HeterogeneousVolumes.h
--- a/Engine/Source/Runtime/Renderer/Private/HeterogeneousVolumes/HeterogeneousVolumes.h
+++ b/Engine/Source/Runtime/Renderer/Private/HeterogeneousVolumes/HeterogeneousVolumes.h
@@ -701,6 +701,12 @@ void RenderAdaptiveVolumetricShadowMapWithLiveShading(
 	TArray<FVisibleLightInfo, SceneRenderingAllocator>& VisibleLightInfos
 );
 
+bool ShouldRenderVolumetricShadowMaskForLight(
+	FRDGBuilder& GraphBuilder,
+	const TConstArrayView<FViewInfo>& Views,
+	const FLightSceneInfo* LightSceneInfo
+);
+
 void RenderVolumetricShadowMaskForLight(
 	FRDGBuilder& GraphBuilder,
 	// Scene data
diff --git a/Engine/Source/Runtime/Renderer/Private/HeterogeneousVolumes/HeterogeneousVolumesLiveShadingPipeline.cpp b/Engine/Source/Runtime/Renderer/Private/HeterogeneousVolumes/HeterogeneousVolumesLiveShadingPipeline.cpp
--- a/Engine/Source/Runtime/Renderer/Private/HeterogeneousVolumes/HeterogeneousVolumesLiveShadingPipeline.cpp
+++ b/Engine/Source/Runtime/Renderer/Private/HeterogeneousVolumes/HeterogeneousVolumesLiveShadingPipeline.cpp
@@ -3516,6 +3516,25 @@ bool RenderVolumetricShadowMaskForLight(
 
 DECLARE_GPU_STAT_NAMED(VolumetricShadowMask, TEXT("Volumetric Shadow Mask"));
 
+bool ShouldRenderVolumetricShadowMaskForLight(
+	FRDGBuilder& GraphBuilder,
+	const TConstArrayView<FViewInfo>& Views,
+	const FLightSceneInfo* LightSceneInfo
+)
+{
+	for (const FViewInfo& View : Views)
+	{
+		TRDGUniformBufferRef<FAdaptiveVolumetricShadowMapUniformBufferParameters> AVSM = HeterogeneousVolumes::GetAdaptiveVolumetricShadowMapUniformBuffer(GraphBuilder, View.ViewState, LightSceneInfo);
+		TRDGUniformBufferRef<FBeerShadowMapUniformBufferParameters> BeerShadowMap = HeterogeneousVolumes::GetBeerShadowMapUniformBuffer(GraphBuilder, View.ViewState, LightSceneInfo);
+		if (AVSM->GetParameters()->bIsEmpty && BeerShadowMap->GetParameters()->bIsEmpty)
+		{
+			return false;
+		}
+	}
+
+	return true;
+}
+
 void RenderVolumetricShadowMaskForLight(
 	FRDGBuilder& GraphBuilder,
 	// Scene data
diff --git a/Engine/Source/Runtime/Renderer/Private/LightRendering.cpp b/Engine/Source/Runtime/Renderer/Private/LightRendering.cpp
--- a/Engine/Source/Runtime/Renderer/Private/LightRendering.cpp
+++ b/Engine/Source/Runtime/Renderer/Private/LightRendering.cpp
@@ -2243,9 +2243,9 @@ void FDeferredShadingSceneRenderer::RenderLights(
 						RenderFirstPersonSelfShadow(GraphBuilder, *this, Views, ScreenShadowMaskTexture, FirstPersonSelfShadowInputs, LightSceneInfo);
 					}
 
-					if (bShouldHeterogeneousVolumesCastShadows)
+					if (bShouldHeterogeneousVolumesCastShadows && ShouldRenderVolumetricShadowMaskForLight(GraphBuilder, Views, &LightSceneInfo))
 					{
-						RenderVolumetricShadowMaskForLight(GraphBuilder,	SceneTextures.Config, Views, LightSceneInfo.Type, &LightSceneInfo, ScreenShadowMaskTexture);
+						RenderVolumetricShadowMaskForLight(GraphBuilder, SceneTextures.Config, Views, LightSceneInfo.Type, &LightSceneInfo, ScreenShadowMaskTexture);
 					}
 
 					bUsedShadowMaskTexture = true;
diff --git a/Engine/Source/Runtime/Renderer/Private/ShadowDepthRendering.cpp b/Engine/Source/Runtime/Renderer/Private/ShadowDepthRendering.cpp
--- a/Engine/Source/Runtime/Renderer/Private/ShadowDepthRendering.cpp
+++ b/Engine/Source/Runtime/Renderer/Private/ShadowDepthRendering.cpp
@@ -1960,7 +1960,8 @@ void FSceneRenderer::RenderShadowDepthMaps(FRDGBuilder& GraphBuilder, FDynamicSh
 	}
 
 	if (ShouldRenderHeterogeneousVolumes(Scene) && ShouldHeterogeneousVolumesCastShadows() &&
-		(HeterogeneousVolumes::GetShadowPipeline() == HeterogeneousVolumes::EShadowPipeline::LiveShading))
+		(HeterogeneousVolumes::GetShadowPipeline() == HeterogeneousVolumes::EShadowPipeline::LiveShading) &&
+		HeterogeneousVolumes::GetShadowType() == HeterogeneousVolumes::EShadowType::AdaptiveVolumetricShadowMap)
 	{
 		for (int32 ViewIndex = 0; ViewIndex < Views.Num(); ViewIndex++)
 		{
