ShouldRenderVolumetricShadowMaskForLight should return true if any view has a shadow map, otherwise it should return false.
Revert shadow-type guard in ShadowDepthRendering

diff --git a/Engine/Source/Runtime/Renderer/Private/HeterogeneousVolumes/HeterogeneousVolumesLiveShadingPipeline.cpp b/Engine/Source/Runtime/Renderer/Private/HeterogeneousVolumes/HeterogeneousVolumesLiveShadingPipeline.cpp
--- a/Engine/Source/Runtime/Renderer/Private/HeterogeneousVolumes/HeterogeneousVolumesLiveShadingPipeline.cpp
+++ b/Engine/Source/Runtime/Renderer/Private/HeterogeneousVolumes/HeterogeneousVolumesLiveShadingPipeline.cpp
@@ -3526,13 +3526,14 @@ bool ShouldRenderVolumetricShadowMaskForLight(
 	{
 		TRDGUniformBufferRef<FAdaptiveVolumetricShadowMapUniformBufferParameters> AVSM = HeterogeneousVolumes::GetAdaptiveVolumetricShadowMapUniformBuffer(GraphBuilder, View.ViewState, LightSceneInfo);
 		TRDGUniformBufferRef<FBeerShadowMapUniformBufferParameters> BeerShadowMap = HeterogeneousVolumes::GetBeerShadowMapUniformBuffer(GraphBuilder, View.ViewState, LightSceneInfo);
-		if (AVSM->GetParameters()->bIsEmpty && BeerShadowMap->GetParameters()->bIsEmpty)
+		bool bShadowMapExists = (!AVSM->GetParameters()->bIsEmpty) || (!BeerShadowMap->GetParameters()->bIsEmpty);
+		if (bShadowMapExists)
 		{
-			return false;
+			return true;
 		}
 	}
 
-	return true;
+	return false;
 }
 
 void RenderVolumetricShadowMaskForLight(
diff --git a/Engine/Source/Runtime/Renderer/Private/ShadowDepthRendering.cpp b/Engine/Source/Runtime/Renderer/Private/ShadowDepthRendering.cpp
--- a/Engine/Source/Runtime/Renderer/Private/ShadowDepthRendering.cpp
+++ b/Engine/Source/Runtime/Renderer/Private/ShadowDepthRendering.cpp
@@ -1960,8 +1960,7 @@ void FSceneRenderer::RenderShadowDepthMaps(FRDGBuilder& GraphBuilder, FDynamicSh
 	}
 
 	if (ShouldRenderHeterogeneousVolumes(Scene) && ShouldHeterogeneousVolumesCastShadows() &&
-		(HeterogeneousVolumes::GetShadowPipeline() == HeterogeneousVolumes::EShadowPipeline::LiveShading) &&
-		HeterogeneousVolumes::GetShadowType() == HeterogeneousVolumes::EShadowType::AdaptiveVolumetricShadowMap)
+		(HeterogeneousVolumes::GetShadowPipeline() == HeterogeneousVolumes::EShadowPipeline::LiveShading))
 	{
 		for (int32 ViewIndex = 0; ViewIndex < Views.Num(); ViewIndex++)
 		{
