Hide the StereoMotionVectors showflag from the UI and fix crash when the showflag is manually toggled.

diff --git a/Engine/Source/Runtime/Engine/Public/ShowFlagsValues.inl b/Engine/Source/Runtime/Engine/Public/ShowFlagsValues.inl
--- a/Engine/Source/Runtime/Engine/Public/ShowFlagsValues.inl
+++ b/Engine/Source/Runtime/Engine/Public/ShowFlagsValues.inl
@@ -483,7 +483,7 @@ SHOWFLAG_ALWAYS_ACCESSIBLE(AllowPrimitiveAlphaHoldout, SFG_Advanced, NSLOCTEXT("
 SHOWFLAG_ALWAYS_ACCESSIBLE(AlphaInvert, SFG_PostProcess, NSLOCTEXT("UnrealEd", "AlphaInvertSF", "Invert Alpha Channel"))
 
 /** Directly render motion vectors and motion vector depth for OpenXR extensions when normal velocity rendering is unsupported. Motion vectors/depth must be in Vulkan NDC-space and can be requested at any size. */
-SHOWFLAG_ALWAYS_ACCESSIBLE(StereoMotionVectors, SFG_Normal, NSLOCTEXT("UnrealEd", "MotionVectorsSF", "Render Stereo Motion Vectors for OpenXR"))
+SHOWFLAG_ALWAYS_ACCESSIBLE(StereoMotionVectors, SFG_Hidden, NSLOCTEXT("UnrealEd", "MotionVectorsSF", "Render Stereo Motion Vectors for OpenXR"))
 
 /** Geometry Inspection modes. */
 SHOWFLAG_FIXED_IN_SHIPPING(0, Clay, SFG_Hidden, NSLOCTEXT("UnrealEd", "Clay", "Clay"))
diff --git a/Engine/Source/Runtime/Renderer/Private/SceneTextures.cpp b/Engine/Source/Runtime/Renderer/Private/SceneTextures.cpp
--- a/Engine/Source/Runtime/Renderer/Private/SceneTextures.cpp
+++ b/Engine/Source/Runtime/Renderer/Private/SceneTextures.cpp
@@ -766,7 +766,7 @@ void FSceneTextures::InitializeViewFamily(FRDGBuilder& GraphBuilder, FViewFamily
 		SceneTextures.DepthAux = CreateTextureMSAA(GraphBuilder, Desc, TEXT("SceneDepthAuxMS"), TEXT("SceneDepthAux"));
 	}
 
-	if (ViewFamily.EngineShowFlags.StereoMotionVectors)
+	if (ViewFamily.EngineShowFlags.StereoMotionVectors && PlatformSupportsOpenXRMotionVectors(ViewFamily.GetShaderPlatform()))
 	{
 		FTextureRHIRef StereoMotionVectorRHI, StereoMotionVectorDepthRHI;
 		FIntPoint StereoMotionVectorSize;
diff --git a/Engine/Source/Runtime/Renderer/Private/VelocityRendering.cpp b/Engine/Source/Runtime/Renderer/Private/VelocityRendering.cpp
--- a/Engine/Source/Runtime/Renderer/Private/VelocityRendering.cpp
+++ b/Engine/Source/Runtime/Renderer/Private/VelocityRendering.cpp
@@ -568,7 +568,7 @@ void FSceneRenderer::RenderVelocities(
 	{
 		FViewInfo& View = InViews[ViewIndex];
 
-		checkf(!View.Family->EngineShowFlags.StereoMotionVectors,
+		checkf(!(View.Family->EngineShowFlags.StereoMotionVectors && PlatformSupportsOpenXRMotionVectors(View.GetShaderPlatform())),
 			TEXT("Normal velocity rendering is not supported alongside motion vector rendering. If this is causing problems in your project, disable r.Velocity.DirectlyRenderOpenXRMotionVectors."));
 
 		if (View.ShouldRenderView())
@@ -738,7 +738,7 @@ void FSceneRenderer::RenderVelocities(
 void FMobileSceneRenderer::RenderVelocityPass(FRHICommandList& RHICmdList, const FViewInfo& View, const FInstanceCullingDrawParams* InstanceCullingDrawParams)
 {
 	checkSlow(RHICmdList.IsInsideRenderPass());
-	checkf(!View.Family->EngineShowFlags.StereoMotionVectors,
+	checkf(!(View.Family->EngineShowFlags.StereoMotionVectors && PlatformSupportsOpenXRMotionVectors(View.GetShaderPlatform())),
 		TEXT("Normal velocity rendering is not supported alongside motion vector rendering. If this is causing problems in your project, disable r.Velocity.DirectlyRenderOpenXRMotionVectors."));
 
 	if (auto* Pass = View.ParallelMeshDrawCommandPasses[EMeshPass::Velocity])
