Substrate - fix for material blending on mobile

#tests pc editor dx12 lyra and preview mobile renderer with substrate enabled
#rnx

diff --git a/Engine/Source/Runtime/Renderer/Private/MobileBasePass.cpp b/Engine/Source/Runtime/Renderer/Private/MobileBasePass.cpp
--- a/Engine/Source/Runtime/Renderer/Private/MobileBasePass.cpp
+++ b/Engine/Source/Runtime/Renderer/Private/MobileBasePass.cpp
@@ -538,7 +538,7 @@ void MobileBasePass::SetTranslucentRenderState(FMeshPassProcessorRenderState& Dr
 	constexpr ERHIFeatureLevel::Type FeatureLevel = ERHIFeatureLevel::ES3_1;
 	EShaderPlatform ShaderPlatform = GetFeatureLevelShaderPlatform(FeatureLevel);
 
-	if (Substrate::IsSubstrateEnabled())
+	if (Material.IsSubstrateMaterial())
 	{
 		if (Material.IsDualBlendingEnabled() || Material.GetBlendMode() == BLEND_TranslucentColoredTransmittance)
 		{
