Fix incorrect right-eye rendering for WholeScene shadows when using Instanced Stereo + Desktop Vulkan.

diff --git a/Engine/Source/Runtime/Renderer/Private/ShadowRendering.cpp b/Engine/Source/Runtime/Renderer/Private/ShadowRendering.cpp
--- a/Engine/Source/Runtime/Renderer/Private/ShadowRendering.cpp
+++ b/Engine/Source/Runtime/Renderer/Private/ShadowRendering.cpp
@@ -1473,7 +1473,8 @@ void FProjectedShadowInfo::RenderProjectionInternal(
 		}
 	}
 
-	uint32 NumberOfInstances = View->GetStereoPassInstanceFactor();
+	const bool bIsInstancedStereoBypassed = View->bIsInstancedStereoEnabled && !View->bIsMobileMultiViewEnabled && IStereoRendering::IsStereoEyeView(*View);
+	const uint32 NumberOfInstances = bIsInstancedStereoBypassed ? 1 : View->GetStereoPassInstanceFactor();
 	if (IsWholeSceneDirectionalShadow())
 	{
 		RHICmdList.SetStreamSource(0, GClearVertexBuffer.VertexBufferRHI, 0);
