Fix crash in forward rendering when Substrate is enabled with Format=0.

diff --git a/Engine/Source/Runtime/Renderer/Private/Substrate/Substrate.cpp b/Engine/Source/Runtime/Renderer/Private/Substrate/Substrate.cpp
--- a/Engine/Source/Runtime/Renderer/Private/Substrate/Substrate.cpp
+++ b/Engine/Source/Runtime/Renderer/Private/Substrate/Substrate.cpp
@@ -247,6 +247,7 @@ bool UsesStochasticLightingClassification(EShaderPlatform InPlatform)
 	return 
 		Substrate::IsSubstrateEnabled() && 
 		Substrate::IsSubstrateBlendableGBufferEnabled(InPlatform) && 
+		Substrate::UsesSubstrateMaterialBuffer(InPlatform) && // Only needed for deferred lighting
 		(DoesPlatformSupportLumenGI(InPlatform) || MegaLights::ShouldCompileShaders(InPlatform)) &&
 		CVarSubstrateUseStochasticLightingClassification.GetValueOnRenderThread();
 }
