Disabling cascade shadow maps when ShadowsToProject is 1.

diff --git a/Engine/Source/Runtime/Renderer/Private/HeterogeneousVolumes/HeterogeneousVolumesLiveShadingPipeline.cpp b/Engine/Source/Runtime/Renderer/Private/HeterogeneousVolumes/HeterogeneousVolumesLiveShadingPipeline.cpp
--- a/Engine/Source/Runtime/Renderer/Private/HeterogeneousVolumes/HeterogeneousVolumesLiveShadingPipeline.cpp
+++ b/Engine/Source/Runtime/Renderer/Private/HeterogeneousVolumes/HeterogeneousVolumesLiveShadingPipeline.cpp
@@ -2458,7 +2458,13 @@ bool RenderVolumetricShadowMapForLightWithLiveShading(
 				FMatrix RotationMatrix;
 				FMatrix ScaleMatrix;
 
-				if (HeterogeneousVolumes::GetCascadeShadowMode() == HeterogeneousVolumes::ECascadeShadowMode::Frustums)
+				HeterogeneousVolumes::ECascadeShadowMode CascadeShadowMode = HeterogeneousVolumes::GetCascadeShadowMode();
+				if (NumShadowMatrices == 1)
+				{
+					CascadeShadowMode = HeterogeneousVolumes::ECascadeShadowMode::Disabled;
+				}
+
+				if (CascadeShadowMode == HeterogeneousVolumes::ECascadeShadowMode::Frustums)
 				{
 					// Define viewing frustum from SplitNear to Split Far
 					const FMatrix ViewMatrix = View.ViewMatrices.GetViewMatrix();
@@ -2543,7 +2549,7 @@ bool RenderVolumetricShadowMapForLightWithLiveShading(
 					TranslationMatrix = FTranslationMatrix(-WorldSpaceBounds.Origin);
 					ScaleMatrix = FScaleMatrix(FVector(1.0) / LightSpaceBounds.BoxExtent);
 				}
-				else if (HeterogeneousVolumes::GetCascadeShadowMode() == HeterogeneousVolumes::ECascadeShadowMode::Clipmaps)
+				else if (CascadeShadowMode == HeterogeneousVolumes::ECascadeShadowMode::Clipmaps)
 				{
 					const FProjectedShadowInfo* ProjectedShadowInfo = HeterogeneousVolumes::GetProjectedShadowInfo(VisibleLightInfo, ShadowIndex);
 					float Far = ProjectedShadowInfo->CascadeSettings.SplitFar;
@@ -2623,7 +2629,7 @@ bool RenderVolumetricShadowMapForLightWithLiveShading(
 					TranslationMatrix = FTranslationMatrix(-WorldSpaceBounds.Origin);
 					ScaleMatrix = FScaleMatrix(FVector(1.0) / LightSpaceBounds.BoxExtent);
 				}
-				else if (HeterogeneousVolumes::GetCascadeShadowMode() == HeterogeneousVolumes::ECascadeShadowMode::Autofit)
+				else if (CascadeShadowMode == HeterogeneousVolumes::ECascadeShadowMode::Autofit)
 				{
 					TSet<FVolumetricMeshBatch> HeterogeneousVolumesMeshBatches;
 					FBoxSphereBounds::Builder WorldVolumeBoundsBuilder;
@@ -2637,7 +2643,7 @@ bool RenderVolumetricShadowMapForLightWithLiveShading(
 
 					TranslatedWorldOrigin[ShadowIndex] = FVector3f(PreViewTranslation + WorldVolumeBounds.Origin - LightDirection * WorldVolumeBounds.SphereRadius);
 				}
-				else
+				else // CascadeShadowMode == HeterogeneousVolumes::ECascadeShadowMode::Disabled
 				{
 					const FProjectedShadowInfo* ProjectedShadowInfo = HeterogeneousVolumes::GetProjectedShadowInfo(VisibleLightInfo, ShadowIndex);
 					FShadowCascadeSettings ShadowCascadeSettings;
