[SingleLayerWater] Mobile SM5 SLW now writes to the GBuffer, not just SceneColor and Depth.

diff --git a/Engine/Shaders/Private/MobileBasePassPixelShader.usf b/Engine/Shaders/Private/MobileBasePassPixelShader.usf
--- a/Engine/Shaders/Private/MobileBasePassPixelShader.usf
+++ b/Engine/Shaders/Private/MobileBasePassPixelShader.usf
@@ -82,7 +82,7 @@
 #define FORCE_VERTEX_NORMAL (MOBILE_QL_DISABLE_MATERIAL_NORMAL)
 #define SUPPORT_SPOTLIGHTS_SHADOW (MATERIALBLENDING_SOLID || MATERIALBLENDING_MASKED) && PROJECT_MOBILE_ENABLE_MOVABLE_SPOTLIGHT_SHADOWS
 // Sky does not write to depth and must not overwrite GBuffer
-#define MOBILE_USE_GBUFFER (MOBILE_DEFERRED_SHADING && ((MATERIALBLENDING_SOLID || MATERIALBLENDING_MASKED) && !MATERIAL_SHADINGMODEL_SINGLELAYERWATER && !MATERIAL_IS_SKY))
+#define MOBILE_USE_GBUFFER (MOBILE_DEFERRED_SHADING && ((MATERIALBLENDING_SOLID || MATERIALBLENDING_MASKED) && (!MATERIAL_SHADINGMODEL_SINGLELAYERWATER || MOBILE_SUPPORTS_SM5_MATERIAL_NODES) && !MATERIAL_IS_SKY))
 
 #if MATERIAL_SHADINGMODEL_THIN_TRANSLUCENT || SUBSTRATE_BLENDING_TRANSLUCENT_COLOREDTRANSMITTANCE
 	#define MOBILE_TRANSLUCENT_COLOR_TRANSMITTANCE 1
diff --git a/Engine/Source/Runtime/Renderer/Private/MobileShadingRenderer.cpp b/Engine/Source/Runtime/Renderer/Private/MobileShadingRenderer.cpp
--- a/Engine/Source/Runtime/Renderer/Private/MobileShadingRenderer.cpp
+++ b/Engine/Source/Runtime/Renderer/Private/MobileShadingRenderer.cpp
@@ -2682,7 +2682,7 @@ void FMobileSceneRenderer::RenderDeferredMultiPass(FRDGBuilder& GraphBuilder, FS
 			{
 				auto* SLWPassParameters = GraphBuilder.AllocParameters<FMobileRenderPassParameters>();
 				*SLWPassParameters = *PassParameters;
-				// Only SceneColor and Depth
+				SLWPassParameters->RenderTargets = BasePassRenderTargets;
 				SLWPassParameters->RenderTargets.DepthStencil = FDepthStencilBinding(BasePassRenderTargets.DepthStencil.GetTexture(), ERenderTargetLoadAction::ELoad, FExclusiveDepthStencil::DepthWrite_StencilNop);
 
 				if (Scene->GPUScene.IsEnabled())
diff --git a/Engine/Source/Runtime/Renderer/Private/MobileSingleLayerWaterRendering.cpp b/Engine/Source/Runtime/Renderer/Private/MobileSingleLayerWaterRendering.cpp
--- a/Engine/Source/Runtime/Renderer/Private/MobileSingleLayerWaterRendering.cpp
+++ b/Engine/Source/Runtime/Renderer/Private/MobileSingleLayerWaterRendering.cpp
@@ -276,10 +276,7 @@ void FMobileSingleLayerWaterPassMeshProcessor::CollectPSOInitializers(const FSce
 	FMeshPassProcessorRenderState DrawRenderState(PassDrawRenderState);
 
 	FGraphicsPipelineRenderTargetsInfo RenderTargetsInfo;
-	RenderTargetsInfo.NumSamples = 1;
-	EPixelFormat SceneColorFormat = SceneTexturesConfig.ColorFormat;
-	ETextureCreateFlags SceneColorCreateFlags = SceneTexturesConfig.ColorCreateFlags;
-	AddRenderTargetInfo(SceneColorFormat, SceneColorCreateFlags, RenderTargetsInfo);
+	SceneTexturesConfig.GetGBufferRenderTargetsInfo(RenderTargetsInfo, GBL_Default);
 	
 	SetupDepthStencilInfo(PF_DepthStencil, SceneTexturesConfig.DepthCreateFlags, ERenderTargetLoadAction::ELoad,
 		ERenderTargetLoadAction::ELoad, FExclusiveDepthStencil::DepthWrite_StencilNop, RenderTargetsInfo);
