Sets up the dynamic shadow parameters used by the mobile renderer in the VolumetricFog pass, which runs later in the frame.

#tests on console

diff --git a/Engine/Source/Runtime/Renderer/Private/LightGridInjection.cpp b/Engine/Source/Runtime/Renderer/Private/LightGridInjection.cpp
--- a/Engine/Source/Runtime/Renderer/Private/LightGridInjection.cpp
+++ b/Engine/Source/Runtime/Renderer/Private/LightGridInjection.cpp
@@ -1142,7 +1142,7 @@ FComputeLightGridOutput FSceneRenderer::ComputeLightGrid(FRDGBuilder& GraphBuild
 							View.ForwardLightingResources.SelectedForwardDirectionalLightProxy = LightProxy;
 
 							// On mobile there is a separate FMobileDirectionalLightShaderParameters UB which holds all directional light data.
-							if(!IsMobilePlatform(View.GetShaderPlatform()))
+							if(!IsMobilePlatform(ShaderPlatform) || MobileSupportsSM5MaterialNodes(ShaderPlatform))
 							{
 								ForwardLightUniformParameters->HasDirectionalLight = 1;
 								ForwardLightUniformParameters->DirectionalLightColor = FVector3f(LightParameters.Color);
