Ensure PrimitiveHasVelocityForView returns false when velocity will not be rendered on mobile.
Velocity is only rendered on mobile when using TAA or rendering StereoMotionVectors for XR frame synthesis.

diff --git a/Engine/Source/Runtime/Renderer/Private/VelocityRendering.cpp b/Engine/Source/Runtime/Renderer/Private/VelocityRendering.cpp
--- a/Engine/Source/Runtime/Renderer/Private/VelocityRendering.cpp
+++ b/Engine/Source/Runtime/Renderer/Private/VelocityRendering.cpp
@@ -811,6 +811,11 @@ bool FVelocityRendering::IsParallelVelocity(EShaderPlatform ShaderPlatform)
 
 bool FVelocityMeshProcessor::PrimitiveHasVelocityForView(const FViewInfo& View, const FPrimitiveSceneProxy* PrimitiveSceneProxy)
 {
+	// Skip if velocity rendering is unsupported.
+	if (!PlatformSupportsVelocityRendering(View.GetShaderPlatform()))
+	{
+		return false;
+	}
 	// Skip camera cuts which effectively reset velocity for the new frame.
 	if (View.bCameraCut && !View.PreviousViewTransform.IsSet())
 	{
@@ -821,9 +826,9 @@ bool FVelocityMeshProcessor::PrimitiveHasVelocityForView(const FViewInfo& View,
 	{
 		return false;
 	}
-
-	// Only enabled on mobile when TAA is enabled or we're rendering OpenXR motion vectors
-	if (!PlatformSupportsVelocityRendering(View.GetShaderPlatform()))
+	// Only enabled on mobile when TAA is enabled or we're rendering OpenXR motion vectors.
+	const bool bUsesTAA = (View.AntiAliasingMethod == AAM_TemporalAA);
+	if (IsMobilePlatform(View.GetShaderPlatform()) && !(bUsesTAA || View.Family->EngineShowFlags.StereoMotionVectors))
 	{
 		return false;
 	}
