Heterogeneous Volumes: Decoupling enable and type from "Compositing with Translucency" options.

r.Translucency.HeterogeneousVolumes
Enables composting with heterogeneous volumes when rendering translucency (default = 0)

r.Translucency.HeterogeneousVolumes.Type
Determines the type of shadow map used for compositing with translucency
0: Adaptive Volumetric Shadow Map
1: Beer Shadow Map (Default)

diff --git a/Engine/Config/BaseScalability.ini b/Engine/Config/BaseScalability.ini
--- a/Engine/Config/BaseScalability.ini
+++ b/Engine/Config/BaseScalability.ini
@@ -805,6 +805,7 @@ r.HeterogeneousVolumes.Shadows.Type=0
 r.HeterogeneousVolumes.Shadows.Resolution=1024
 r.HeterogeneousVolumes.Shadows.MaxSampleCount=16
 r.HeterogeneousVolumes.UseExistenceMask=1
+r.Translucency.HeterogeneousVolumes.Type=0
 
 ;-----------------------------------------------------------------------------------------------------------------
 
diff --git a/Engine/Source/Runtime/Renderer/Private/HeterogeneousVolumes/HeterogeneousVolumes.cpp b/Engine/Source/Runtime/Renderer/Private/HeterogeneousVolumes/HeterogeneousVolumes.cpp
--- a/Engine/Source/Runtime/Renderer/Private/HeterogeneousVolumes/HeterogeneousVolumes.cpp
+++ b/Engine/Source/Runtime/Renderer/Private/HeterogeneousVolumes/HeterogeneousVolumes.cpp
@@ -29,13 +29,19 @@ static TAutoConsoleVariable<int32> CVarHeterogeneousVolumesShadows(
 static TAutoConsoleVariable<int32> CVarTranslucencyHeterogeneousVolumes(
 	TEXT("r.Translucency.HeterogeneousVolumes"),
 	0,
-	TEXT("Enables composting with heterogeneous volumes when rendering translucency (Default = 0)\n")
-	TEXT("0: Off\n")
-	TEXT("1: Camera-based Adaptive Volumetric Shadow Map\n")
-	TEXT("2: Camera-based Beer-Law Shadow Map\n"),
+	TEXT("Enables composting with heterogeneous volumes when rendering translucency (default = 0)"),
 	ECVF_RenderThreadSafe
 );
 
+static TAutoConsoleVariable<int32> CVarTranslucencyHeterogeneousVolumesType(
+	TEXT("r.Translucency.HeterogeneousVolumes.Type"),
+	1,
+	TEXT("Determines the type of shadow map used for compositing with translucency\n")
+	TEXT("0: Adaptive Volumetric Shadow Map\n")
+	TEXT("1: Beer Shadow Map (Default)"),
+	ECVF_Scalability | ECVF_RenderThreadSafe
+);
+
 static TAutoConsoleVariable<float> CVarHeterogeneousVolumesDownsampleFactor(
 	TEXT("r.HeterogeneousVolumes.DownsampleFactor"),
 	1.0,
@@ -363,10 +369,10 @@ namespace HeterogeneousVolumes
 		return static_cast<EShadowType>(ShadowValue);
 	}
 
-	ETranslucencyCompositingMode GetTranslucencyCompositingMode()
+	EShadowType GetTranslucencyCompositingMode()
 	{
-		int32 ClampedValue = FMath::Clamp(CVarTranslucencyHeterogeneousVolumes.GetValueOnRenderThread(), 0, 2);
-		return static_cast<ETranslucencyCompositingMode>(ClampedValue);
+		int32 ShadowValue = FMath::Clamp(CVarTranslucencyHeterogeneousVolumesType.GetValueOnAnyThread(), 0, 1);
+		return static_cast<EShadowType>(ShadowValue);
 	}
 }
 
@@ -996,7 +1002,8 @@ void FDeferredShadingSceneRenderer::RenderHeterogeneousVolumeShadows(
 	{
 		FViewInfo& View = Views[ViewIndex];
 
-		if (HeterogeneousVolumes::GetTranslucencyCompositingMode() == HeterogeneousVolumes::ETranslucencyCompositingMode::VolumetricShadowMap)
+		if (ShouldCompositeHeterogeneousVolumesWithTranslucency() &&
+			HeterogeneousVolumes::GetTranslucencyCompositingMode() == HeterogeneousVolumes::EShadowType::AdaptiveVolumetricShadowMap)
 		{
 			if (HeterogeneousVolumes::GetShadowPipeline() == HeterogeneousVolumes::EShadowPipeline::LiveShading)
 			{
@@ -1098,7 +1105,8 @@ void FDeferredShadingSceneRenderer::RenderHeterogeneousVolumes(
 			AddClearUAVPass(GraphBuilder, GraphBuilder.CreateUAV(HeterogeneousVolumeHoldout), FLinearColor::Black);
 		}
 
-		if (HeterogeneousVolumes::GetTranslucencyCompositingMode() == HeterogeneousVolumes::ETranslucencyCompositingMode::BeerShadowMap)
+		if (ShouldCompositeHeterogeneousVolumesWithTranslucency() &&
+			HeterogeneousVolumes::GetTranslucencyCompositingMode() == HeterogeneousVolumes::EShadowType::BeerShadowMap)
 		{
 			Desc.Format = PF_FloatRGBA;
 			HeterogeneousVolumeBeerShadowMap = GraphBuilder.CreateTexture(Desc, TEXT("HeterogeneousVolume.BeerShadowMap"));
@@ -1283,7 +1291,8 @@ void FDeferredShadingSceneRenderer::RenderHeterogeneousVolumes(
 						if (View.ViewState)
 						{
 							// Stash the temporary beer-law shadow map as a camera shadow if we aren't already creating one
-							if (HeterogeneousVolumes::GetTranslucencyCompositingMode() == HeterogeneousVolumes::ETranslucencyCompositingMode::BeerShadowMap)
+							if (ShouldCompositeHeterogeneousVolumesWithTranslucency() &&
+								HeterogeneousVolumes::GetTranslucencyCompositingMode() == HeterogeneousVolumes::EShadowType::BeerShadowMap)
 							{
 								// Resolution
 								FIntPoint ShadowMapResolution = HeterogeneousVolumes::GetDownsampledResolution(View.ViewRect.Size(), HeterogeneousVolumes::GetDownsampleFactor());
diff --git a/Engine/Source/Runtime/Renderer/Private/HeterogeneousVolumes/HeterogeneousVolumes.h b/Engine/Source/Runtime/Renderer/Private/HeterogeneousVolumes/HeterogeneousVolumes.h
--- a/Engine/Source/Runtime/Renderer/Private/HeterogeneousVolumes/HeterogeneousVolumes.h
+++ b/Engine/Source/Runtime/Renderer/Private/HeterogeneousVolumes/HeterogeneousVolumes.h
@@ -121,13 +121,7 @@ namespace HeterogeneousVolumes
 	EIndirectLightingMode GetIndirectLightingMode();
 
 	// Translucency compositing
-	enum class ETranslucencyCompositingMode
-	{
-		Disabled,
-		VolumetricShadowMap,
-		BeerShadowMap
-	};
-	ETranslucencyCompositingMode GetTranslucencyCompositingMode();
+	EShadowType GetTranslucencyCompositingMode();
 
 	int32 GetMipLevel();
 	int32 GetDebugMode();
diff --git a/Engine/Source/Runtime/Renderer/Private/HeterogeneousVolumes/HeterogeneousVolumesLiveShadingPipeline.cpp b/Engine/Source/Runtime/Renderer/Private/HeterogeneousVolumes/HeterogeneousVolumesLiveShadingPipeline.cpp
--- a/Engine/Source/Runtime/Renderer/Private/HeterogeneousVolumes/HeterogeneousVolumesLiveShadingPipeline.cpp
+++ b/Engine/Source/Runtime/Renderer/Private/HeterogeneousVolumes/HeterogeneousVolumesLiveShadingPipeline.cpp
@@ -159,7 +159,8 @@ namespace HeterogeneousVolumes
 
 	bool SupportsOverlappingVolumes()
 	{
-		return (GetTranslucencyCompositingMode() == ETranslucencyCompositingMode::VolumetricShadowMap) &&
+		return ShouldCompositeHeterogeneousVolumesWithTranslucency() &&
+			(GetTranslucencyCompositingMode() == EShadowType::AdaptiveVolumetricShadowMap) &&
 			(CVarHeterogeneousVolumesSupportOverlappingVolumes.GetValueOnRenderThread() != 0);
 	}
 
@@ -1320,7 +1321,7 @@ void RenderSingleScatteringWithLiveShading(
 		PassParameters->FogStruct = FogBuffer;
 		PassParameters->bApplyHeightFog = HeterogeneousVolumes::ShouldApplyHeightFog();
 		PassParameters->bApplyVolumetricFog = HeterogeneousVolumes::ShouldApplyVolumetricFog();
-		bool bCreateBeerShadowMap = HeterogeneousVolumes::GetTranslucencyCompositingMode() == HeterogeneousVolumes::ETranslucencyCompositingMode::BeerShadowMap;
+		bool bCreateBeerShadowMap = ShouldCompositeHeterogeneousVolumesWithTranslucency() && HeterogeneousVolumes::GetTranslucencyCompositingMode() == HeterogeneousVolumes::EShadowType::BeerShadowMap;
 		PassParameters->bCreateBeerShadowMap = bCreateBeerShadowMap;
 
 		// Indirect lighting data
@@ -3239,7 +3240,7 @@ void RenderAdaptiveVolumetricCameraMapWithLiveShading(
 		LODInfo.DownsampleFactor = HeterogeneousVolumes::GetCameraDownsampleFactor();
 		LODInfo.bIsPerspective = true;
 
-		HeterogeneousVolumes::EShadowType ShadowType = HeterogeneousVolumes::GetTranslucencyCompositingMode() == HeterogeneousVolumes::ETranslucencyCompositingMode::VolumetricShadowMap ? HeterogeneousVolumes::EShadowType::AdaptiveVolumetricShadowMap : HeterogeneousVolumes::EShadowType::BeerShadowMap;
+		HeterogeneousVolumes::EShadowType ShadowType = HeterogeneousVolumes::GetTranslucencyCompositingMode();
 		// Build a camera shadow for one volume
 		int32 LightType = 0;
 		FLightSceneInfo* LightSceneInfo = nullptr;
