[HWRT] Prefer non-inline ray tracing for debug visualizations when available.
- some inline visualizations are not fully featured (eg: normals).

diff --git a/Engine/Source/Runtime/Renderer/Private/RayTracing/RayTracingDebug.cpp b/Engine/Source/Runtime/Renderer/Private/RayTracing/RayTracingDebug.cpp
--- a/Engine/Source/Runtime/Renderer/Private/RayTracing/RayTracingDebug.cpp
+++ b/Engine/Source/Runtime/Renderer/Private/RayTracing/RayTracingDebug.cpp
@@ -1809,29 +1809,7 @@ void RenderRayTracingDebug(FRDGBuilder& GraphBuilder, const FScene& Scene, const
 	RDG_EVENT_SCOPE_STAT(GraphBuilder, RayTracingDebug, "RayTracingDebug");
 	RDG_GPU_STAT_SCOPE(GraphBuilder, RayTracingDebug);
 
-	if (bInlineRayTracing && RayTracingDebugModeSupportsInline(DebugVisualizationMode))
-	{
-		FRayTracingDebugCS::FParameters* InlinePassParameters = GraphBuilder.AllocParameters<FRayTracingDebugCS::FParameters>();
-
-		InlinePassParameters->SharedParameters = SharedParameters;
-
-		TShaderRef<FRayTracingDebugCS> ComputeShader = View.ShaderMap->GetShader<FRayTracingDebugCS>();
-
-		FIntPoint Resolution(View.ViewRect.Width(), View.ViewRect.Height());
-
-		const FIntPoint GroupSize(FRayTracingDebugCS::ThreadGroupSizeX, FRayTracingDebugCS::ThreadGroupSizeY);
-		const FIntVector GroupCount = FComputeShaderUtils::GetGroupCount(ViewRect.Size(), GroupSize);
-
-		GraphBuilder.AddPass(
-			RDG_EVENT_NAME("RayTracingDebug (INLINE)"),
-			InlinePassParameters,
-			ERDGPassFlags::Compute,
-			[InlinePassParameters, ComputeShader, GroupCount](FRDGAsyncTask, FRHIComputeCommandList& RHICmdList)
-			{
-				FComputeShaderUtils::Dispatch(RHICmdList, ComputeShader, *InlinePassParameters, GroupCount);
-			});
-	}
-	else
+	if (bRayTracingPipeline)
 	{
 		FRayTracingDebugRGS::FParameters* RayGenParameters = GraphBuilder.AllocParameters<FRayTracingDebugRGS::FParameters>();
 
@@ -1907,6 +1885,32 @@ void RenderRayTracingDebug(FRDGBuilder& GraphBuilder, const FScene& Scene, const
 
 				RHICmdList.RayTraceDispatch(Pipeline, RayGenShader.GetRayTracingShader(), SBT, GlobalResources, ViewRect.Size().X, ViewRect.Size().Y);
 			});
+	}
+	else if (bInlineRayTracing && RayTracingDebugModeSupportsInline(DebugVisualizationMode))
+	{
+		FRayTracingDebugCS::FParameters* InlinePassParameters = GraphBuilder.AllocParameters<FRayTracingDebugCS::FParameters>();
+
+		InlinePassParameters->SharedParameters = SharedParameters;
+
+		TShaderRef<FRayTracingDebugCS> ComputeShader = View.ShaderMap->GetShader<FRayTracingDebugCS>();
+
+		FIntPoint Resolution(View.ViewRect.Width(), View.ViewRect.Height());
+
+		const FIntPoint GroupSize(FRayTracingDebugCS::ThreadGroupSizeX, FRayTracingDebugCS::ThreadGroupSizeY);
+		const FIntVector GroupCount = FComputeShaderUtils::GetGroupCount(ViewRect.Size(), GroupSize);
+
+		GraphBuilder.AddPass(
+			RDG_EVENT_NAME("RayTracingDebug (INLINE)"),
+			InlinePassParameters,
+			ERDGPassFlags::Compute,
+			[InlinePassParameters, ComputeShader, GroupCount](FRDGAsyncTask, FRHIComputeCommandList& RHICmdList)
+			{
+				FComputeShaderUtils::Dispatch(RHICmdList, ComputeShader, *InlinePassParameters, GroupCount);
+			});
+	}
+	else
+	{
+		checkNoEntry(); // should have earlied out above
 	}	
 
 	if (DebugVisualizationMode == RAY_TRACING_DEBUG_VIZ_INSTANCE_OVERLAP)
