[HWRT] Fix race condition accessing ReferencedGeometryGroups.
- GatherRelevantStaticPrimitives(...) can run in parallel with FDynamicRayTracingInstancesContext::Finish(...) which modifies ReferencedGeometryGroups.
- instead of validating geometry groups referenced by static instances in GatherRelevantStaticPrimitives(...), gather those group handles and validate later in the frame when it's safe.

diff --git a/Engine/Source/Runtime/Renderer/Private/RayTracing/RayTracing.cpp b/Engine/Source/Runtime/Renderer/Private/RayTracing/RayTracing.cpp
--- a/Engine/Source/Runtime/Renderer/Private/RayTracing/RayTracing.cpp
+++ b/Engine/Source/Runtime/Renderer/Private/RayTracing/RayTracing.cpp
@@ -354,6 +354,10 @@ namespace RayTracing
 		int32 NumCachedStaticVisibleShaderBindings = 0; // TODO: Could remove this but it's used to Reserve
 
 		bool bUsesLightingChannels = false;
+
+#if DO_CHECK
+		TArray<TSet<RayTracing::FGeometryGroupHandle>> ReferencedGeometryGroupsToCheck;
+#endif
 	};
 
 	struct FGatherInstancesTaskData
@@ -1311,6 +1315,10 @@ namespace RayTracing
 			int32 NumCachedStaticVisibleShaderBindings = 0;
 
 			bool bUsesLightingChannels = false;
+
+#if DO_CHECK
+			TSet<RayTracing::FGeometryGroupHandle> ReferencedGeometryGroupsToCheck;
+#endif
 		};
 
 		TArray<FRelevantStaticPrimitivesContext> Contexts;
@@ -1433,13 +1441,12 @@ namespace RayTracing
 						FRayTracingGeometry* TargetRayTracingGeometry = SceneInfo->GetStaticRayTracingGeometry(LODIndex);
 
 						// TODO: Should have an assert here but disabled it due to UE-112448
-						if (TargetRayTracingGeometry == nullptr)
+						if (TargetRayTracingGeometry != nullptr)
 						{
-							return;
+							// It is not safe to directly call FRayTracingGeometryManager::IsGeometryGroupReferenced(...) here since other threads might be modifying it.
+							// Instead we gather the group handles to validate later.
+							Context.ReferencedGeometryGroupsToCheck.Add(TargetRayTracingGeometry->GroupHandle);
 						}
-
-						ensure(((FRayTracingGeometryManager*)GRayTracingGeometryManager)->IsGeometryGroupReferenced(TargetRayTracingGeometry->GroupHandle));
-
 					}
 #endif
 
@@ -1516,6 +1523,10 @@ namespace RayTracing
 				{
 					Nanite::GRayTracingManager.AddVisiblePrimitive(SceneInfo);
 				}
+
+#if DO_CHECK
+				TaskData.ReferencedGeometryGroupsToCheck.Add(Context.ReferencedGeometryGroupsToCheck);
+#endif
 			}
 		}
 	}
@@ -2033,6 +2044,15 @@ namespace RayTracing
 						ViewTaskData.CachedStaticPrimitives,
 						ViewTaskData.View.VisibleRayTracingShaderBindings);
 
+#if DO_CHECK
+					for(TSet<RayTracing::FGeometryGroupHandle> GeometryGroupHandles : ViewTaskData.ReferencedGeometryGroupsToCheck)
+					{
+						for (RayTracing::FGeometryGroupHandle GeometryGroupHandle : GeometryGroupHandles)
+						{
+							ensure(((FRayTracingGeometryManager*)GRayTracingGeometryManager)->IsGeometryGroupReferenced(GeometryGroupHandle));
+						}
+					}
+#endif
 				}, ViewTaskData.GatherRelevantStaticPrimitivesTask));
 
 			bAnyViewLightingChannelsUsingAHS |= ViewTaskData.SceneOptions.bLightingChannelsUsingAHS;
