[DFShadows] Prevent crash when NumPlanes is outside expected range.

diff --git a/Engine/Source/Runtime/Renderer/Private/DistanceFieldShadowing.cpp b/Engine/Source/Runtime/Renderer/Private/DistanceFieldShadowing.cpp
--- a/Engine/Source/Runtime/Renderer/Private/DistanceFieldShadowing.cpp
+++ b/Engine/Source/Runtime/Renderer/Private/DistanceFieldShadowing.cpp
@@ -162,6 +162,8 @@ FIntRect GetScissorRectForDFShadows(FIntRect ScissorRect)
 	return ScissorRect / GetDFShadowDownsampleFactor();
 }
 
+#define MAX_NUM_SHADOW_CONVEX_HULL_PLANES (12)
+
 class FCullObjectsForShadowCS : public FGlobalShader
 {
 	DECLARE_GLOBAL_SHADER(FCullObjectsForShadowCS);
@@ -178,7 +180,7 @@ class FCullObjectsForShadowCS : public FGlobalShader
 		SHADER_PARAMETER(uint32, bDrawNaniteMeshes)
 		SHADER_PARAMETER(uint32, bCullHeighfieldsNotInAtlas)
 		SHADER_PARAMETER(FVector4f, ShadowBoundingSphere)
-		SHADER_PARAMETER_ARRAY(FVector4f,ShadowConvexHull,[12])
+		SHADER_PARAMETER_ARRAY(FVector4f,ShadowConvexHull,[MAX_NUM_SHADOW_CONVEX_HULL_PLANES])
 	END_SHADER_PARAMETER_STRUCT()
 
 	class FPrimitiveType : SHADER_PERMUTATION_INT("DISTANCEFIELD_PRIMITIVE_TYPE", 2);
@@ -635,6 +637,9 @@ void CullDistanceFieldObjectsForLight(
 
 	AddClearUAVPass(GraphBuilder, CulledObjectBufferParameters.RWObjectIndirectArguments, 0);
 
+	ensureMsgf(NumPlanes <= MAX_NUM_SHADOW_CONVEX_HULL_PLANES, TEXT("NumPlanes (%d) is expected to be less or equal to %d."), NumPlanes, MAX_NUM_SHADOW_CONVEX_HULL_PLANES);
+	NumPlanes = FMath::Clamp(NumPlanes, 0, MAX_NUM_SHADOW_CONVEX_HULL_PLANES);
+
 	{
 		FCullObjectsForShadowCS::FParameters* PassParameters = GraphBuilder.AllocParameters<FCullObjectsForShadowCS::FParameters>();
 			
@@ -649,13 +654,11 @@ void CullDistanceFieldObjectsForLight(
 		PassParameters->bDrawNaniteMeshes = !(LightSceneProxy->UseVirtualShadowMaps() && LightSceneProxy->GetLightType() == LightType_Directional) || !bCullingForDirectShadowing;
 		PassParameters->bCullHeighfieldsNotInAtlas = bCullHeighfieldsNotInAtlas;
 
-		check(NumPlanes <= 12);
-
-		for (int32 i = 0; i < NumPlanes; i++)
+		for (int32 Index = 0; Index < NumPlanes; ++Index)
 		{
 			// translated planes from translated-shadow-space to translated-world-space
-			const FPlane4f Plane(PlaneData[i].TranslateBy(View.ViewMatrices.GetPreViewTranslation() - PrePlaneTranslation));
-			PassParameters->ShadowConvexHull[i] = FVector4f(FVector3f(Plane), Plane.W);
+			const FPlane4f Plane(PlaneData[Index].TranslateBy(View.ViewMatrices.GetPreViewTranslation() - PrePlaneTranslation));
+			PassParameters->ShadowConvexHull[Index] = FVector4f(FVector3f(Plane), Plane.W);
 		}
 
 		FCullObjectsForShadowCS::FPermutationDomain PermutationVector;
