[SingleLayerWater] Correct ENABLE_SHADOWMASKTEXTURE for mobile SM5 SLW. This fixes an issue where the screen space shadow mask for the opaque scene was used for SLW, leading to incorrectly projected shadows. It also enables CSM on water, same as other translucent materials.

diff --git a/Engine/Source/Runtime/Renderer/Private/MobileBasePassRendering.h b/Engine/Source/Runtime/Renderer/Private/MobileBasePassRendering.h
--- a/Engine/Source/Runtime/Renderer/Private/MobileBasePassRendering.h
+++ b/Engine/Source/Runtime/Renderer/Private/MobileBasePassRendering.h
@@ -398,8 +398,9 @@ public:
 		const bool bDeferredShadingEnabled = IsMobileDeferredShadingEnabled(Parameters.Platform);
 		const bool bTranslucentMaterial = IsTranslucentBlendMode(Parameters.MaterialParameters);
 		const bool bIsSLW = Parameters.MaterialParameters.ShadingModels.HasShadingModel(MSM_SingleLayerWater);
+		const bool bTranslucentOrSLWMaterial = bTranslucentMaterial || bIsSLW;
 		const bool bIsLit = Parameters.MaterialParameters.ShadingModels.IsLit();
-		const bool bMaterialUsesForwardShading = bIsLit && (bTranslucentMaterial || bIsSLW);
+		const bool bMaterialUsesForwardShading = bIsLit && bTranslucentOrSLWMaterial;
 		// Translucent materials always support clustered shading on mobile deferred
 		const bool bForwardShading = !bDeferredShadingEnabled || bMaterialUsesForwardShading;
 				
@@ -416,18 +417,18 @@ public:
 		uint32 MergedLocalLights = 0u;
 		if (LocalLightSetting == EMobileLocalLightSetting::LOCAL_LIGHTS_BUFFER)
 		{
-			if (MobileMergeLocalLightsInPrepassEnabled(Parameters.Platform) && !bTranslucentMaterial)
+			if (MobileMergeLocalLightsInPrepassEnabled(Parameters.Platform) && !bTranslucentOrSLWMaterial)
 			{
 				MergedLocalLights = 1u;
 			}
-			else if (MobileMergeLocalLightsInBasepassEnabled(Parameters.Platform) || bTranslucentMaterial)
+			else if (MobileMergeLocalLightsInBasepassEnabled(Parameters.Platform) || bTranslucentOrSLWMaterial)
 			{
 				MergedLocalLights = 2u;
 			}
 		}
 		OutEnvironment.SetDefine(TEXT("MERGED_LOCAL_LIGHTS_MOBILE"), MergedLocalLights);
 		OutEnvironment.SetDefine(TEXT("ENABLE_CLUSTERED_REFLECTION"), bEnableClusteredReflections ? 1u : 0u);
-		OutEnvironment.SetDefine(TEXT("ENABLE_SHADOWMASKTEXTURE"), bMobileUsesShadowMaskTexture && !bTranslucentMaterial ? 1u : 0u);
+		OutEnvironment.SetDefine(TEXT("ENABLE_SHADOWMASKTEXTURE"), bMobileUsesShadowMaskTexture && !bTranslucentOrSLWMaterial ? 1u : 0u);
 		OutEnvironment.SetDefine(TEXT("ENABLE_DBUFFER_TEXTURES"), Parameters.MaterialParameters.MaterialDomain == MD_Surface ? 1u : 0u);
 		EMobileTranslucentColorTransmittanceMode TranslucentColorTransmittanceMode = EMobileTranslucentColorTransmittanceMode::DEFAULT;
 		if (MaterialRequiresColorTransmittanceBlending(Parameters.MaterialParameters))
