Fix crash in path tracing tests when Substrate is enabled with Vulkan.

diff --git a/Engine/Source/Runtime/Renderer/Private/HairStrands/HairStrandsVisibility.cpp b/Engine/Source/Runtime/Renderer/Private/HairStrands/HairStrandsVisibility.cpp
--- a/Engine/Source/Runtime/Renderer/Private/HairStrands/HairStrandsVisibility.cpp
+++ b/Engine/Source/Runtime/Renderer/Private/HairStrands/HairStrandsVisibility.cpp
@@ -4679,7 +4679,7 @@ void RenderHairStrandsVisibilityBuffer(
 	FRDGTextureRef SceneMaterial0 = nullptr;
 	FRDGTextureRef SceneMaterial1 = nullptr;
 	FRDGTextureRef SceneMaterial2 = nullptr;
-	if (Substrate::IsSubstrateEnabled())
+	if (Substrate::IsSubstrateEnabled() && !Substrate::IsSubstrateBlendableGBufferEnabled(View.GetShaderPlatform()))
 	{
 		SceneMaterial0 = View.SubstrateViewData.SceneData->MaterialTextureArray;
 		SceneMaterial1 = View.SubstrateViewData.SceneData->TopLayerTexture;
