Add missing CopySceneColor between SingleLayerWater and Translucency in the mobile renderer

diff --git a/Engine/Source/Runtime/Renderer/Private/MobileShadingRenderer.cpp b/Engine/Source/Runtime/Renderer/Private/MobileShadingRenderer.cpp
--- a/Engine/Source/Runtime/Renderer/Private/MobileShadingRenderer.cpp
+++ b/Engine/Source/Runtime/Renderer/Private/MobileShadingRenderer.cpp
@@ -2666,24 +2666,26 @@ void FMobileSceneRenderer::RenderDeferredMultiPass(FRDGBuilder& GraphBuilder, FS
 			const bool bShouldRenderTranslucency = ShouldRenderTranslucency(StandardTranslucencyPass, Views) && ViewFamily.EngineShowFlags.Translucency && (View.ParallelMeshDrawCommandPasses[StandardTranslucencyMeshPass] != nullptr);
 			const bool bShouldRenderWater = bShouldRenderFullSLW && View.bHasSingleLayerWaterMaterial;
 
-			if (bShouldRenderTranslucency || bShouldRenderWater)
+			if (bShouldRenderWater)
 			{
+				auto* SLWPassParameters = GraphBuilder.AllocParameters<FMobileRenderPassParameters>();
+				*SLWPassParameters = *PassParameters;
+				SLWPassParameters->RenderTargets = BasePassRenderTargets;
+				SLWPassParameters->RenderTargets.DepthStencil = FDepthStencilBinding(BasePassRenderTargets.DepthStencil.GetTexture(), ERenderTargetLoadAction::ELoad, FExclusiveDepthStencil::DepthWrite_StencilNop);
+
+				FMobileBasePassUniformParameters* SLWBasePassParameters = GraphBuilder.AllocParameters<FMobileBasePassUniformParameters>();
+				SetupMobileBasePassUniformParameters(GraphBuilder, View, EMobileBasePass::Translucent, SetupMode, *SLWBasePassParameters, BasePassTextures);
+				SLWPassParameters->MobileBasePass = GraphBuilder.CreateUniformBuffer(SLWBasePassParameters);
+
 				FRDGTextureRef SceneColorCopyTexture = AddCopySceneColorPass(GraphBuilder, Views, SceneTextures.Color, false, false);
 				if (SceneColorCopyTexture != nullptr)
 				{
-					BasePassParameters->SceneColorCopyTexture = SceneColorCopyTexture;
+					SLWBasePassParameters->SceneColorCopyTexture = SceneColorCopyTexture;
 				}
 				if (SceneDepthCopyTexture != nullptr)
 				{
-					BasePassParameters->SceneDepthCopyTexture = SceneDepthCopyTexture;
+					SLWBasePassParameters->SceneDepthCopyTexture = SceneDepthCopyTexture;
 				}
-			}
-			if (bShouldRenderWater)
-			{
-				auto* SLWPassParameters = GraphBuilder.AllocParameters<FMobileRenderPassParameters>();
-				*SLWPassParameters = *PassParameters;
-				SLWPassParameters->RenderTargets = BasePassRenderTargets;
-				SLWPassParameters->RenderTargets.DepthStencil = FDepthStencilBinding(BasePassRenderTargets.DepthStencil.GetTexture(), ERenderTargetLoadAction::ELoad, FExclusiveDepthStencil::DepthWrite_StencilNop);
 
 				if (Scene->GPUScene.IsEnabled())
 				{
@@ -2706,6 +2708,16 @@ void FMobileSceneRenderer::RenderDeferredMultiPass(FRDGBuilder& GraphBuilder, FS
 			}
 			if (bShouldRenderTranslucency)
 			{
+				FRDGTextureRef SceneColorCopyTexture = AddCopySceneColorPass(GraphBuilder, Views, SceneTextures.Color, false, false);
+				if (SceneColorCopyTexture != nullptr)
+				{
+					BasePassParameters->SceneColorCopyTexture = SceneColorCopyTexture;
+				}
+				if (SceneDepthCopyTexture != nullptr)
+				{
+					BasePassParameters->SceneDepthCopyTexture = SceneDepthCopyTexture;
+				}
+
 				GraphBuilder.AddPass(
 					RDG_EVENT_NAME("Translucency"),
 					PassParameters,
