[SingleLayerWater] Fixed shader compiler error on dx11 when using static lighting, the velocity render target, the SeparateMainDirectionalLight render target and virtual textures in the material. VTs add a UAV to the PS, but the shader is already using Rendertargets 0 through 7. However, GBufferD (CustomData) is never actually used for water, so this CL bind the SeparateMainDirectionalLight RT (only used on SLW) to the GBufferD slot, freeing up slot 7, which can then be used by the VT UAV.

diff --git a/Engine/Shaders/Private/BasePassPixelShader.usf b/Engine/Shaders/Private/BasePassPixelShader.usf
--- a/Engine/Shaders/Private/BasePassPixelShader.usf
+++ b/Engine/Shaders/Private/BasePassPixelShader.usf
@@ -2335,12 +2335,12 @@ void FPixelShaderInOut_MainPS(
 				#if SUBSTRATE_GBUFFER_FORMAT==1
 				Out.MRT[(GBUFFER_HAS_VELOCITY ? 2 : 1) + (GBUFFER_HAS_PRECSHADOWFACTOR ? 1 : 0)] = float4(SeparatedWaterMainDirLightLuminance * View.PreExposure, 1.0f);
 				#else
-				Out.MRT[(GBUFFER_HAS_VELOCITY ? 6 : 5) + (GBUFFER_HAS_PRECSHADOWFACTOR ? 1 : 0)] = float4(SeparatedWaterMainDirLightLuminance * View.PreExposure, 1.0f);
+				Out.MRT[(GBUFFER_HAS_VELOCITY ? 5 : 4)] = float4(SeparatedWaterMainDirLightLuminance * View.PreExposure, 1.0f);
 				#endif
 			#else
 			if (GBuffer.ShadingModelID == SHADINGMODELID_SINGLELAYERWATER)
 			{
-				Out.MRT[(GBUFFER_HAS_VELOCITY ? 6 : 5) + (GBUFFER_HAS_PRECSHADOWFACTOR ? 1 : 0)] = float4(SeparatedWaterMainDirLightLuminance * View.PreExposure, 1.0f);
+				Out.MRT[(GBUFFER_HAS_VELOCITY ? 5 : 4)] = float4(SeparatedWaterMainDirLightLuminance * View.PreExposure, 1.0f);
 			}
 			#endif
 		#endif
diff --git a/Engine/Source/Runtime/RenderCore/Private/GBufferInfo.cpp b/Engine/Source/Runtime/RenderCore/Private/GBufferInfo.cpp
--- a/Engine/Source/Runtime/RenderCore/Private/GBufferInfo.cpp
+++ b/Engine/Source/Runtime/RenderCore/Private/GBufferInfo.cpp
@@ -371,13 +371,11 @@ FGBufferInfo RENDERCORE_API FetchLegacyGBufferInfo(const FGBufferParams& Params)
 	{
 		TargetGBufferD = 4;
 		Info.Targets[4].Init(GBT_Unorm_8_8_8_8,  TEXT("GBufferD"), false,  true,  true,  true);
-		TargetSeparatedMainDirLight = 5;
 
 		if (Params.bHasPrecShadowFactor)
 		{
 			TargetGBufferE = 5;
 			Info.Targets[5].Init(GBT_Unorm_8_8_8_8, TEXT("GBufferE"), false, true, true, true);
-			TargetSeparatedMainDirLight = 6;
 		}
 	}
 	else if (Params.bHasVelocity)
@@ -388,13 +386,11 @@ FGBufferInfo RENDERCORE_API FetchLegacyGBufferInfo(const FGBufferParams& Params)
 		// note the false for use extra flags for velocity, not quite sure of all the ramifications, but this keeps it consistent with previous usage
 		Info.Targets[4].Init(Params.bUsesVelocityDepth ? GBT_Unorm_16_16_16_16 : (IsAndroidOpenGLESPlatform(Params.ShaderPlatform) ? GBT_Float_16_16 : GBT_Unorm_16_16), TEXT("Velocity"), false, true, true, false);
 		Info.Targets[5].Init(GBT_Unorm_8_8_8_8, TEXT("GBufferD"), false, true, true, true);
-		TargetSeparatedMainDirLight = 6;
 
 		if (Params.bHasPrecShadowFactor)
 		{
 			TargetGBufferE = 6;
 			Info.Targets[6].Init(GBT_Unorm_8_8_8_8, TEXT("GBufferE"), false, true, true, false);
-			TargetSeparatedMainDirLight = 7;
 		}
 	}
 	else if (Params.bHasTangent)
@@ -403,12 +399,10 @@ FGBufferInfo RENDERCORE_API FetchLegacyGBufferInfo(const FGBufferParams& Params)
 		TargetGBufferD = 5;
 		Info.Targets[4].Init(GBT_Unorm_8_8_8_8,  TEXT("GBufferF"), false,  true,  true,  true);
 		Info.Targets[5].Init(GBT_Unorm_8_8_8_8, TEXT("GBufferD"), false, true, true, true);
-		TargetSeparatedMainDirLight = 6;
 		if (Params.bHasPrecShadowFactor)
 		{
 			TargetGBufferE = 6;
 			Info.Targets[6].Init(GBT_Unorm_8_8_8_8, TEXT("GBufferE"), false, true, true, true);
-			TargetSeparatedMainDirLight = 7;
 		}
 	}
 	else
@@ -417,6 +411,9 @@ FGBufferInfo RENDERCORE_API FetchLegacyGBufferInfo(const FGBufferParams& Params)
 		check(0);
 	}
 
+	// SLW never uses GBufferD (CustomData) and we want to tightly pack RTVs so there's space for UAVs after the RTV range (D3D11.0 only supports a total of 8 RTVs and UAVs and they must have non-overlapping ranges).
+	TargetSeparatedMainDirLight = TargetGBufferD;
+
 	// this value isn't correct, because it doesn't respect the scene color format cvar, but it's ignored anyways
 	// so it's ok for now
 	Info.Slots[GBS_SceneColor] = FGBufferItem(GBS_SceneColor, GBC_Raw_Float_11_11_10, GBCH_Both);
diff --git a/Engine/Source/Runtime/Renderer/Private/SingleLayerWaterRendering.cpp b/Engine/Source/Runtime/Renderer/Private/SingleLayerWaterRendering.cpp
--- a/Engine/Source/Runtime/Renderer/Private/SingleLayerWaterRendering.cpp
+++ b/Engine/Source/Runtime/Renderer/Private/SingleLayerWaterRendering.cpp
@@ -1791,12 +1791,26 @@ void FDeferredShadingSceneRenderer::RenderSingleLayerWaterInner(
 	const FRDGSystemTextures& SystemTextures = FRDGSystemTextures::Get(GraphBuilder);
 
 	const EGBufferLayout GBufferLayout = GetSingleLayerWaterGBufferLayout();
+	const FGBufferBindings& GBufferBindings = SceneTextures.Config.GBufferBindings[GBufferLayout];
 	TStaticArray<FTextureRenderTargetBinding, MaxSimultaneousRenderTargets> BasePassTextures;
 	uint32 BasePassTextureCount = SceneTextures.GetGBufferRenderTargets(BasePassTextures, GBufferLayout);
 	if(IsWaterSeparateMainDirLightEnabled(Scene->GetShaderPlatform()))
 	{
+		const bool bUsesSubstrateAdaptiveGBufferFormat = Substrate::IsSubstrateEnabled() && !Substrate::IsSubstrateBlendableGBufferEnabled(ShaderPlatform);
+		int32 SeparateMainDirLightIndex = INDEX_NONE;
+		if (bUsesSubstrateAdaptiveGBufferFormat)
+		{
+			SeparateMainDirLightIndex = BasePassTextureCount++;
+		}
+		else
+		{
+			// SLW never uses GBufferD (CustomData) and we want to tightly pack RTVs so there's space for UAVs after the RTV range (D3D11.0 only supports a total of 8 RTVs and UAVs and they must have non-overlapping ranges).
+			check(GBufferBindings.GBufferD.Index >= 0 && GBufferBindings.GBufferD.Index < MaxSimultaneousRenderTargets);
+			SeparateMainDirLightIndex = GBufferBindings.GBufferD.Index;
+		}
+		
 		const bool bNeverClear = true;
-		BasePassTextures[BasePassTextureCount++] = FTextureRenderTargetBinding(SceneWithoutWaterTextures.SeparatedMainDirLightTexture, bNeverClear);
+		BasePassTextures[SeparateMainDirLightIndex] = FTextureRenderTargetBinding(SceneWithoutWaterTextures.SeparatedMainDirLightTexture, bNeverClear);
 	}
 	Substrate::AppendSubstrateMRTs(*this, BasePassTextureCount, BasePassTextures);
 	TArrayView<FTextureRenderTargetBinding> BasePassTexturesView = MakeArrayView(BasePassTextures.GetData(), BasePassTextureCount);
@@ -1842,7 +1856,6 @@ void FDeferredShadingSceneRenderer::RenderSingleLayerWaterInner(
 		PassParameters->RenderTargets.DepthStencil = DepthStencilBinding;
 
 		// Make sure to clear the velocity texture if it wasn't already written to. This can be the case if the velocity pass is set to "Write after base pass".
-		const FGBufferBindings& GBufferBindings = SceneTextures.Config.GBufferBindings[GBufferLayout];
 		if (GBufferBindings.GBufferVelocity.Index > 0 && GBufferBindings.GBufferVelocity.Index < (int32)BasePassTextureCount && !HasBeenProduced(SceneTextures.Velocity))
 		{
 			PassParameters->RenderTargets[GBufferBindings.GBufferVelocity.Index].SetLoadAction(ERenderTargetLoadAction::EClear);
@@ -2054,7 +2067,21 @@ void FSingleLayerWaterPassMeshProcessor::CollectPSOInitializers(const FSceneText
 		SetupGBufferRenderTargetInfo(SceneTexturesConfig, RenderTargetsInfo, true /*bSetupDepthStencil*/, GBufferLayout);
 		if (IsWaterSeparateMainDirLightEnabled(GMaxRHIShaderPlatform))
 		{
-			AddRenderTargetInfo(PF_FloatR11G11B10, TexCreate_ShaderResource | TexCreate_RenderTargetable, RenderTargetsInfo);
+			const bool bUsesSubstrateAdaptiveGBufferFormat = Substrate::IsSubstrateEnabled() && !Substrate::IsSubstrateBlendableGBufferEnabled(GMaxRHIShaderPlatform);
+			int32 SeparateMainDirLightIndex = INDEX_NONE;
+			if (bUsesSubstrateAdaptiveGBufferFormat)
+			{
+				SeparateMainDirLightIndex = RenderTargetsInfo.RenderTargetsEnabled++;
+			}
+			else
+			{
+				const FGBufferBindings& GBufferBindings = SceneTexturesConfig.GBufferBindings[GBufferLayout];
+				// SLW never uses GBufferD (CustomData) and we want to tightly pack RTVs so there's space for UAVs after the RTV range (D3D11.0 only supports a total of 8 RTVs and UAVs and they must have non-overlapping ranges).
+				check(GBufferBindings.GBufferD.Index >= 0 && GBufferBindings.GBufferD.Index < MaxSimultaneousRenderTargets);
+				SeparateMainDirLightIndex = GBufferBindings.GBufferD.Index;
+			}
+			RenderTargetsInfo.RenderTargetFormats[SeparateMainDirLightIndex] = PF_FloatR11G11B10;
+			RenderTargetsInfo.RenderTargetFlags[SeparateMainDirLightIndex] = TexCreate_ShaderResource | TexCreate_RenderTargetable;
 		}
 
 		const bool bHasDepthPrepass = IsSingleLayerWaterDepthPrepassEnabled(GetFeatureLevelShaderPlatform(FeatureLevel), FeatureLevel);
