Fixed RHI validation error with lens flare dummy transition buffer.

diff --git a/Engine/Source/Runtime/RenderCore/Public/GlobalRenderResources.h b/Engine/Source/Runtime/RenderCore/Public/GlobalRenderResources.h
--- a/Engine/Source/Runtime/RenderCore/Public/GlobalRenderResources.h
+++ b/Engine/Source/Runtime/RenderCore/Public/GlobalRenderResources.h
@@ -23,6 +23,7 @@ extern RENDERCORE_API FTexture* GTransparentBlackTexture;
 extern RENDERCORE_API FTextureWithSRV* GTransparentBlackTextureWithSRV;
 
 // A global dummy texture used for manual transition (Ex: when the VS is sampling produce in a previous pass, and needs a barrier)
+UE_DEPRECATED(5.7, "This is no longer being used.")
 extern RENDERCORE_API FTexture* GDummyTransitionTexture;
 
 // An empty vertex buffer with a UAV
diff --git a/Engine/Source/Runtime/Renderer/Private/PostProcess/PostProcessLensFlares.cpp b/Engine/Source/Runtime/Renderer/Private/PostProcess/PostProcessLensFlares.cpp
--- a/Engine/Source/Runtime/Renderer/Private/PostProcess/PostProcessLensFlares.cpp
+++ b/Engine/Source/Runtime/Renderer/Private/PostProcess/PostProcessLensFlares.cpp
@@ -129,6 +129,22 @@ ELensFlareQuality GetLensFlareQuality()
 		static_cast<int32>(ELensFlareQuality::MAX) - 1));
 }
 
+class FLensFlareTransitionBuffer : public FVertexBuffer
+{
+public:
+	virtual void InitRHI(FRHICommandListBase& RHICmdList) override
+	{
+		const FRHIBufferCreateDesc CreateDesc =
+			FRHIBufferCreateDesc::CreateStructured(TEXT("LensFlareTransitionBuffer"), sizeof(uint32) * 4, sizeof(uint32))
+			.AddUsage(BUF_Static | BUF_ShaderResource | BUF_UnorderedAccess)
+			.DetermineInitialState();
+
+		VertexBufferRHI = RHICmdList.CreateBuffer(CreateDesc);
+	}
+};
+
+FLensFlareTransitionBuffer* GLensFlareTransitionBuffer = new TGlobalResource<FLensFlareTransitionBuffer, FRenderResource::EInitPhase::Pre>;
+
 FScreenPassTexture AddLensFlaresPass(
 	FRDGBuilder& GraphBuilder,
 	const FViewInfo& View,
@@ -225,11 +241,10 @@ FScreenPassTexture AddLensFlaresPass(
 					FlareViewport.Rect.Max.Y,
 					1.0f);
 
-				// Force a transition to SRVGraphicsNonPixel on a dummy texture to account for VS accessing the texture (using NON_PIXEL_SRV in struct gets merged)
+				// Force a transition to SRVGraphicsNonPixel on an Empty Structured Buffer to account for VS accessing the texture (using NON_PIXEL_SRV in struct gets merged)
 				if (GRHINeedsSRVGraphicsNonPixelWorkaround)
 				{
-					RHICmdList.Transition(FRHITransitionInfo(GDummyTransitionTexture->GetTextureRHI(), ERHIAccess::Unknown, ERHIAccess::SRVGraphicsNonPixel));
-
+					RHICmdList.Transition(FRHITransitionInfo(GLensFlareTransitionBuffer->GetRHI(), ERHIAccess::UAVMask, ERHIAccess::VertexOrIndexBuffer | ERHIAccess::SRVGraphicsNonPixel));
 
 					FRHIRenderPassInfo RPInfo(PassParameters->RenderTargets[0].GetTexture()->GetRHI(), ERenderTargetActions::Clear_Store);
 					RHICmdList.BeginRenderPass(RPInfo, TEXT("LensFlareBlur"));
@@ -255,6 +270,7 @@ FScreenPassTexture AddLensFlaresPass(
 				if (GRHINeedsSRVGraphicsNonPixelWorkaround)
 				{
 					RHICmdList.EndRenderPass();
+					RHICmdList.Transition(FRHITransitionInfo(GLensFlareTransitionBuffer->GetRHI(), ERHIAccess::VertexOrIndexBuffer | ERHIAccess::SRVGraphicsNonPixel, ERHIAccess::UAVMask));
 				}
 			});
 	}
