Substrate - Fix crash on Android with DXT flavor.

diff --git a/Engine/Source/Runtime/Renderer/Private/ShadingEnergyConservation.cpp b/Engine/Source/Runtime/Renderer/Private/ShadingEnergyConservation.cpp
--- a/Engine/Source/Runtime/Renderer/Private/ShadingEnergyConservation.cpp
+++ b/Engine/Source/Runtime/Renderer/Private/ShadingEnergyConservation.cpp
@@ -128,9 +128,14 @@ TGlobalResource<FShadingEnergyConservationResources> GShadingEnergyConservationR
 
 namespace ShadingEnergyConservationData
 {
+	static bool IsTextureDataValid(UTexture2D* In)
+	{
+		return In && In->GetPlatformData() && In->GetCPUCopy();
+	}
+
 	static TRefCountPtr<IPooledRenderTarget> CreateTexture2D(FRHICommandListImmediate& RHICmdList, TObjectPtr<class UTexture2D>& InCPUTexture, EPixelFormat InFormat, const TCHAR* InName)
 	{
-		if (!InCPUTexture || !InCPUTexture->GetPlatformData())
+		if (!IsTextureDataValid(InCPUTexture))
 		{
 			return nullptr;
 		}
@@ -195,7 +200,7 @@ namespace ShadingEnergyConservationData
 
 	static TRefCountPtr<IPooledRenderTarget> CreateTexture3D(FRHICommandListImmediate& RHICmdList, TObjectPtr<class UTexture2D>& InCPUTexture, EPixelFormat InFormat, const TCHAR* InName)
 	{
-		if (!InCPUTexture || !InCPUTexture->GetPlatformData())
+		if (!IsTextureDataValid(InCPUTexture))
 		{
 			return nullptr;
 		}
