Substrate - fix for crash when megalight is enabled and GBufferFormat=1.

#tests pc editor dx12
[RN]

diff --git a/Engine/Source/Runtime/Renderer/Private/Substrate/Substrate.cpp b/Engine/Source/Runtime/Renderer/Private/Substrate/Substrate.cpp
--- a/Engine/Source/Runtime/Renderer/Private/Substrate/Substrate.cpp
+++ b/Engine/Source/Runtime/Renderer/Private/Substrate/Substrate.cpp
@@ -398,7 +398,7 @@ static bool NeedsSampledMaterials(const FScene* Scene, const FViewInfo& View)
 static bool NeedsClosureOffsets(const FScene* Scene, const FViewInfo& View)
 {
 	// No need for closure index when either BlendableGBuffer is enabled or if ClosureCount == 1
-	return  (ShouldRenderLumenDiffuseGI(Scene, View) || ShouldRenderLumenReflections(View) || NeedsSampledMaterials(Scene, *View.Family) || Substrate::ShouldRenderSubstrateDebugPasses(View))
+	return  (ShouldRenderLumenDiffuseGI(Scene, View) || ShouldRenderLumenReflections(View) || MegaLights::IsEnabled(*View.Family) || NeedsSampledMaterials(Scene, *View.Family) || Substrate::ShouldRenderSubstrateDebugPasses(View))
 		&& !Substrate::IsSubstrateBlendableGBufferEnabled(View.GetShaderPlatform())
 		&& View.SubstrateViewData.MaxClosurePerPixel > 1;
 }
