Path Tracer: Fix assert triggered when using reference clouds with mGPU

#rnx

diff --git a/Engine/Source/Runtime/Renderer/Private/PathTracing.cpp b/Engine/Source/Runtime/Renderer/Private/PathTracing.cpp
--- a/Engine/Source/Runtime/Renderer/Private/PathTracing.cpp
+++ b/Engine/Source/Runtime/Renderer/Private/PathTracing.cpp
@@ -4027,7 +4027,7 @@ void FDeferredShadingSceneRenderer::RenderPathTracing(
 
 			FRDGTexture* CloudAccelerationMap = PathTracingState->CloudAccelerationMap.IsValid() ? GraphBuilder.RegisterExternalTexture(PathTracingState->CloudAccelerationMap, TEXT("PathTracer.CloudAccelerationMap")) : nullptr;
 
-			const bool bEnableClouds   = (Config.PathTracingData.VolumeFlags & PATH_TRACER_VOLUME_ENABLE_CLOUDS) != 0;
+			const bool bEnableClouds = (Config.PathTracingData.VolumeFlags & PATH_TRACER_VOLUME_ENABLE_CLOUDS) != 0 && CloudAccelerationMap != nullptr;
 
 			if (bEnableClouds)
 			{
@@ -4081,7 +4081,9 @@ void FDeferredShadingSceneRenderer::RenderPathTracing(
 			}
 			if (CloudAccelerationMap == nullptr)
 			{
+				// Clouds are not enabled, make sure we don't accidentally try to use them during rendering
 				CloudAccelerationMap = GraphBuilder.RegisterExternalTexture(GSystemTextures.BlackDummy);
+				Config.PathTracingData.VolumeFlags &= ~PATH_TRACER_VOLUME_ENABLE_CLOUDS;
 			}
 
 			// We are writing to the texture, we'll need to extract it...
