Path Tracer: Avoid crash in post process material setup when Path Tracing is not actually supported by the current shader platform

This can happen if the editor cannot use path tracing for some reason (such as running on an RHI which doesn't support it, or if some ray tracing settings have been locally modified). Make sure that the path tracing permutation is not selected in these cases when looking for the post process material

#rnx

diff --git a/Engine/Source/Runtime/Renderer/Private/PostProcess/PostProcessMaterial.cpp b/Engine/Source/Runtime/Renderer/Private/PostProcess/PostProcessMaterial.cpp
--- a/Engine/Source/Runtime/Renderer/Private/PostProcess/PostProcessMaterial.cpp
+++ b/Engine/Source/Runtime/Renderer/Private/PostProcess/PostProcessMaterial.cpp
@@ -369,6 +369,7 @@ public:
 
 static bool GetMaterialShaders(
 	const FMaterial& Material, 
+	const EShaderPlatform ShaderPlatform,
 	bool bManualStencilTest, 
 	bool bNeuralPostProcessPrepass,
 	bool bPathTracingEnabled,
@@ -380,7 +381,7 @@ static bool GetMaterialShaders(
 	FPostProcessMaterialPS::FPermutationDomain PermutationVectorPS;
 	PermutationVectorPS.Set<FPostProcessMaterialPS::FManualStencilTestDim>(bManualStencilTest);
 	PermutationVectorPS.Set<FPostProcessMaterialPS::FNeuralPostProcessPrePass>(bNeuralPostProcessPrepass);
-	PermutationVectorPS.Set<FPostProcessMaterialPS::FPathTracingDim>(bPathTracingEnabled);
+	PermutationVectorPS.Set<FPostProcessMaterialPS::FPathTracingDim>(bPathTracingEnabled && PathTracing::ShouldCompilePathTracingShadersForProject(ShaderPlatform));
 
 	ShaderTypes.AddShaderType<FPostProcessMaterialVS>();
 	ShaderTypes.AddShaderType<FPostProcessMaterialPS>(PermutationVectorPS.ToDimensionValueId());
@@ -399,7 +400,8 @@ static bool GetMaterialShaders(
 
 static void GetMaterialInfo(
 	const UMaterialInterface* InMaterialInterface,
-	ERHIFeatureLevel::Type InFeatureLevel,
+	const ERHIFeatureLevel::Type InFeatureLevel,
+	const EShaderPlatform ShaderPlatform,
 	const bool bPathTracingEnabled,
 	const FPostProcessMaterialInputs& Inputs,
 	const FMaterial*& OutMaterial,
@@ -420,7 +422,7 @@ static void GetMaterialInfo(
 		if (Material && Material->GetMaterialDomain() == MD_PostProcess)
 		{
 			const bool bManualStencilTest = Inputs.bManualStencilTest && Material->IsStencilTestEnabled();
-			if (GetMaterialShaders(*Material, bManualStencilTest, bNeuralPostProcessPrepass, bPathTracingEnabled, OutVertexShader, OutPixelShader))
+			if (GetMaterialShaders(*Material, ShaderPlatform, bManualStencilTest, bNeuralPostProcessPrepass, bPathTracingEnabled, OutVertexShader, OutPixelShader))
 			{
 				break;
 			}
@@ -428,7 +430,9 @@ static void GetMaterialInfo(
 		MaterialProxy = MaterialProxy->GetFallback(InFeatureLevel);
 	}
 
+	// NOTE: we expect to reach this point only via the break statement. If we exit the while loop via a null pointer, something has gone wrong
 	check(Material);
+	check(MaterialProxy);
 
 	const FMaterialShaderMap* MaterialShaderMap = Material->GetRenderingThreadShaderMap();
 	check(MaterialShaderMap);
@@ -598,7 +602,7 @@ void AddNeuralPostProcessPass(
 	const FMaterialShaderMap* MaterialShaderMap = nullptr;
 	TShaderRef<FPostProcessMaterialVS> NeuralPostProcessPassVertexShader;
 	TShaderRef<FPostProcessMaterialPS> NeuralPostProcessPassPixelShader;
-	GetMaterialInfo(MaterialInterface, FeatureLevel, bPathTracingEnabled, Inputs, Material, MaterialRenderProxy, MaterialShaderMap, NeuralPostProcessPassVertexShader, NeuralPostProcessPassPixelShader, true);
+	GetMaterialInfo(MaterialInterface, FeatureLevel, View.GetShaderPlatform(), bPathTracingEnabled, Inputs, Material, MaterialRenderProxy, MaterialShaderMap, NeuralPostProcessPassVertexShader, NeuralPostProcessPassPixelShader, true);
 	
 	check(NeuralPostProcessPassVertexShader.IsValid());
 	check(NeuralPostProcessPassPixelShader.IsValid());
@@ -745,7 +749,7 @@ FScreenPassTexture AddPostProcessMaterialPass(
 	const FMaterialShaderMap* MaterialShaderMap = nullptr;
 	TShaderRef<FPostProcessMaterialVS> VertexShader;
 	TShaderRef<FPostProcessMaterialPS> PixelShader;
-	GetMaterialInfo(MaterialInterface, FeatureLevel, bPathTracingEnabled, Inputs, Material, MaterialRenderProxy, MaterialShaderMap, VertexShader, PixelShader);
+	GetMaterialInfo(MaterialInterface, FeatureLevel, View.GetShaderPlatform(), bPathTracingEnabled, Inputs, Material, MaterialRenderProxy, MaterialShaderMap, VertexShader, PixelShader);
 
 	EBlendableLocation BlendableLocation = MaterialRenderProxy->GetBlendableLocation(Material);
 	const FScreenPassTextureSlice SceneColorOutput = Inputs.GetSceneColorOutput(BlendableLocation);
@@ -1569,7 +1573,7 @@ void FPostProcessMaterialPSOCollector::CollectPSOInitializers(
 		bool bNeuralPostProcessPrepass = false;
 		bool bPathTracingEnabled = false;
 
-		if (!GetMaterialShaders(Material, bManualStencilTest, bNeuralPostProcessPrepass, bPathTracingEnabled, VertexShader, PixelShader))
+		if (!GetMaterialShaders(Material, SceneTexturesConfig.ShaderPlatform, bManualStencilTest, bNeuralPostProcessPrepass, bPathTracingEnabled, VertexShader, PixelShader))
 		{
 			return;
 		}
