[HWRT] Fix crashes due to missing shaders when using Metal SM6 preview mode.

diff --git a/Engine/Config/BaseEngine.ini b/Engine/Config/BaseEngine.ini
--- a/Engine/Config/BaseEngine.ini
+++ b/Engine/Config/BaseEngine.ini
@@ -3437,6 +3437,7 @@ bEnableRayTracing=true
 [/Script/MacTargetPlatform.MacTargetSettings]
 MetalLanguageVersion=7
 +TargetedRHIs=SF_METAL_SM5
++TargetedRHIs=SF_METAL_SM6
 UseFastIntrinsics=False
 EnableMathOptimisations=True
 TargetArchitecture=MacTargetArchitectureUniversal
@@ -3444,7 +3445,6 @@ EditorTargetArchitecture=MacTargetArchitectureUniversal
 EditorDefaultArchitecture=MacTargetArchitectureHost
 DefaultArchitecture=MacTargetArchitectureHost
 bBuildAllSupportedOnBuildMachine=true
-bEnableRayTracing=false
 AudioSampleRate=48000
 AudioCallbackBufferFrameSize=1024
 AudioNumBuffersToEnqueue=1
diff --git a/Engine/Source/Runtime/Renderer/Private/DeferredShadingRenderer.cpp b/Engine/Source/Runtime/Renderer/Private/DeferredShadingRenderer.cpp
--- a/Engine/Source/Runtime/Renderer/Private/DeferredShadingRenderer.cpp
+++ b/Engine/Source/Runtime/Renderer/Private/DeferredShadingRenderer.cpp
@@ -2044,7 +2044,7 @@ void FDeferredShadingSceneRenderer::Render(FRDGBuilder& GraphBuilder, const FSce
 			PreparePathTracingCloudMaterial(GraphBuilder, Scene, Views);
 		}
 
-		if (IsRayTracingEnabled(ViewFamily.GetShaderPlatform()) && GRHISupportsRayTracingShaders)
+		if (IsRayTracingEnabled(ViewFamily.GetShaderPlatform()) && ShouldCompileRayTracingShadersForProject(ViewFamily.GetShaderPlatform()))
 		{
 			if (!ViewFamily.EngineShowFlags.PathTracing)
 			{
