[PostProcessing] Fixed issue where showflag.VisualizePostProcessStack was drawing to an uninitialized texture

diff --git a/Engine/Source/Runtime/Renderer/Private/PostProcess/PostProcessing.cpp b/Engine/Source/Runtime/Renderer/Private/PostProcess/PostProcessing.cpp
--- a/Engine/Source/Runtime/Renderer/Private/PostProcess/PostProcessing.cpp
+++ b/Engine/Source/Runtime/Renderer/Private/PostProcess/PostProcessing.cpp
@@ -1636,7 +1636,7 @@ void AddPostProcessingPasses(
 	{
 		FScreenPassRenderTarget OverrideOutput;
 		PassSequence.AcceptOverrideIfLastPass(EPass::VisualizePostProcessStack, OverrideOutput);
-		OverrideOutput = OverrideOutput.IsValid() ? OverrideOutput : FScreenPassRenderTarget::CreateFromInput(GraphBuilder, SceneColor, View.GetOverwriteLoadAction(), TEXT("VisualizePostProcessStack"));
+		OverrideOutput = OverrideOutput.IsValid() ? OverrideOutput : FScreenPassRenderTarget(SceneColor, ERenderTargetLoadAction::ELoad);
 		SceneColor = AddFinalPostProcessDebugInfoPasses(GraphBuilder, View, OverrideOutput);
 	}
 #endif
