Fix TSAN race condition.
Move virtual texture residency tracking out of async work (where it was mistakenly moved in 41445722) and back to render thread.
#rnx

diff --git a/Engine/Source/Runtime/Renderer/Private/VT/VirtualTextureSystem.cpp b/Engine/Source/Runtime/Renderer/Private/VT/VirtualTextureSystem.cpp
--- a/Engine/Source/Runtime/Renderer/Private/VT/VirtualTextureSystem.cpp
+++ b/Engine/Source/Runtime/Renderer/Private/VT/VirtualTextureSystem.cpp
@@ -2967,10 +2967,6 @@ void FVirtualTextureSystem::BeginUpdate(FRDGBuilder& GraphBuilder, FVirtualTextu
 	{
 		TRACE_CPUPROFILER_EVENT_SCOPE(FVirtualTextureSystem::GatherAndSubmitRequests);
 
-		// This doesn't need to be synced with requests and can be updated at any point in the frame.
-		// Do it here so that it is off the render thread.
-		UpdateResidencyTracking();
-
 		Updater->MergedRequestList = Allocator.Create<FUniqueRequestList>(Allocator);
 		Updater->MergedRequestList->Initialize();
 
@@ -3152,6 +3148,7 @@ void FVirtualTextureSystem::EndUpdate(FRDGBuilder& GraphBuilder, TUniquePtr<FVir
 		SubmitThrottledRequests(GraphBuilder.RHICmdList, Updater.Get(), EUpdatePhase::End);
 	}
 
+	UpdateResidencyTracking();
 	GrowPhysicalPools();
 
 #if !UE_BUILD_SHIPPING
