Beer Shadow Map composites incorrectly under PIE because of a discrepancy between the reported field-of-view and the value in the projection matrix.

#test ContentExamples

diff --git a/Engine/Source/Runtime/Renderer/Private/HeterogeneousVolumes/HeterogeneousVolumes.cpp b/Engine/Source/Runtime/Renderer/Private/HeterogeneousVolumes/HeterogeneousVolumes.cpp
--- a/Engine/Source/Runtime/Renderer/Private/HeterogeneousVolumes/HeterogeneousVolumes.cpp
+++ b/Engine/Source/Runtime/Renderer/Private/HeterogeneousVolumes/HeterogeneousVolumes.cpp
@@ -1311,7 +1311,7 @@ void FDeferredShadingSceneRenderer::RenderHeterogeneousVolumes(
 
 								// Transform
 								const FMatrix ProjectionMatrix = View.ViewMatrices.GetProjectionMatrix();
-								float FOV = FMath::DegreesToRadians(View.FOV * 0.5);
+								float FOV = FMath::Atan(1.0f / ProjectionMatrix.M[0][0]);
 								FMatrix ViewToClip = FPerspectiveMatrix(
 									FOV,
 									ShadowMapResolution.X,
