[Ray Tracing] Fix a bug in the build instance buffer pass where incorrect GPU scene resources are used due to getting them too late. We call GPUScene->GetShaderParametersRHI() inside an RDG pass instead of getting those parameters at pass addition time and capture them as a lambda argument

#tests tested in editor and on a console

diff --git a/Engine/Plugins/Experimental/GPULightmass/Source/GPULightmass/Private/LightmapRenderer.cpp b/Engine/Plugins/Experimental/GPULightmass/Source/GPULightmass/Private/LightmapRenderer.cpp
--- a/Engine/Plugins/Experimental/GPULightmass/Source/GPULightmass/Private/LightmapRenderer.cpp
+++ b/Engine/Plugins/Experimental/GPULightmass/Source/GPULightmass/Private/LightmapRenderer.cpp
@@ -1226,6 +1226,7 @@ bool FSceneRenderState::SetupRayTracingScene(FRDGBuilder& GraphBuilder, FSceneUn
 				RHICmdList,
 				nullptr,
 				nullptr,
+				nullptr,
 				InstanceBuffer.UAV,
 				GRHIGlobals.RayTracing.RequiresSeparateHitGroupContributionsBuffer ? HitGroupContributionsBuffer.UAV : nullptr,
 				SceneInitializer.MaxNumInstances,
diff --git a/Engine/Source/Runtime/Renderer/Private/RayTracing/RayTracingInstanceBufferUtil.cpp b/Engine/Source/Runtime/Renderer/Private/RayTracing/RayTracingInstanceBufferUtil.cpp
--- a/Engine/Source/Runtime/Renderer/Private/RayTracing/RayTracingInstanceBufferUtil.cpp
+++ b/Engine/Source/Runtime/Renderer/Private/RayTracing/RayTracingInstanceBufferUtil.cpp
@@ -54,6 +54,12 @@ static TAutoConsoleVariable<bool> CVarRayTracingInstanceBufferParallelFill(
 	TEXT("Whether to fill the ray tracing instance upload buffer using parallel loops. 0=disabled, 1=enabled (default)"),
 	ECVF_RenderThreadSafe);
 
+static TAutoConsoleVariable<bool> CVarRayTracingInstanceBufferLateGetGPUSceneParameters(
+	TEXT("r.RayTracing.InstanceBuffer.LateGetGPUSceneParameters"),
+	false,
+	TEXT(""),
+	ECVF_RenderThreadSafe);
+
 struct FRayTracingInstanceGroup
 {
 	uint32 BaseInstanceIndex : 30;
@@ -666,7 +672,7 @@ IMPLEMENT_GLOBAL_SHADER(FRayTracingBuildInstanceBufferCS, "/Engine/Private/Raytr
 
 void BuildRayTracingInstanceBuffer(
 	FRHICommandList& RHICmdList,
-	const FGPUScene* GPUScene,
+	const FGPUSceneResourceParametersRHI* GPUSceneParameters,
 	const FDFVector3& PreViewTranslation,
 	uint32 MaxNumInstances,
 	uint32 NumGroups,
@@ -700,9 +706,9 @@ void BuildRayTracingInstanceBuffer(
 	PassParams.PreViewTranslationHigh = PreViewTranslation.High;
 	PassParams.PreViewTranslationLow = PreViewTranslation.Low;
 
-	if (GPUScene)
+	if (GPUSceneParameters)
 	{
-		PassParams.GPUSceneParameters = GPUScene->GetShaderParametersRHI();
+		PassParams.GPUSceneParameters = *GPUSceneParameters;
 	}
 
 	if (CullingParameters)
@@ -747,7 +753,7 @@ void BuildRayTracingInstanceBuffer(
 
 void BuildRayTracingInstanceBuffer(
 	FRHICommandList& RHICmdList,
-	const FGPUScene* GPUScene,
+	const FGPUSceneResourceParametersRHI* GPUSceneParameters,
 	const FDFVector3& PreViewTranslation,
 	FRHIUnorderedAccessView* InstancesUAV,
 	FRHIUnorderedAccessView* HitGroupContributionsUAV,
@@ -770,7 +776,7 @@ void BuildRayTracingInstanceBuffer(
 	{
 		BuildRayTracingInstanceBuffer(
 			RHICmdList,
-			GPUScene,
+			GPUSceneParameters,
 			PreViewTranslation,
 			MaxNumInstances,
 			NumGPUGroups,
@@ -794,7 +800,7 @@ void BuildRayTracingInstanceBuffer(
 	{
 		BuildRayTracingInstanceBuffer(
 			RHICmdList,
-			GPUScene,
+			GPUSceneParameters,
 			PreViewTranslation,
 			MaxNumInstances,
 			NumCPUGroups,
@@ -830,9 +836,15 @@ void BuildRayTracingInstanceBuffer(
 	FRHIUnorderedAccessView* OutputStatsUAV,
 	FRHIUnorderedAccessView* InstanceExtraDataUAV)
 {
+	FGPUSceneResourceParametersRHI GPUSceneParameters;
+	if (GPUScene)
+	{
+		GPUSceneParameters = GPUScene->GetShaderParametersRHI();
+	}
+
 	BuildRayTracingInstanceBuffer(
 		RHICmdList,
-		GPUScene,
+		GPUScene ? &GPUSceneParameters : nullptr,
 		PreViewTranslation,
 		InstancesUAV,
 		HitGroupContributionsUAV,
@@ -1034,9 +1046,15 @@ void FRayTracingInstanceBufferBuilder::BuildRayTracingInstanceBuffer(
 	uint32 OutputStatsOffset,
 	FRHIUnorderedAccessView* InstanceExtraDataUAV)
 {
+	FGPUSceneResourceParametersRHI GPUSceneParameters;
+	if (GPUScene)
+	{
+		GPUSceneParameters = GPUScene->GetShaderParametersRHI();
+	}
+
 	::BuildRayTracingInstanceBuffer(
 		RHICmdList,
-		GPUScene,
+		GPUScene ? &GPUSceneParameters : nullptr,
 		FDFVector3(PreViewTranslation),
 		InstancesUAV,
 		HitGroupContributionsUAV,
@@ -1056,6 +1074,58 @@ void FRayTracingInstanceBufferBuilder::BuildRayTracingInstanceBuffer(
 		InstanceExtraDataUAV);
 }
 
+void FRayTracingInstanceBufferBuilder::BuildRayTracingInstanceBuffer(
+	FRHICommandList& RHICmdList,
+	const FGPUScene* GPUScene,
+	const FGPUSceneResourceParametersRHI* GPUSceneParameters,
+	const FRayTracingCullingParameters* CullingParameters,
+	FRHIUnorderedAccessView* InstancesUAV,
+	FRHIUnorderedAccessView* HitGroupContributionsUAV,
+	uint32 MaxNumInstances,
+	bool bCompactOutput,
+	FRHIUnorderedAccessView* OutputStatsUAV,
+	uint32 OutputStatsOffset,
+	FRHIUnorderedAccessView* InstanceExtraDataUAV)
+{
+	if (CVarRayTracingInstanceBufferLateGetGPUSceneParameters.GetValueOnAnyThread())
+	{
+		BuildRayTracingInstanceBuffer(
+			RHICmdList,
+			GPUScene,
+			CullingParameters,
+			InstancesUAV,
+			HitGroupContributionsUAV,
+			MaxNumInstances,
+			bCompactOutput,
+			OutputStatsUAV,
+			OutputStatsOffset,
+			InstanceExtraDataUAV);
+	}
+	else
+	{
+		::BuildRayTracingInstanceBuffer(
+			RHICmdList,
+			GPUSceneParameters,
+			FDFVector3(PreViewTranslation),
+			InstancesUAV,
+			HitGroupContributionsUAV,
+			InstanceGroupUploadSRV,
+			InstanceUploadSRV,
+			AccelerationStructureAddressesBuffer.SRV,
+			TransformUploadSRV,
+			MaxNumInstances,
+			Data.NumGPUInstanceGroups,
+			Data.NumCPUInstanceGroups,
+			Data.NumGPUInstanceDescriptors,
+			Data.NumCPUInstanceDescriptors,
+			CullingParameters,
+			bCompactOutput,
+			OutputStatsUAV,
+			OutputStatsOffset,
+			InstanceExtraDataUAV);
+	}
+}
+
 PRAGMA_ENABLE_DEPRECATION_WARNINGS
 
 #endif //RHI_RAYTRACING
diff --git a/Engine/Source/Runtime/Renderer/Private/RayTracing/RayTracingScene.cpp b/Engine/Source/Runtime/Renderer/Private/RayTracing/RayTracingScene.cpp
--- a/Engine/Source/Runtime/Renderer/Private/RayTracing/RayTracingScene.cpp
+++ b/Engine/Source/Runtime/Renderer/Private/RayTracing/RayTracingScene.cpp
@@ -87,6 +87,7 @@ BEGIN_SHADER_PARAMETER_STRUCT(FBuildInstanceBufferPassParams, )
 	SHADER_PARAMETER_RDG_BUFFER_UAV(RWStructuredBuffer, OutputStats)
 	SHADER_PARAMETER_RDG_BUFFER_UAV(RWStructuredBuffer, InstanceExtraDataBuffer)
 	SHADER_PARAMETER_RDG_UNIFORM_BUFFER(FSceneUniformParameters, Scene)
+	SHADER_PARAMETER_STRUCT_INCLUDE(FGPUSceneResourceParameters, GPUSceneParameters)
 END_SHADER_PARAMETER_STRUCT()
 
 const FRayTracingScene::FInstanceHandle FRayTracingScene::INVALID_INSTANCE_HANDLE = FInstanceHandle();
@@ -401,6 +402,7 @@ void FRayTracingScene::Update(FRDGBuilder& GraphBuilder, FSceneUniformBuffer& Sc
 					
 					PassParams->InstanceExtraDataBuffer = InstanceExtraDataBufferUAV;
 					PassParams->Scene = SceneUniformBuffer.GetBuffer(GraphBuilder);
+					PassParams->GPUSceneParameters = GPUScene->GetShaderParameters(GraphBuilder);
 					PassParams->OutputStats = bCompactInstanceBuffer || bStatsEnabled ? InstanceStatsBufferUAV : nullptr;
 
 					const uint32 OutputStatsOffset = LayerIndex * MaxNumViews + ViewIndex;
@@ -418,9 +420,18 @@ void FRayTracingScene::Update(FRDGBuilder& GraphBuilder, FSceneUniformBuffer& Sc
 						bCompactInstanceBuffer
 						](FRHICommandList& RHICmdList)
 						{
+							FGPUSceneResourceParametersRHI GPUSceneParameters;
+							GPUSceneParameters.CommonParameters = PassParams->GPUSceneParameters.CommonParameters;
+							GPUSceneParameters.GPUSceneInstanceSceneData = PassParams->GPUSceneParameters.GPUSceneInstanceSceneData->GetRHI();
+							GPUSceneParameters.GPUSceneInstancePayloadData = PassParams->GPUSceneParameters.GPUSceneInstancePayloadData->GetRHI();
+							GPUSceneParameters.GPUScenePrimitiveSceneData = PassParams->GPUSceneParameters.GPUScenePrimitiveSceneData->GetRHI();
+							GPUSceneParameters.GPUSceneLightmapData = PassParams->GPUSceneParameters.GPUSceneLightmapData->GetRHI();
+							GPUSceneParameters.GPUSceneLightData = PassParams->GPUSceneParameters.GPUSceneLightData->GetRHI();
+
 							InstanceBufferBuilder.BuildRayTracingInstanceBuffer(
 								RHICmdList,
 								GPUScene,
+								&GPUSceneParameters,
 								CullingParameters,
 								PassParams->InstanceBuffer->GetRHI(),
 								GRHIGlobals.RayTracing.RequiresSeparateHitGroupContributionsBuffer ? PassParams->HitGroupContributionsBuffer->GetRHI() : nullptr,
diff --git a/Engine/Source/Runtime/Renderer/Private/Tests/RayTracingTestbed.cpp b/Engine/Source/Runtime/Renderer/Private/Tests/RayTracingTestbed.cpp
--- a/Engine/Source/Runtime/Renderer/Private/Tests/RayTracingTestbed.cpp
+++ b/Engine/Source/Runtime/Renderer/Private/Tests/RayTracingTestbed.cpp
@@ -219,6 +219,7 @@ bool RunRayTracingTestbed_RenderThread(const FString& Parameters)
 		RHICmdList,
 		nullptr,
 		nullptr,
+		nullptr,
 		InstanceBuffer.UAV,
 		GRHIGlobals.RayTracing.RequiresSeparateHitGroupContributionsBuffer ? HitGroupContributionsBuffer.UAV : nullptr,
 		SceneInitializer.MaxNumInstances,
diff --git a/Engine/Source/Runtime/Renderer/Public/RayTracingInstanceBufferUtil.h b/Engine/Source/Runtime/Renderer/Public/RayTracingInstanceBufferUtil.h
--- a/Engine/Source/Runtime/Renderer/Public/RayTracingInstanceBufferUtil.h
+++ b/Engine/Source/Runtime/Renderer/Public/RayTracingInstanceBufferUtil.h
@@ -9,6 +9,7 @@
 #if RHI_RAYTRACING
 
 class FGPUScene;
+class FGPUSceneResourceParametersRHI;
 struct FRayTracingCullingParameters;
 struct FDFVector3;
 
@@ -133,6 +134,19 @@ public:
 		uint32 OutputStatsOffset,
 		FRHIUnorderedAccessView* InstanceExtraDataUAV);
 
+	RENDERER_API void BuildRayTracingInstanceBuffer(
+		FRHICommandList& RHICmdList,
+		const FGPUScene* GPUScene,
+		const FGPUSceneResourceParametersRHI* GPUSceneParameters,
+		const FRayTracingCullingParameters* CullingParameters,
+		FRHIUnorderedAccessView* InstancesUAV,
+		FRHIUnorderedAccessView* HitGroupContributionsUAV,
+		uint32 MaxNumInstances,
+		bool bCompactOutput,
+		FRHIUnorderedAccessView* OutputStatsUAV,
+		uint32 OutputStatsOffset,
+		FRHIUnorderedAccessView* InstanceExtraDataUAV);
+
 	PRAGMA_DISABLE_DEPRECATION_WARNINGS
 
 	uint32 GetMaxNumInstances() const { return Data.NumNativeGPUSceneInstances + Data.NumNativeCPUInstances; }
