[PostProcessing] Fix off-by-one error in GetUpscaleMethod

#Synced-CL 48380560

diff --git a/Engine/Source/Runtime/Renderer/Private/PostProcess/PostProcessUpscale.cpp b/Engine/Source/Runtime/Renderer/Private/PostProcess/PostProcessUpscale.cpp
--- a/Engine/Source/Runtime/Renderer/Private/PostProcess/PostProcessUpscale.cpp
+++ b/Engine/Source/Runtime/Renderer/Private/PostProcess/PostProcessUpscale.cpp
@@ -233,9 +233,10 @@ IMPLEMENT_GLOBAL_SHADER(FUpscaleCS, "/Engine/Private/PostProcessUpscale.usf", "M
 
 EUpscaleMethod GetUpscaleMethod()
 {
-	const int32 Value = CVarUpscaleQuality.GetValueOnRenderThread();
+	const int32 StartIndex = 1; // 0 is EUpscaleMethod::None, methods start at 1
+	const int32 Value = CVarUpscaleQuality.GetValueOnRenderThread() + StartIndex;
 
-	return static_cast<EUpscaleMethod>(FMath::Clamp(Value, 0, static_cast<int32>(EUpscaleMethod::Gaussian)));
+	return static_cast<EUpscaleMethod>(FMath::Clamp(Value, StartIndex, static_cast<int32>(EUpscaleMethod::Gaussian)));
 }
 
 // static
@@ -331,7 +332,7 @@ FScreenPassTexture ISpatialUpscaler::AddDefaultUpscalePass(
 		ComputePassParameters->RWOutputTexture = GraphBuilder.CreateUAV(ComputeRenderTarget);
 		ComputePassParameters->UpscaleMethod = (uint32)Method;
 
-		ESharpeningQuality SharpeningQuality = (ESharpeningQuality)FMath::Clamp(CVarUpscaleSharpeningQuality.GetValueOnRenderThread(), 0, (int32)ESharpeningQuality::MAX); 
+		ESharpeningQuality SharpeningQuality = (ESharpeningQuality)FMath::Clamp(CVarUpscaleSharpeningQuality.GetValueOnRenderThread(), 0, static_cast<int32>(ESharpeningQuality::MAX) - 1); 
 		SharpeningQuality = Sharpening != 0.0f ? SharpeningQuality : ESharpeningQuality::Off;
 
 		FUpscaleCS::FPermutationDomain PermutationVector;
