PR #14064: Renderer: Export FViewInfo::GetPrimaryView to fix linker errors

diff --git a/Engine/Source/Runtime/Renderer/Private/SceneRendering.h b/Engine/Source/Runtime/Renderer/Private/SceneRendering.h
--- a/Engine/Source/Runtime/Renderer/Private/SceneRendering.h
+++ b/Engine/Source/Runtime/Renderer/Private/SceneRendering.h
@@ -1781,7 +1781,7 @@ public:
 	FViewShaderParameters GetShaderParameters() const;
 
 	/** Returns the primary view associated with the input view, or null if none exists. */
-	const FViewInfo* GetPrimaryView() const;
+	RENDERER_API const FViewInfo* GetPrimaryView() const;
 
 	/** Returns the instanced view associated with the input view, or null if none exists. */
 	inline const FViewInfo* GetInstancedView() const { return static_cast<const FViewInfo*>(GetInstancedSceneView()); }
