[HWRT] Added extra logging to IsCachedRayTracingGeometryValid() to help investigate failing assert.

diff --git a/Engine/Source/Runtime/Renderer/Private/PrimitiveSceneInfo.cpp b/Engine/Source/Runtime/Renderer/Private/PrimitiveSceneInfo.cpp
--- a/Engine/Source/Runtime/Renderer/Private/PrimitiveSceneInfo.cpp
+++ b/Engine/Source/Runtime/Renderer/Private/PrimitiveSceneInfo.cpp
@@ -390,17 +390,48 @@ bool FPrimitiveSceneInfo::IsCachedRayTracingGeometryValid() const
 	{
 		// TODO: Doesn't take Nanite Ray Tracing into account
 
+		check(!CachedRayTracingGeometry->GetRequiresBuild() && !CachedRayTracingGeometry->HasPendingBuildRequest());
+
+#if DO_CHECK
+		auto LogErrorDebugInfo = [this](uint32 RayTracingInstanceIndex)
+			{
+				FRHIRayTracingGeometry* CachedInstanceGeometryRHI = Scene->RayTracingScene.GetCachedInstanceGeometry(RayTracingInstanceIndex);
+				FRHIRayTracingGeometry* CachedRayTracingGeometryRHI = CachedRayTracingGeometry->GetRHI();
+				UE_LOG(LogRenderer, Error, TEXT("IsCachedRayTracingGeometryValid() Error Debug Info \
+					PrimitiveRayTracingFlags: %u \
+					bIsCachedRayTracingInstanceValid: %d \
+					bCachedInstanceHasGeometryRHI: %d, \
+					bCachedRayTracingGeometryHasRHI: %d, \
+					bCachedRayTracingGeometryIsValid: %d, \
+					bCachedRayTracingGeometryIsEvicted: %d, \
+					CachedRayTracingGeometryFlags: %u"),
+					Scene->PrimitiveRayTracingFlags[GetIndex()],
+					bIsCachedRayTracingInstanceValid ? 1 : 0,
+					CachedInstanceGeometryRHI != nullptr ? 1 : 0,
+					CachedRayTracingGeometryRHI != nullptr ? 1 : 0,
+					CachedRayTracingGeometry->IsValid() ? 1 : 0,
+					CachedRayTracingGeometry->IsEvicted() ? 1 : 0,
+					uint32(CachedRayTracingGeometry->GetGeometryState()));
+			};
+
 		if (RayTracingInstanceIndexMain != UINT32_MAX)
 		{
-			check(Scene->RayTracingScene.GetCachedInstanceGeometry(RayTracingInstanceIndexMain) == CachedRayTracingGeometry->GetRHI());
+			if (Scene->RayTracingScene.GetCachedInstanceGeometry(RayTracingInstanceIndexMain) != CachedRayTracingGeometry->GetRHI())
+			{
+				LogErrorDebugInfo(RayTracingInstanceIndexMain);
+				checkf(false, TEXT("Cached ray tracing instance has an unexpected geometry. Check log above."))
+			}
 		}
 
 		if (RayTracingInstanceIndexDecal != UINT32_MAX)
 		{
-			check(Scene->RayTracingScene.GetCachedInstanceGeometry(RayTracingInstanceIndexDecal) == CachedRayTracingGeometry->GetRHI());
+			if (Scene->RayTracingScene.GetCachedInstanceGeometry(RayTracingInstanceIndexDecal) != CachedRayTracingGeometry->GetRHI())
+			{
+				LogErrorDebugInfo(RayTracingInstanceIndexDecal);
+				checkf(false, TEXT("Cached ray tracing instance has an unexpected geometry. Check log above."))
+			}
 		}
-
-		check(!CachedRayTracingGeometry->GetRequiresBuild() && !CachedRayTracingGeometry->HasPendingBuildRequest());
+#endif // DO_CHECK
 
 		return CachedRayTracingGeometry->IsValid() && !CachedRayTracingGeometry->IsEvicted();
 	}
