Path Tracing: Maintain a tracked SBT for the path tracer to avoid issues when mixing PT and non-PT views of the same Scene

This avoids crashes when using the path tracer together with nDisplay light cards.

#Synced-CL 49605136

diff --git a/Engine/Source/Runtime/Renderer/Private/DeferredShadingRenderer.cpp b/Engine/Source/Runtime/Renderer/Private/DeferredShadingRenderer.cpp
--- a/Engine/Source/Runtime/Renderer/Private/DeferredShadingRenderer.cpp
+++ b/Engine/Source/Runtime/Renderer/Private/DeferredShadingRenderer.cpp
@@ -747,16 +747,18 @@ bool FDeferredShadingSceneRenderer::SetupRayTracingPipelineStatesAndSBT(FRDGBuil
 
 			bMaterialRayTracingSBTRequired = true;
 
+			auto* SBT = bIsPathTracing ? &(Scene->RayTracingPersistentSBTState.PathTracingSBT) : &(Scene->RayTracingPersistentSBTState.MaterialSBT);
+
 			// Allocate persistent SBT RHI if necessary
 			// (only when shader binding layout is enabled which is required to be able to share SBT between views)
 			if (RHIGetStaticShaderBindingLayoutSupport(ShaderPlatform) != ERHIStaticShaderBindingLayoutSupport::Unsupported
-				&& Scene->RayTracingPersistentSBTState.MaterialSBT.ID == INDEX_NONE)
+				&& SBT->ID == INDEX_NONE)
 			{
-				Scene->RayTracingPersistentSBTState.MaterialSBT.ID = Scene->RayTracingSBT.AllocatePersistentSBTID(GraphBuilder.RHICmdList, ERayTracingShaderBindingMode::RTPSO);
+				SBT->ID = Scene->RayTracingSBT.AllocatePersistentSBTID(GraphBuilder.RHICmdList, ERayTracingShaderBindingMode::RTPSO);
 			}
 
 			// Material SBT release handled by ReleaseUnusedRayTracingSBTs
-			Scene->RayTracingPersistentSBTState.MaterialSBT.bUsedThisFrame = true;
+			SBT->bUsedThisFrame = true;
 		}
 	}
 
@@ -840,25 +842,24 @@ bool FDeferredShadingSceneRenderer::SetupRayTracingPipelineStatesAndSBT(FRDGBuil
 
 			bool bFirstViewWithRayTracingPass = true;
 
-			for (int32 ViewIndex = 0; ViewIndex < Views.Num(); ++ViewIndex)
+			for (FViewInfo& View : Views)
 			{
-				FViewInfo& View = Views[ViewIndex];
-
 				if (View.bHasAnyRayTracingPass)
 				{
-					if (Scene->RayTracingPersistentSBTState.MaterialSBT.ID != INDEX_NONE)
+					auto* SBT = bIsPathTracing ? &(Scene->RayTracingPersistentSBTState.PathTracingSBT) : &(Scene->RayTracingPersistentSBTState.MaterialSBT);
+					if (SBT->ID != INDEX_NONE)
 					{
-						View.MaterialRayTracingData.ShaderBindingTable = Scene->RayTracingSBT.GetPersistentSBT(Scene->RayTracingPersistentSBTState.MaterialSBT.ID);
+						View.MaterialRayTracingData.ShaderBindingTable = Scene->RayTracingSBT.GetPersistentSBT(SBT->ID);
 
 						// Cache the PipelineState in Scene if View has one, so other ViewFamilies can reuse it
 						if (View.MaterialRayTracingData.PipelineState)
 						{
-							Scene->RayTracingPersistentSBTState.MaterialSBT.LastValidPipelineState = View.MaterialRayTracingData.PipelineState;
+							SBT->LastValidPipelineState = View.MaterialRayTracingData.PipelineState;
 						}
 						// Use cached PipelineState as fallback if View doesn't have one
-						else if (Scene->RayTracingPersistentSBTState.MaterialSBT.LastValidPipelineState)
+						else if (SBT->LastValidPipelineState)
 						{
-							View.MaterialRayTracingData.PipelineState = Scene->RayTracingPersistentSBTState.MaterialSBT.LastValidPipelineState;
+							View.MaterialRayTracingData.PipelineState = SBT->LastValidPipelineState;
 						}
 					}
 					else
diff --git a/Engine/Source/Runtime/Renderer/Private/RendererScene.cpp b/Engine/Source/Runtime/Renderer/Private/RendererScene.cpp
--- a/Engine/Source/Runtime/Renderer/Private/RendererScene.cpp
+++ b/Engine/Source/Runtime/Renderer/Private/RendererScene.cpp
@@ -4185,11 +4185,13 @@ void FScene::ReleaseUnusedPersistentSBTs()
 
 	// Check and release each SBT type
 	ReleaseSBTIfUnused(RayTracingPersistentSBTState.MaterialSBT);
+	ReleaseSBTIfUnused(RayTracingPersistentSBTState.PathTracingSBT);
 	ReleaseSBTIfUnused(RayTracingPersistentSBTState.LumenSBT);
 	ReleaseSBTIfUnused(RayTracingPersistentSBTState.InlineSBT);
 
 	// Reset our frame usage flags since this function is meant to be called once per scene at the end of the frame.
 	RayTracingPersistentSBTState.MaterialSBT.bUsedThisFrame = false;
+	RayTracingPersistentSBTState.PathTracingSBT.bUsedThisFrame = false;
 	RayTracingPersistentSBTState.LumenSBT.bUsedThisFrame = false;
 	RayTracingPersistentSBTState.InlineSBT.bUsedThisFrame = false;
 }
diff --git a/Engine/Source/Runtime/Renderer/Private/ScenePrivate.h b/Engine/Source/Runtime/Renderer/Private/ScenePrivate.h
--- a/Engine/Source/Runtime/Renderer/Private/ScenePrivate.h
+++ b/Engine/Source/Runtime/Renderer/Private/ScenePrivate.h
@@ -3215,6 +3215,7 @@ public:
 		};
 
 		FTrackedSBT MaterialSBT;
+		FTrackedSBT PathTracingSBT;
 		FTrackedSBT LumenSBT;
 		FTrackedSBT InlineSBT;
 	};
