RHI: Enable SBT immediate release by default, restoring original behavior prior to 47757289.

#rnx

diff --git a/Engine/Source/Runtime/Renderer/Private/DeferredShadingRenderer.cpp b/Engine/Source/Runtime/Renderer/Private/DeferredShadingRenderer.cpp
--- a/Engine/Source/Runtime/Renderer/Private/DeferredShadingRenderer.cpp
+++ b/Engine/Source/Runtime/Renderer/Private/DeferredShadingRenderer.cpp
@@ -257,6 +257,16 @@ static TAutoConsoleVariable<int32> CVarRayTracingAsyncBuild(
 	ECVF_RenderThreadSafe
 );
 
+int32 GRayTracingSBTImmediateRelease = 1;
+static FAutoConsoleVariableRef CVarRayTracingSBTImmediateRelease(
+	TEXT("r.RayTracing.SBT.ImmediateRelease"),
+	GRayTracingSBTImmediateRelease,
+	TEXT("When enabled, SBTs are released immediately when not needed by the current ViewFamily\n")
+	TEXT(" 0: SBTs are kept alive for re-use during the same frame and can improve performance of multi-viewfamily scenarios\n")
+	TEXT(" 1: SBTs are released immediately when not needed by current ViewFamily"),
+	ECVF_RenderThreadSafe
+);
+
 static int32 GRayTracingMultiGpuTLASMask = 0;
 static FAutoConsoleVariableRef CVarRayTracingMultiGpuTLASMask(
 	TEXT("r.RayTracing.MultiGpuMaskTLAS"),
@@ -757,9 +767,22 @@ bool FDeferredShadingSceneRenderer::SetupRayTracingPipelineStatesAndSBT(FRDGBuil
 				SBT->ID = Scene->RayTracingSBT.AllocatePersistentSBTID(GraphBuilder.RHICmdList, ERayTracingShaderBindingMode::RTPSO);
 			}
 
-			// Material SBT release handled by ReleaseUnusedRayTracingSBTs
+			// Material SBT release handled by ReleaseUnusedRayTracingSBTs if GRayTracingSBTImmediateRelease is 0.
 			SBT->bUsedThisFrame = true;
 		}
+		else if (GRayTracingSBTImmediateRelease)
+		{
+			FScene::FRayTracingPersistentSBTState::FTrackedSBT& SBT = 
+				bIsPathTracing ? 
+				(Scene->RayTracingPersistentSBTState.PathTracingSBT) 
+				: (Scene->RayTracingPersistentSBTState.MaterialSBT);
+
+			if (SBT.ID != INDEX_NONE)
+			{
+				Scene->RayTracingSBT.ReleasePersistentSBT(SBT.ID);
+				SBT.Reset();
+			}
+		}
 	}
 
 	if (GRHISupportsRayTracingShaders)
@@ -813,9 +836,14 @@ bool FDeferredShadingSceneRenderer::SetupRayTracingPipelineStatesAndSBT(FRDGBuil
 				Scene->RayTracingPersistentSBTState.LumenSBT.ID = Scene->RayTracingSBT.AllocatePersistentSBTID(GraphBuilder.RHICmdList, ERayTracingShaderBindingMode::RTPSO);
 			}
 
-			// Lumen SBT release handled by ReleaseUnusedRayTracingSBTs
+			// Lumen SBT release handled by ReleaseUnusedRayTracingSBTs if GRayTracingSBTImmediateRelease is 0.
 			Scene->RayTracingPersistentSBTState.LumenSBT.bUsedThisFrame = true;
 		}
+		else if (GRayTracingSBTImmediateRelease && (Scene->RayTracingPersistentSBTState.LumenSBT.ID != INDEX_NONE))
+		{
+			Scene->RayTracingSBT.ReleasePersistentSBT(Scene->RayTracingPersistentSBTState.LumenSBT.ID);
+			Scene->RayTracingPersistentSBTState.LumenSBT.Reset();
+		}
 	}
 
 	// Check if inline SBT is needed or not
@@ -826,11 +854,16 @@ bool FDeferredShadingSceneRenderer::SetupRayTracingPipelineStatesAndSBT(FRDGBuil
 			Scene->RayTracingPersistentSBTState.InlineSBT.ID = Scene->RayTracingSBT.AllocatePersistentSBTID(GraphBuilder.RHICmdList, ERayTracingShaderBindingMode::Inline);
 		}
 
-		// Inline SBT release handled by ReleaseUnusedRayTracingSBTs
+		// Inline SBT release handled by ReleaseUnusedRayTracingSBTs if GRayTracingSBTImmediateRelease is 0.
 		Scene->RayTracingPersistentSBTState.InlineSBT.bUsedThisFrame = true;
 
 		EnumAddFlags(ShaderBindingMode, ERayTracingShaderBindingMode::Inline);
 	}
+	else if (GRayTracingSBTImmediateRelease && Scene->RayTracingPersistentSBTState.InlineSBT.ID != INDEX_NONE)
+	{
+		Scene->RayTracingSBT.ReleasePersistentSBT(Scene->RayTracingPersistentSBTState.InlineSBT.ID);
+		Scene->RayTracingPersistentSBTState.InlineSBT.Reset();
+	}
 
 	if (ShaderBindingMode != ERayTracingShaderBindingMode::Disabled)
 	{
@@ -852,14 +885,16 @@ bool FDeferredShadingSceneRenderer::SetupRayTracingPipelineStatesAndSBT(FRDGBuil
 						View.MaterialRayTracingData.ShaderBindingTable = Scene->RayTracingSBT.GetPersistentSBT(SBT->ID);
 
 						// Cache the PipelineState in Scene if View has one, so other ViewFamilies can reuse it
-						if (View.MaterialRayTracingData.PipelineState)
+						if (!GRayTracingSBTImmediateRelease)
 						{
-							SBT->LastValidPipelineState = View.MaterialRayTracingData.PipelineState;
-						}
-						// Use cached PipelineState as fallback if View doesn't have one
-						else if (SBT->LastValidPipelineState)
-						{
-							View.MaterialRayTracingData.PipelineState = SBT->LastValidPipelineState;
+							if (View.MaterialRayTracingData.PipelineState)
+							{
+								SBT->LastValidPipelineState = View.MaterialRayTracingData.PipelineState;
+							}
+							else if (SBT->LastValidPipelineState)
+							{
+								View.MaterialRayTracingData.PipelineState = SBT->LastValidPipelineState;
+							}
 						}
 					}
 					else
@@ -872,14 +907,16 @@ bool FDeferredShadingSceneRenderer::SetupRayTracingPipelineStatesAndSBT(FRDGBuil
 						View.LumenRayTracingData.ShaderBindingTable = Scene->RayTracingSBT.GetPersistentSBT(Scene->RayTracingPersistentSBTState.LumenSBT.ID);
 
 						// Cache the PipelineState in Scene if View has one, so other ViewFamilies can reuse it
-						if (View.LumenRayTracingData.PipelineState)
-						{
-							Scene->RayTracingPersistentSBTState.LumenSBT.LastValidPipelineState = View.LumenRayTracingData.PipelineState;
-						}
-						// Use cached PipelineState as fallback if View doesn't have one
-						else if (Scene->RayTracingPersistentSBTState.LumenSBT.LastValidPipelineState)
+						if (!GRayTracingSBTImmediateRelease)
 						{
-							View.LumenRayTracingData.PipelineState = Scene->RayTracingPersistentSBTState.LumenSBT.LastValidPipelineState;
+							if (View.LumenRayTracingData.PipelineState)
+							{
+								Scene->RayTracingPersistentSBTState.LumenSBT.LastValidPipelineState = View.LumenRayTracingData.PipelineState;
+							}
+							else if (Scene->RayTracingPersistentSBTState.LumenSBT.LastValidPipelineState)
+							{
+								View.LumenRayTracingData.PipelineState = Scene->RayTracingPersistentSBTState.LumenSBT.LastValidPipelineState;
+							}
 						}
 					}
 					else
diff --git a/Engine/Source/Runtime/Renderer/Private/RendererScene.cpp b/Engine/Source/Runtime/Renderer/Private/RendererScene.cpp
--- a/Engine/Source/Runtime/Renderer/Private/RendererScene.cpp
+++ b/Engine/Source/Runtime/Renderer/Private/RendererScene.cpp
@@ -4168,8 +4168,21 @@ void FScene::UpdateCachedRayTracingState(FPrimitiveSceneProxy* SceneProxy)
 	}
 }
 
+extern int32 GRayTracingSBTImmediateRelease;
+
 void FScene::ReleaseUnusedPersistentSBTs()
 {
+	// When immediate release is enabled, it is no longer the responsibility of this function to release SBTs.
+	if (GRayTracingSBTImmediateRelease)
+	{
+		RayTracingPersistentSBTState.MaterialSBT.bUsedThisFrame = false;
+		RayTracingPersistentSBTState.PathTracingSBT.bUsedThisFrame = false;
+		RayTracingPersistentSBTState.LumenSBT.bUsedThisFrame = false;
+		RayTracingPersistentSBTState.InlineSBT.bUsedThisFrame = false;
+
+		return;
+	}
+
 	// This function is called at the end of frame, after all viewfamilies being rendered
 	// had a chance to mark used SBTs as such.
 
