[HWRT] Fix race condition between AddDynamicInstancesToRayTracingScene(...) and GatherRelevantStaticPrimitives(...)
- FPrimitiveSceneInfo::IsCachedRayTracingGeometryValid(...) was reading from Instances arrays in FRayTracingScene while those arrays could be resized by AddDynamicInstancesToRayTracingScene(...).
- Moved some checks later to AddVisibleCachedInstances(...) when access to FRayTracingScene is serialized using AddInstancesPipe.

#Synced-CL 50307718

diff --git a/Engine/Source/Runtime/Renderer/Private/PrimitiveSceneInfo.cpp b/Engine/Source/Runtime/Renderer/Private/PrimitiveSceneInfo.cpp
--- a/Engine/Source/Runtime/Renderer/Private/PrimitiveSceneInfo.cpp
+++ b/Engine/Source/Runtime/Renderer/Private/PrimitiveSceneInfo.cpp
@@ -392,47 +392,6 @@ bool FPrimitiveSceneInfo::IsCachedRayTracingGeometryValid() const
 
 		check(!CachedRayTracingGeometry->GetRequiresBuild() && !CachedRayTracingGeometry->HasPendingBuildRequest());
 
-#if DO_CHECK
-		auto LogErrorDebugInfo = [this](uint32 RayTracingInstanceIndex)
-			{
-				FRHIRayTracingGeometry* CachedInstanceGeometryRHI = Scene->RayTracingScene.GetCachedInstanceGeometry(RayTracingInstanceIndex);
-				FRHIRayTracingGeometry* CachedRayTracingGeometryRHI = CachedRayTracingGeometry->GetRHI();
-				UE_LOG(LogRenderer, Error, TEXT("IsCachedRayTracingGeometryValid() Error Debug Info \
-					PrimitiveRayTracingFlags: %u \
-					bIsCachedRayTracingInstanceValid: %d \
-					bCachedInstanceHasGeometryRHI: %d, \
-					bCachedRayTracingGeometryHasRHI: %d, \
-					bCachedRayTracingGeometryIsValid: %d, \
-					bCachedRayTracingGeometryIsEvicted: %d, \
-					CachedRayTracingGeometryFlags: %u"),
-					Scene->PrimitiveRayTracingFlags[GetIndex()],
-					bIsCachedRayTracingInstanceValid ? 1 : 0,
-					CachedInstanceGeometryRHI != nullptr ? 1 : 0,
-					CachedRayTracingGeometryRHI != nullptr ? 1 : 0,
-					CachedRayTracingGeometry->IsValid() ? 1 : 0,
-					CachedRayTracingGeometry->IsEvicted() ? 1 : 0,
-					uint32(CachedRayTracingGeometry->GetGeometryState()));
-			};
-
-		if (RayTracingInstanceIndexMain != UINT32_MAX)
-		{
-			if (Scene->RayTracingScene.GetCachedInstanceGeometry(RayTracingInstanceIndexMain) != CachedRayTracingGeometry->GetRHI())
-			{
-				LogErrorDebugInfo(RayTracingInstanceIndexMain);
-				checkf(false, TEXT("Cached ray tracing instance has an unexpected geometry. Check log above."))
-			}
-		}
-
-		if (RayTracingInstanceIndexDecal != UINT32_MAX)
-		{
-			if (Scene->RayTracingScene.GetCachedInstanceGeometry(RayTracingInstanceIndexDecal) != CachedRayTracingGeometry->GetRHI())
-			{
-				LogErrorDebugInfo(RayTracingInstanceIndexDecal);
-				checkf(false, TEXT("Cached ray tracing instance has an unexpected geometry. Check log above."))
-			}
-		}
-#endif // DO_CHECK
-
 		return CachedRayTracingGeometry->IsValid() && !CachedRayTracingGeometry->IsEvicted();
 	}
 
diff --git a/Engine/Source/Runtime/Renderer/Private/RayTracing/RayTracing.cpp b/Engine/Source/Runtime/Renderer/Private/RayTracing/RayTracing.cpp
--- a/Engine/Source/Runtime/Renderer/Private/RayTracing/RayTracing.cpp
+++ b/Engine/Source/Runtime/Renderer/Private/RayTracing/RayTracing.cpp
@@ -223,6 +223,9 @@ namespace RayTracing
 		int32 SBTAllocationUniqueId = INDEX_NONE;
 		uint32 MainRayTracingInstanceIndex = UINT32_MAX;
 		uint32 DecalRayTracingInstanceIndex = UINT32_MAX;
+#if DO_CHECK
+		FRHIRayTracingGeometry* RayTracingGeometry = nullptr;
+#endif
 	};
 
 	struct FRelevantPrimitive
@@ -1423,6 +1426,13 @@ namespace RayTracing
 					RelevantPrimitive->SBTAllocationUniqueId = RTLODData.SBTAllocationUniqueId;
 					RelevantPrimitive->MainRayTracingInstanceIndex = MainRayTracingInstanceIndex;
 					RelevantPrimitive->DecalRayTracingInstanceIndex = DecalRayTracingInstanceIndex;
+#if DO_CHECK
+					if (!bUsingNaniteRayTracing)
+					{
+						// Currently RTInstanceData.RayTracingGeometry is only available for non-nanite geometries
+						RelevantPrimitive->RayTracingGeometry = SceneInfo->GetCachedRayTracingGeometry()->GetRHI();
+					}
+#endif
 
 					const int32 NumBindings = RTLODData.CachedShaderBindingDataBase.Num() + RTLODData.CachedShaderBindingDataDecal.Num();
 
@@ -1770,11 +1780,17 @@ namespace RayTracing
 
 			if (RelevantPrimitive.MainRayTracingInstanceIndex != UINT32_MAX)
 			{
+#if DO_CHECK
+				check(RelevantPrimitive.RayTracingGeometry == nullptr || RayTracingScene.GetCachedInstanceGeometry(RelevantPrimitive.MainRayTracingInstanceIndex) == RelevantPrimitive.RayTracingGeometry);
+#endif
 				RayTracingScene.MarkInstanceVisible(RelevantPrimitive.MainRayTracingInstanceIndex, View.GetRayTracingSceneViewHandle());
 			}
 
 			if (RelevantPrimitive.DecalRayTracingInstanceIndex != UINT32_MAX && !bExcludeDecals)
 			{
+#if DO_CHECK
+				check(RelevantPrimitive.RayTracingGeometry == nullptr || RayTracingScene.GetCachedInstanceGeometry(RelevantPrimitive.DecalRayTracingInstanceIndex) == RelevantPrimitive.RayTracingGeometry);
+#endif
 				RayTracingScene.MarkInstanceVisible(RelevantPrimitive.DecalRayTracingInstanceIndex, View.GetRayTracingSceneViewHandle());
 			}
 
