Various ACES 2.0 changes

- Change limiting and encoding color spaces from AP1 to Rec. 709/2100 to match display outputs.
- Change ACESOutputTransform to return the color in the encoding color space to remove excess chromatic adaptation that ACES presumably has already applied.
- Fix apply_white_scale never being called because UE_ACES_SCALE_WHITE is false. Currently, none of our code calls true though because our limiting/encoding spaces always have matching white points (both D65) for now.
- Remove white_limiting function, and fix clamp_zero_to_peakLuminance being called from the wrong place in the inverse function.
- Fix LUT not being updated if r.HDR.Aces.Version changes.
- Remove redundant peakLuminance parameter from outputTransform_fwd and outputTransform_inv functions.

diff --git a/Engine/Shaders/Private/ACES/ACESDisplayEncoding.ush b/Engine/Shaders/Private/ACES/ACESDisplayEncoding.ush
--- a/Engine/Shaders/Private/ACES/ACESDisplayEncoding.ush
+++ b/Engine/Shaders/Private/ACES/ACESDisplayEncoding.ush
@@ -15,6 +15,15 @@ float3 apply_white_scale( float3 XYZluminance,
                             ODTParams PARAMS, 
                             bool _invert)
 {
+    // If the creative white differs from the calibration white of the display,
+    // unequal display code values will be required to produce the neutral of
+    // the creative white. Without scaling, one channel would hit the max value
+    // first while the other channels continue to increase, resulting in a hue
+    // shift. To avoid this, the white scaling finds the largest channel and
+    // applies a scale factor to force the point where this channel hits max to
+    // 1.0, assuring that all three channels "fit" within the peak value. In the
+    // inverse direction, the white scaling is removed.
+
     float3 RGB_w = mul( PARAMS.OUTPUT_XYZ_TO_RGB, PARAMS.XYZ_w_limit );
     float3 RGB_w_f = (1./referenceLuminance) * RGB_w;
     float scale = 1. / max( max(RGB_w_f[0], RGB_w_f[1]), RGB_w_f[2]); 
@@ -35,31 +44,3 @@ float3 clamp_zero_to_peakLuminance( float3 XYZ_in,
 
     return XYZ;
 }
-
-float3 white_limiting(float3 XYZ_in,
-                        ODTParams PARAMS,
-                        bool scale_white,
-                        bool _invert = false)
-{
-    // The white limiting step clamps the output to between zero and the peak
-    // luminance value so that values do not exceed the maximum value of the
-    // tonescale. 
-    // If the creative white differs from the calibration white of the display,
-    // unequal display code values will be required to produce the neutral of
-    // the creative white. Without scaling, one channel would hit the max value
-    // first while the other channels continue to increase, resulting in a hue
-    // shift. To avoid this, the white scaling finds the largest channel and
-    // applies a scale factor to force the point where this channel hits max to
-    // 1.0, assuring that all three channels "fit" within the peak value. In the
-    // inverse direction, the white scaling is removed.
-
-    float3 XYZ = XYZ_in;
-
-    // Clamp to peak luminance
-    XYZ = clamp_zero_to_peakLuminance(XYZ, PARAMS);
-
-    // White point scaling
-    XYZ = scale_white ? apply_white_scale(XYZ, PARAMS, _invert) : XYZ;
-
-    return XYZ;
-}
\ No newline at end of file
diff --git a/Engine/Shaders/Private/ACES/ACESOutputTransform.ush b/Engine/Shaders/Private/ACES/ACESOutputTransform.ush
--- a/Engine/Shaders/Private/ACES/ACESOutputTransform.ush
+++ b/Engine/Shaders/Private/ACES/ACESOutputTransform.ush
@@ -986,10 +986,9 @@ float3 gamutMap_inv(float3 JMh,
 }
 
 float3 outputTransform_fwd( float3 aces,
-                            float peakLuminance,
                             ODTParams PARAMS)
 {
-    float3 JMh = aces_to_JMh(aces, peakLuminance);
+    float3 JMh = aces_to_JMh(aces, PARAMS.peakLuminance);
 
     float3 tonemappedJMh = tonemapAndCompress_fwd(JMh, PARAMS);
 
@@ -1002,7 +1001,6 @@ float3 outputTransform_fwd( float3 aces,
 }
 
 float3 outputTransform_inv( float3 XYZ,
-                            float peakLuminance,
                             ODTParams PARAMS)
 {
     float3 compressedJMh = XYZ_output_to_JMh(XYZ,
@@ -1015,7 +1013,7 @@ float3 outputTransform_inv( float3 XYZ,
                                           PARAMS);
 
     float3 aces = JMh_to_aces(JMh,
-                                peakLuminance);
+                                PARAMS.peakLuminance);
 
     return aces;
 }
\ No newline at end of file
diff --git a/Engine/Shaders/Private/PostProcessCombineLUTs.usf b/Engine/Shaders/Private/PostProcessCombineLUTs.usf
--- a/Engine/Shaders/Private/PostProcessCombineLUTs.usf
+++ b/Engine/Shaders/Private/PostProcessCombineLUTs.usf
@@ -335,7 +335,7 @@ float4 CombineLUTsCommon(float2 InUV, uint InLayerIndex)
 
 	// sRGB, user specified gamut
 	if( GetOutputDevice() == TONEMAPPER_OUTPUT_sRGB)
-	{		
+	{
 		// Convert from sRGB to specified output gamut	
 		//float3 OutputGamutColor = mul( AP1_2_Output, mul( sRGB_2_AP1, FilmColor ) );
 
@@ -361,38 +361,59 @@ float4 CombineLUTsCommon(float2 InUV, uint InLayerIndex)
 	else if( GetOutputDevice() == TONEMAPPER_OUTPUT_ACES1000nitST2084 || GetOutputDevice() == TONEMAPPER_OUTPUT_ACES2000nitST2084)
 	{
 #if USE_ACES_2
-		float3 ODTColor = ACESOutputTransform(GradedColor * ACESSceneColorMultiplier, (float3x3) WorkingColorSpace.ToAP0, OutputMaxLuminance, ACESReachTable, ACESGamutTable, ACESGammaTable);
+		// TODO: The display limiting color space should be user-exposed as part of display characterization.
+		float3 ODTColor = ACESOutputTransform(
+			GradedColor * ACESSceneColorMultiplier,
+			(float3x3)WorkingColorSpace.ToAP0,
+			OutputMaxLuminance,
+			ACESReachTable,
+			ACESGamutTable,
+			ACESGammaTable,
+			false, // bScaleWhite
+			GetColorSpace(TONEMAPPER_GAMUT_Rec2020_D65),
+			GetColorSpace(TONEMAPPER_GAMUT_Rec2020_D65));
 #else
 		FACESTonemapParams AcesParams = ComputeACESTonemapParams(ACESMinMaxData, ACESMidData, ACESCoefsLow_0, ACESCoefsHigh_0, ACESCoefsLow_4, ACESCoefsHigh_4, ACESSceneColorMultiplier, ACESGamutCompression);
 		float3 ODTColor = ACESOutputTransform( GradedColor, (float3x3)WorkingColorSpace.ToAP0, AcesParams);
-#endif
 		ODTColor = mul(AP1_2_Output, ODTColor);
+#endif
 
 		// Apply conversion to ST-2084 (Dolby PQ)
 		OutDeviceColor = LinearToST2084( ODTColor );
 	}
 
-	// ACES transform with PQ/2084 encoding, AP1 gamut
+	// ACES transform to linear ScRGB
 	else if( GetOutputDevice() == TONEMAPPER_OUTPUT_ACES1000nitScRGB || GetOutputDevice() == TONEMAPPER_OUTPUT_ACES2000nitScRGB)
 	{
 #if USE_ACES_2
-		float3 ODTColor = ACESOutputTransform(GradedColor * ACESSceneColorMultiplier, (float3x3) WorkingColorSpace.ToAP0, OutputMaxLuminance, ACESReachTable, ACESGamutTable, ACESGammaTable);
+		// TODO: The display limiting color space should be user-exposed as part of display characterization.
+		float3 ODTColor = ACESOutputTransform(
+			GradedColor * ACESSceneColorMultiplier,
+			(float3x3)WorkingColorSpace.ToAP0,
+			OutputMaxLuminance,
+			ACESReachTable,
+			ACESGamutTable,
+			ACESGammaTable,
+			false, // bScaleWhite
+			GetColorSpace(TONEMAPPER_GAMUT_Rec2020_D65),
+			GetColorSpace(TONEMAPPER_GAMUT_sRGB_D65));
+		OutDeviceColor = ODTColor / UE_SCRGB_WHITE_NITS;
 #else
 		FACESTonemapParams AcesParams = ComputeACESTonemapParams(ACESMinMaxData, ACESMidData, ACESCoefsLow_0, ACESCoefsHigh_0, ACESCoefsLow_4, ACESCoefsHigh_4, ACESSceneColorMultiplier, ACESGamutCompression);
 		float3 ODTColor = ACESOutputTransform( GradedColor, (float3x3)WorkingColorSpace.ToAP0, AcesParams);
-#endif
 		OutDeviceColor = mul(AP1_2_sRGB, ODTColor / UE_SCRGB_WHITE_NITS);
-	}	
+#endif
+	}
 	
 	else if( GetOutputDevice() == TONEMAPPER_OUTPUT_LinearEXR)
 	{
-			float3 OutputGamutColor = mul( AP1_2_Output, mul( (float3x3)WorkingColorSpace.ToAP1, GradedColor ) );
-			OutDeviceColor = LinearToST2084( OutputGamutColor );
+		float3 OutputGamutColor = mul( AP1_2_Output, mul( (float3x3)WorkingColorSpace.ToAP1, GradedColor ) );
+		OutDeviceColor = LinearToST2084( OutputGamutColor );
 	}
 	// Linear HDR, including all color correction, but no tone curve
 	else if( GetOutputDevice() == TONEMAPPER_OUTPUT_NoToneCurve)
 	{
-			OutDeviceColor = GradedColor;
+		OutDeviceColor = GradedColor;
 	}
 	// "Linear" including all color correction and the tone curve, but no device gamma
 	else if (GetOutputDevice() == TONEMAPPER_OUTPUT_WithToneCurve)
diff --git a/Engine/Shaders/Private/PostProcessDeviceEncodingOnly.usf b/Engine/Shaders/Private/PostProcessDeviceEncodingOnly.usf
--- a/Engine/Shaders/Private/PostProcessDeviceEncodingOnly.usf
+++ b/Engine/Shaders/Private/PostProcessDeviceEncodingOnly.usf
@@ -79,9 +79,9 @@ float4 DeviceEncodingOnlyCommonPS(
 
 	// sRGB, user specified gamut
 	if( GetOutputDevice() == TONEMAPPER_OUTPUT_sRGB)
-	{		
+	{
 		// Convert from sRGB to specified output gamut	
-		// float3 OutputGamutColor = mul( AP1_2_Output, mul( sRGB_2_AP1, FilmColor ) );
+		//float3 OutputGamutColor = mul( AP1_2_Output, mul( sRGB_2_AP1, FilmColor ) );
 
 		// FIXME: Workaround for UE-29935, pushing all colors with a 0 component to black output
 		// Default parameters seem to cancel out (sRGB->XYZ->AP1->XYZ->sRGB), so should be okay for a temp fix
@@ -94,7 +94,7 @@ float4 DeviceEncodingOnlyCommonPS(
 	// Rec 709, user specified gamut
 	else if( GetOutputDevice() == TONEMAPPER_OUTPUT_Rec709)
 	{
-		// Convert from sRGB to specified output gamut
+		// Convert from the working color space to specified output gamut
 		float3 OutputGamutColor = mul( AP1_2_Output, mul( (float3x3)WorkingColorSpace.ToAP1, FilmColor ) );
 
 		// Didn't profile yet if the branching version would be faster (different linear segment).
@@ -102,35 +102,51 @@ float4 DeviceEncodingOnlyCommonPS(
 	}
 
 	// ACES transform with PQ/2084 encoding, user specified gamut 
-	else if (GetOutputDevice() == TONEMAPPER_OUTPUT_ACES1000nitST2084 || GetOutputDevice() == TONEMAPPER_OUTPUT_ACES2000nitST2084)
+	else if( GetOutputDevice() == TONEMAPPER_OUTPUT_ACES1000nitST2084 || GetOutputDevice() == TONEMAPPER_OUTPUT_ACES2000nitST2084)
 	{
 #if USE_ACES_2
-		float3 ODTColor = ACESOutputTransform( FilmColor * ACESSceneColorMultiplier, (float3x3) WorkingColorSpace.ToAP0, OutputMaxLuminance, ACESReachTable, ACESGamutTable, ACESGammaTable);
+		// TODO: The display limiting color space should be user-exposed as part of display characterization.
+		float3 ODTColor = ACESOutputTransform(
+			FilmColor * ACESSceneColorMultiplier,
+			(float3x3)WorkingColorSpace.ToAP0,
+			OutputMaxLuminance,
+			ACESReachTable,
+			ACESGamutTable,
+			ACESGammaTable,
+			false, // bScaleWhite
+			GetColorSpace(TONEMAPPER_GAMUT_Rec2020_D65),
+			GetColorSpace(TONEMAPPER_GAMUT_Rec2020_D65));
 #else
-		// ACES ODT
 		FACESTonemapParams AcesParams = ComputeACESTonemapParams(ACESMinMaxData, ACESMidData, ACESCoefsLow_0, ACESCoefsHigh_0, ACESCoefsLow_4, ACESCoefsHigh_4, ACESSceneColorMultiplier, ACESGamutCompression);
 		float3 ODTColor = ACESOutputTransform( FilmColor, (float3x3)WorkingColorSpace.ToAP0, AcesParams);
+		ODTColor = mul(AP1_2_Output, ODTColor);
 #endif
-		// Convert from AP1 to specified output gamut
-		ODTColor = mul( AP1_2_Output, ODTColor );
 
 		// Apply conversion to ST-2084 (Dolby PQ)
 		OutDeviceColor = LinearToST2084( ODTColor );
 	}
-	 
+
 	// ACES transform to linear ScRGB
-	else if (GetOutputDevice() == TONEMAPPER_OUTPUT_ACES1000nitScRGB || GetOutputDevice() == TONEMAPPER_OUTPUT_ACES2000nitScRGB)
+	else if( GetOutputDevice() == TONEMAPPER_OUTPUT_ACES1000nitScRGB || GetOutputDevice() == TONEMAPPER_OUTPUT_ACES2000nitScRGB)
 	{
 #if USE_ACES_2
-		float3 ODTColor = ACESOutputTransform(FilmColor * ACESSceneColorMultiplier, (float3x3) WorkingColorSpace.ToAP0, OutputMaxLuminance, ACESReachTable, ACESGamutTable, ACESGammaTable);
+		// TODO: The display limiting color space should be user-exposed as part of display characterization.
+		float3 ODTColor = ACESOutputTransform(
+			FilmColor * ACESSceneColorMultiplier,
+			(float3x3)WorkingColorSpace.ToAP0,
+			OutputMaxLuminance,
+			ACESReachTable,
+			ACESGamutTable,
+			ACESGammaTable,
+			false, // bScaleWhite
+			GetColorSpace(TONEMAPPER_GAMUT_Rec2020_D65),
+			GetColorSpace(TONEMAPPER_GAMUT_sRGB_D65));
+		OutDeviceColor = ODTColor / UE_SCRGB_WHITE_NITS;
 #else
 		FACESTonemapParams AcesParams = ComputeACESTonemapParams(ACESMinMaxData, ACESMidData, ACESCoefsLow_0, ACESCoefsHigh_0, ACESCoefsLow_4, ACESCoefsHigh_4, ACESSceneColorMultiplier, ACESGamutCompression);
-
-		// ACES ODT
 		float3 ODTColor = ACESOutputTransform(FilmColor, (float3x3) WorkingColorSpace.ToAP0, AcesParams);
+		OutDeviceColor = mul(AP1_2_sRGB, ODTColor / UE_SCRGB_WHITE_NITS);
 #endif
-		// Apply conversion to ScRGB
-		OutDeviceColor = mul( AP1_2_sRGB, ODTColor / UE_SCRGB_WHITE_NITS );
 	}
 	
 	else if( GetOutputDevice() == TONEMAPPER_OUTPUT_LinearEXR)
@@ -155,7 +171,7 @@ float4 DeviceEncodingOnlyCommonPS(
 	// Gamma 2.2, user specified gamut
 	else
 	{
-		// Convert from sRGB to specified output gamut
+		// Convert from the working color space to specified output gamut
 		float3 OutputGamutColor = mul( AP1_2_Output, mul( (float3x3)WorkingColorSpace.ToAP1, FilmColor ) );
 
 		// This is different than the prior "gamma" curve adjustment (but reusing the variable).
diff --git a/Engine/Shaders/Private/PostProcessTonemap.usf b/Engine/Shaders/Private/PostProcessTonemap.usf
--- a/Engine/Shaders/Private/PostProcessTonemap.usf
+++ b/Engine/Shaders/Private/PostProcessTonemap.usf
@@ -132,16 +132,6 @@ float2 ConvertLensViewportSpaceToScreenViewportSpace(float2 UV)
 	return LensPrincipalPointOffsetScaleInverse.xy + UV * LensPrincipalPointOffsetScaleInverse.zw;
 }
 
-half3 LinearToPostTonemapSpace(half3 lin) 
-{
-#if IOS
-	// Note, iOS native output is raw gamma 2.2 not sRGB!
-	return pow(lin, 1.0 / 2.2);
-#else
-	return LinearToSrgbBranchless(lin);
-#endif
-}
-
 
 // Nuke-style Color Correct
 float ComputeFilmGrainIntensity(float3 LinearColorRGBGamut)
diff --git a/Engine/Shaders/Private/TonemapCommon.ush b/Engine/Shaders/Private/TonemapCommon.ush
--- a/Engine/Shaders/Private/TonemapCommon.ush
+++ b/Engine/Shaders/Private/TonemapCommon.ush
@@ -316,8 +316,6 @@ static const int UE_ACES_SURROUND = 0;      // 0: dark ( NOTE: this is the only
 static const bool UE_ACES_STRETCH_BLACK = true; // stretch black luminance to a PQ code value of 0
 static const bool UE_ACES_D60_SIM = false;
 static const bool UE_ACES_LEGAL_RANGE = false;
-// Scale white is usually only enabled when using a limiting white different from the encoding white
-static const bool UE_ACES_SCALE_WHITE = false;
 
 static const float UE_SCRGB_WHITE_NITS = 80.0; // ScRGB D65 white luminance (cd/m^2)
 
@@ -327,17 +325,15 @@ float3 ACESOutputTransform(
 	const float3x3 WCS_2_AP0,
 	float PeakLuminance,
 	Texture2D<float> ReachTable,
-    Texture2D<float3> GamutTable,
-    Texture2D<float> GammaTable,
-	const bool bApplyWhiteLimiting = true
+	Texture2D<float3> GamutTable,
+	Texture2D<float> GammaTable,
+	bool bScaleWhite,
+	FColorSpace LimitingColorSpace,
+	FColorSpace EncodingColorSpace
 )
 {
-	// TODO: The display limiting color space should be user-exposed as part of display characterization.
-	const FColorSpace LimitingColorSpace = GetColorSpace(TONEMAPPER_GAMUT_ACEScg_D60);
-	const FColorSpace EncodingColorSpace = GetColorSpace(TONEMAPPER_GAMUT_ACEScg_D60);
-
 	float3 Aces = mul( WCS_2_AP0, SceneReferredLinearColor );
-	
+
 	ODTParams PARAMS = init_ODTParams(
 		PeakLuminance,
 		LimitingColorSpace.RGBtoXYZ,
@@ -349,32 +345,34 @@ float3 ACESOutputTransform(
 		GammaTable
 	);
 
-	float3 ColorXYZ = outputTransform_fwd(Aces, PeakLuminance, PARAMS);
+	float3 ColorXYZ = outputTransform_fwd(Aces, PARAMS);
 
-	if(bApplyWhiteLimiting)
+	// Clamp to peak luminance
+	ColorXYZ = clamp_zero_to_peakLuminance(ColorXYZ, PARAMS);
+
+	if (bScaleWhite)
 	{
-		ColorXYZ = white_limiting(ColorXYZ, PARAMS, UE_ACES_SCALE_WHITE);
+		ColorXYZ = apply_white_scale(ColorXYZ, PARAMS, false);
 	}
 
 	// Scale by reference luminance to return values in nits, matching our 1.3 transform
-	return mul(XYZ_2_AP1_MAT, ColorXYZ) * referenceLuminance;
+	return mul(EncodingColorSpace.XYZtoRGB, ColorXYZ) * referenceLuminance;
 }
 
 float3 ACESInvOutputTransform(
-	float3 ColorAP1_Nits,
+	float3 EncodingNits,
 	const float3x3 AP0_2_WCS,
 	float PeakLuminance,
 	Texture2D<float> ReachTable,
-    Texture2D<float3> GamutTable,
-    Texture2D<float> GammaTable,
-	const bool bApplyWhiteLimiting = true
+	Texture2D<float3> GamutTable,
+	Texture2D<float> GammaTable,
+	const bool bScaleWhite,
+	FColorSpace LimitingColorSpace,
+	FColorSpace EncodingColorSpace
 )
 {
-	float3 ColorXYZ = mul(AP1_2_XYZ_MAT, ColorAP1_Nits / referenceLuminance);
+	float3 ColorXYZ = mul(EncodingColorSpace.RGBtoXYZ, EncodingNits / referenceLuminance);
 
-	const FColorSpace LimitingColorSpace = GetColorSpace(TONEMAPPER_GAMUT_ACEScg_D60);
-	const FColorSpace EncodingColorSpace = GetColorSpace(TONEMAPPER_GAMUT_ACEScg_D60);
-	
 	ODTParams PARAMS = init_ODTParams(
 		PeakLuminance,
 		LimitingColorSpace.RGBtoXYZ,
@@ -386,12 +384,15 @@ float3 ACESInvOutputTransform(
 		GammaTable
 	);
 
-	if(bApplyWhiteLimiting)
+	if (bScaleWhite)
 	{
-		ColorXYZ = white_limiting(ColorXYZ, PARAMS, UE_ACES_SCALE_WHITE, true);
+		ColorXYZ = apply_white_scale(ColorXYZ, PARAMS, true);
 	}
 
-    float3 Aces = outputTransform_inv(ColorXYZ, PeakLuminance, PARAMS);
+	// Clamp to peak luminance
+	ColorXYZ = clamp_zero_to_peakLuminance(ColorXYZ, PARAMS);
+
+	float3 Aces = outputTransform_inv(ColorXYZ, PARAMS);
 
 	return mul(AP0_2_WCS, Aces);
 }
diff --git a/Engine/Source/Runtime/Renderer/Private/PostProcess/ACESUtils.cpp b/Engine/Source/Runtime/Renderer/Private/PostProcess/ACESUtils.cpp
--- a/Engine/Source/Runtime/Renderer/Private/PostProcess/ACESUtils.cpp
+++ b/Engine/Source/Runtime/Renderer/Private/PostProcess/ACESUtils.cpp
@@ -2,7 +2,6 @@
 
 #include "ACESUtils.h"
 
-#include "ColorManagement/ColorSpace.h"
 #include "GlobalRenderResources.h"
 #include "RenderGraphBuilder.h"
 #include "RenderResource.h"
@@ -850,7 +849,7 @@ public:
 	}
 };
 
-void GetTransformResources(FRDGBuilder& GraphBuilder, float InPeakLuminance, FRHIShaderResourceView*& OutReachMTable, FRHIShaderResourceView*& OutGamutCuspTable, FRHIShaderResourceView*& OutUpperHullGammaTable)
+void GetTransformResources(FRDGBuilder& GraphBuilder, float InPeakLuminance, const UE::Color::FColorSpace& LimitingColorSpace, FRHIShaderResourceView*& OutReachMTable, FRHIShaderResourceView*& OutGamutCuspTable, FRHIShaderResourceView*& OutUpperHullGammaTable)
 {
 	static const auto CVarAcesVersion = IConsoleManager::Get().FindTConsoleVariableDataInt(TEXT("r.HDR.Aces.Version"));
 	
@@ -863,20 +862,20 @@ void GetTransformResources(FRDGBuilder& GraphBuilder, float InPeakLuminance, FRH
 		static FTextureLookupBase* UpperHullGammaTable = new TGlobalResource<FUpperHullGammaTable>;
 
 		static float CachedPeakLuminance = 0.0f;
+		static UE::Color::FColorSpace CachedLimitingColorSpace;
 
-		if (!FMath::IsNearlyEqual(CachedPeakLuminance, InPeakLuminance))
+		if (!FMath::IsNearlyEqual(CachedPeakLuminance, InPeakLuminance) ||
+			(CachedLimitingColorSpace != LimitingColorSpace))
 		{
-			// TODO: The display limiting color space should be user-exposed as part of display characterization.
-			static UE::Color::FColorSpace LimitingColorSpace = UE::Color::FColorSpace(UE::Color::EColorSpace::ACESAP1);
-
 			FRHICommandListImmediate& RHICmdListImmediate = GraphBuilder.RHICmdList.GetAsImmediate();
 
 			ReachMTable->Update(RHICmdListImmediate, LimitingColorSpace, InPeakLuminance);
 			GamutCuspTable->Update(RHICmdListImmediate, LimitingColorSpace, InPeakLuminance);
 			UpperHullGammaTable->Update(RHICmdListImmediate, LimitingColorSpace, InPeakLuminance);
 
-			// Update static cached peak luminance.
+			// Update static cached info.
 			CachedPeakLuminance = InPeakLuminance;
+			CachedLimitingColorSpace = LimitingColorSpace;
 		}
 
 		OutReachMTable = ReachMTable->ShaderResourceViewRHI;
diff --git a/Engine/Source/Runtime/Renderer/Private/PostProcess/ACESUtils.h b/Engine/Source/Runtime/Renderer/Private/PostProcess/ACESUtils.h
--- a/Engine/Source/Runtime/Renderer/Private/PostProcess/ACESUtils.h
+++ b/Engine/Source/Runtime/Renderer/Private/PostProcess/ACESUtils.h
@@ -3,6 +3,7 @@
 #pragma once
 
 #include "RHIFwd.h"
+#include "ColorManagement/ColorSpace.h"
 
 class FRDGBuilder;
 
@@ -13,6 +14,7 @@ namespace UE::Color::ACES
 	* 
 	* @param GraphBuilder			Frame render dependency graph.
 	* @param InPeakLuminance		Peak luminance value in nits (cd/m^2).
+	* @param LimitingColorSpace		Color space of the target display.
 	* @param OutReachMTable			Reach M value table resource output.
 	* @param OutGamutCuspTable		Gamut cusp table resource output.
 	* @param OutUpperHullGammaTable	Upper hull gamma table resource output.
@@ -20,6 +22,7 @@ namespace UE::Color::ACES
 	void GetTransformResources(
 		FRDGBuilder& GraphBuilder,
 		float InPeakLuminance,
+		const UE::Color::FColorSpace& LimitingColorSpace,
 		FRHIShaderResourceView*& OutReachMTable,
 		FRHIShaderResourceView*& OutGamutCuspTable,
 		FRHIShaderResourceView*& OutUpperHullGammaTable
diff --git a/Engine/Source/Runtime/Renderer/Private/PostProcess/PostProcessCombineLUTs.cpp b/Engine/Source/Runtime/Renderer/Private/PostProcess/PostProcessCombineLUTs.cpp
--- a/Engine/Source/Runtime/Renderer/Private/PostProcess/PostProcessCombineLUTs.cpp
+++ b/Engine/Source/Runtime/Renderer/Private/PostProcess/PostProcessCombineLUTs.cpp
@@ -183,6 +183,7 @@ struct FCachedLUTSettings
 	FCombineLUTParameters Parameters;
 	FWorkingColorSpaceShaderParameters WorkingColorSpaceShaderParameters;
 	bool bUseCompute = false;
+	int32 AcesVersion = 0;
 
 	bool UpdateCachedValues(FRDGBuilder& GraphBuilder, const FViewInfo& View, const FTexture* const* Textures, const float* Weights, uint32 BlendCount)
 	{
@@ -191,6 +192,8 @@ struct FCachedLUTSettings
 		UPDATE_CACHE_SETTINGS(UniqueID, View.State ? View.State->GetViewKey() : 0, bHasChanged);
 		UPDATE_CACHE_SETTINGS(ShaderPlatform, View.GetShaderPlatform(), bHasChanged);
 		UPDATE_CACHE_SETTINGS(bUseCompute, View.bUseComputePasses, bHasChanged);
+		static const auto CVarAcesVersion = IConsoleManager::Get().FindTConsoleVariableDataInt(TEXT("r.HDR.Aces.Version"));
+		UPDATE_CACHE_SETTINGS(AcesVersion, CVarAcesVersion->GetValueOnRenderThread(), bHasChanged);
 
 		const FWorkingColorSpaceShaderParameters* InWorkingColorSpaceShaderParameters = reinterpret_cast<const FWorkingColorSpaceShaderParameters*>(GDefaultWorkingColorSpaceUniformBuffer.GetContents());
 		if (InWorkingColorSpaceShaderParameters)
@@ -240,6 +243,17 @@ struct FCachedLUTSettings
 
 		Parameters.WorkingColorSpace = GDefaultWorkingColorSpaceUniformBuffer.GetUniformBufferRef();
 
+		FTonemapperOutputDeviceParameters TonemapperOutputDeviceParameters = GetTonemapperOutputDeviceParameters(ViewFamily);
+		UPDATE_CACHE_SETTINGS(Parameters.OutputDevice.InverseGamma, TonemapperOutputDeviceParameters.InverseGamma, bHasChanged);
+		UPDATE_CACHE_SETTINGS(Parameters.OutputDevice.OutputDevice, TonemapperOutputDeviceParameters.OutputDevice, bHasChanged);
+		UPDATE_CACHE_SETTINGS(Parameters.OutputDevice.OutputGamut, TonemapperOutputDeviceParameters.OutputGamut, bHasChanged);
+		UPDATE_CACHE_SETTINGS(Parameters.OutputDevice.OutputMaxLuminance, TonemapperOutputDeviceParameters.OutputMaxLuminance, bHasChanged);
+
+		// TODO: Plumb color space of the display.
+		static const UE::Color::FColorSpace LimitingColorSpaceHDR(UE::Color::EColorSpace::Rec2020);
+		static const UE::Color::FColorSpace LimitingColorSpaceSDR(UE::Color::EColorSpace::sRGB);
+		const UE::Color::FColorSpace& LimitingColorSpace = (TonemapperOutputDeviceParameters.OutputDevice == static_cast<uint32>(EDisplayOutputFormat::SDR_sRGB)) ? LimitingColorSpaceSDR : LimitingColorSpaceHDR;
+
 		FACESTonemapParams TonemapperParams;
 		GetACESTonemapParameters(TonemapperParams);
 		UPDATE_CACHE_SETTINGS(Parameters.ACESTonemapParameters.ACESMinMaxData, TonemapperParams.ACESMinMaxData, bHasChanged);
@@ -252,7 +266,8 @@ struct FCachedLUTSettings
 		UPDATE_CACHE_SETTINGS(Parameters.ACESTonemapParameters.ACESGamutCompression, TonemapperParams.ACESGamutCompression, bHasChanged);
 		UE::Color::ACES::GetTransformResources(
 			GraphBuilder,
-			HDRGetDisplayMaximumLuminance(),
+			TonemapperOutputDeviceParameters.OutputMaxLuminance,
+			LimitingColorSpace,
 			Parameters.ACESTonemapParameters.ACESReachTable,
 			Parameters.ACESTonemapParameters.ACESGamutTable,
 			Parameters.ACESTonemapParameters.ACESGammaTable
@@ -306,12 +321,6 @@ struct FCachedLUTSettings
 		UPDATE_CACHE_SETTINGS(Parameters.FilmShoulder, Settings.FilmShoulder, bHasChanged);
 		UPDATE_CACHE_SETTINGS(Parameters.FilmBlackClip, Settings.FilmBlackClip, bHasChanged);
 		UPDATE_CACHE_SETTINGS(Parameters.FilmWhiteClip, Settings.FilmWhiteClip, bHasChanged);
-
-		FTonemapperOutputDeviceParameters TonemapperOutputDeviceParameters = GetTonemapperOutputDeviceParameters(ViewFamily);
-		UPDATE_CACHE_SETTINGS(Parameters.OutputDevice.InverseGamma, TonemapperOutputDeviceParameters.InverseGamma, bHasChanged);
-		UPDATE_CACHE_SETTINGS(Parameters.OutputDevice.OutputDevice, TonemapperOutputDeviceParameters.OutputDevice, bHasChanged);
-		UPDATE_CACHE_SETTINGS(Parameters.OutputDevice.OutputGamut, TonemapperOutputDeviceParameters.OutputGamut, bHasChanged);
-		UPDATE_CACHE_SETTINGS(Parameters.OutputDevice.OutputMaxLuminance, TonemapperOutputDeviceParameters.OutputMaxLuminance, bHasChanged);
 	}
 };
 
@@ -540,9 +549,7 @@ FRDGTextureRef AddCombineLUTPass(FRDGBuilder& GraphBuilder, const FViewInfo& Vie
 
 	FLUTBlenderShader::FPermutationDomain PermutationVector;
 	PermutationVector.Set<FLUTBlenderShader::FBlendCount>(LocalCount);
-
-	static const auto CVarAcesVersion = IConsoleManager::Get().FindTConsoleVariableDataInt(TEXT("r.HDR.Aces.Version"));
-	PermutationVector.Set<FLUTBlenderShader::FUseACES2>(CVarAcesVersion->GetValueOnRenderThread() > 1);
+	PermutationVector.Set<FLUTBlenderShader::FUseACES2>(CachedLUTSettings.AcesVersion > 1);
 
 	if (bUseComputePass)
 	{
diff --git a/Engine/Source/Runtime/Renderer/Private/PostProcess/PostProcessDeviceEncodingOnly.cpp b/Engine/Source/Runtime/Renderer/Private/PostProcess/PostProcessDeviceEncodingOnly.cpp
--- a/Engine/Source/Runtime/Renderer/Private/PostProcess/PostProcessDeviceEncodingOnly.cpp
+++ b/Engine/Source/Runtime/Renderer/Private/PostProcess/PostProcessDeviceEncodingOnly.cpp
@@ -13,7 +13,6 @@
 #include "PostProcess/SceneFilterRendering.h"
 #include "PostProcess/PostProcessCombineLUTs.h"
 #include "PostProcess/PostProcessMobile.h"
-#include "PostProcess/PostProcessing.h"
 #include "ClearQuad.h"
 #include "PipelineStateCache.h"
 #include "HDRHelper.h"
@@ -78,7 +77,26 @@ FDeviceEncodingOnlyOutputDeviceParameters GetDeviceEncodingOnlyOutputDeviceParam
 	Parameters.InverseGamma = (FVector3f)InvDisplayGammaValue;
 	Parameters.OutputDevice = static_cast<uint32>(OutputDeviceValue);
 	Parameters.OutputGamut = static_cast<uint32>(Family.RenderTarget->GetDisplayColorGamut());
-	Parameters.OutputMaxLuminance = HDRGetDisplayMaximumLuminance();
+	switch (OutputDeviceValue)
+	{
+	case EDisplayOutputFormat::SDR_sRGB:
+	case EDisplayOutputFormat::SDR_Rec709:
+	case EDisplayOutputFormat::SDR_ExplicitGammaMapping:
+		Parameters.OutputMaxLuminance = 100.f;
+		break;
+	case EDisplayOutputFormat::HDR_ACES_1000nit_ST2084:
+	case EDisplayOutputFormat::HDR_ACES_2000nit_ST2084:
+	case EDisplayOutputFormat::HDR_ACES_1000nit_ScRGB:
+	case EDisplayOutputFormat::HDR_ACES_2000nit_ScRGB:
+	case EDisplayOutputFormat::HDR_LinearEXR:
+	case EDisplayOutputFormat::HDR_LinearNoToneCurve:
+	case EDisplayOutputFormat::HDR_LinearWithToneCurve:
+		// TODO: Plumb luminance through FSceneViewFamily or related.
+		Parameters.OutputMaxLuminance = HDRGetDisplayMaximumLuminance();
+		break;
+	default:
+		check(false);
+	}
 	return Parameters;
 }
 
@@ -223,10 +241,17 @@ FScreenPassTexture AddDeviceEncodingOnlyPass(FRDGBuilder& GraphBuilder, const FV
 	CommonParameters.OutputDevice = GetDeviceEncodingOnlyOutputDeviceParameters(ViewFamily);
 	CommonParameters.Color = GetScreenPassTextureViewportParameters(SceneColorViewport);
 	CommonParameters.Output = GetScreenPassTextureViewportParameters(OutputViewport);
+
+	// TODO: Plumb color space of the display.
+	static const UE::Color::FColorSpace LimitingColorSpaceHDR(UE::Color::EColorSpace::Rec2020);
+	static const UE::Color::FColorSpace LimitingColorSpaceSDR(UE::Color::EColorSpace::sRGB);
+	const UE::Color::FColorSpace& LimitingColorSpace = (CommonParameters.OutputDevice.OutputDevice == static_cast<uint32>(EDisplayOutputFormat::SDR_sRGB)) ? LimitingColorSpaceSDR : LimitingColorSpaceHDR;
+
 	GetACESTonemapParameters(CommonParameters.ACESTonemapParameters);
 	UE::Color::ACES::GetTransformResources(
 		GraphBuilder,
-		HDRGetDisplayMaximumLuminance(),
+		CommonParameters.OutputDevice.OutputMaxLuminance,
+		LimitingColorSpace,
 		CommonParameters.ACESTonemapParameters.ACESReachTable,
 		CommonParameters.ACESTonemapParameters.ACESGamutTable,
 		CommonParameters.ACESTonemapParameters.ACESGammaTable
diff --git a/Engine/Source/Runtime/Renderer/Private/PostProcess/PostProcessTonemap.cpp b/Engine/Source/Runtime/Renderer/Private/PostProcess/PostProcessTonemap.cpp
--- a/Engine/Source/Runtime/Renderer/Private/PostProcess/PostProcessTonemap.cpp
+++ b/Engine/Source/Runtime/Renderer/Private/PostProcess/PostProcessTonemap.cpp
@@ -266,7 +266,26 @@ FTonemapperOutputDeviceParameters GetTonemapperOutputDeviceParameters(const FSce
 	Parameters.InverseGamma = InvDisplayGammaValue;
 	Parameters.OutputDevice = static_cast<uint32>(OutputDeviceValue);
 	Parameters.OutputGamut = static_cast<uint32>(Family.RenderTarget->GetDisplayColorGamut());
-	Parameters.OutputMaxLuminance = HDRGetDisplayMaximumLuminance();
+	switch (OutputDeviceValue)
+	{
+	case EDisplayOutputFormat::SDR_sRGB:
+	case EDisplayOutputFormat::SDR_Rec709:
+	case EDisplayOutputFormat::SDR_ExplicitGammaMapping:
+		Parameters.OutputMaxLuminance = 100.f;
+		break;
+	case EDisplayOutputFormat::HDR_ACES_1000nit_ST2084:
+	case EDisplayOutputFormat::HDR_ACES_2000nit_ST2084:
+	case EDisplayOutputFormat::HDR_ACES_1000nit_ScRGB:
+	case EDisplayOutputFormat::HDR_ACES_2000nit_ScRGB:
+	case EDisplayOutputFormat::HDR_LinearEXR:
+	case EDisplayOutputFormat::HDR_LinearNoToneCurve:
+	case EDisplayOutputFormat::HDR_LinearWithToneCurve:
+		// TODO: Plumb luminance through FSceneViewFamily or related.
+		Parameters.OutputMaxLuminance = HDRGetDisplayMaximumLuminance();
+		break;
+	default:
+		check(false);
+	}
 	return Parameters;
 }
 
