Fix incorrect texture type in declaration in FDrawShadowCastersFalseColor

#rnx

diff --git a/Engine/Source/Runtime/Renderer/Private/VirtualShadowMaps/VirtualShadowMapArray.cpp b/Engine/Source/Runtime/Renderer/Private/VirtualShadowMaps/VirtualShadowMapArray.cpp
--- a/Engine/Source/Runtime/Renderer/Private/VirtualShadowMaps/VirtualShadowMapArray.cpp
+++ b/Engine/Source/Runtime/Renderer/Private/VirtualShadowMaps/VirtualShadowMapArray.cpp
@@ -4710,7 +4710,7 @@ class FDrawShadowCastersFalseColor : public FGlobalShader
 
 	BEGIN_SHADER_PARAMETER_STRUCT(FParameters, )
 		SHADER_PARAMETER_RDG_TEXTURE(Texture2D<uint>, DebugAuxTexture)
-		SHADER_PARAMETER_RDG_TEXTURE(Texture2D<uint>, SceneColorTexture)
+		SHADER_PARAMETER_RDG_TEXTURE(Texture2D, SceneColorTexture)
 		SHADER_PARAMETER_STRUCT(FScreenPassTextureViewportParameters, DebugAuxTextureViewport)
 		SHADER_PARAMETER_STRUCT(FScreenPassTextureViewportParameters, SceneColorTextureViewport)
 		SHADER_PARAMETER(FScreenTransform, ViewportUVToDebugAuxTextureUV)
