Expose the GPU scene's instance payload buffer so that FRendererModule::CreateSinglePrimitiveSceneUniformBuffer can set up primitives that require additional data per instance.

diff --git a/Engine/Source/Runtime/RenderCore/Public/RendererInterface.h b/Engine/Source/Runtime/RenderCore/Public/RendererInterface.h
--- a/Engine/Source/Runtime/RenderCore/Public/RendererInterface.h
+++ b/Engine/Source/Runtime/RenderCore/Public/RendererInterface.h
@@ -44,6 +44,7 @@ class FBatchedPrimitiveParameters;
 class ISceneRenderer;
 class IStateStreamManager;
 class ISceneRenderBuilder;
+class FPrimitiveSceneProxy;
 
 namespace Nanite
 {
@@ -728,6 +729,8 @@ public:
 	virtual void InitializeSystemTextures(FRHICommandListImmediate& RHICmdList) = 0;
 
 	/** Create a Scene Uniform Buffer containing only the scene representation for a single primitive */
+	UE_EXPERIMENTAL(5.8, "Potentially replaced with scene descriptors using IPrimitiveComponent")
+	virtual FSceneUniformBuffer* CreateSinglePrimitiveSceneUniformBuffer(FRDGBuilder& GraphBuilder, const FPrimitiveSceneProxy* PrimitiveSceneProxy, ERHIFeatureLevel::Type FeatureLevel, FMeshBatch& Mesh) = 0;
 	virtual FSceneUniformBuffer* CreateSinglePrimitiveSceneUniformBuffer(FRDGBuilder& GraphBuilder, ERHIFeatureLevel::Type FeatureLevel, FMeshBatch& Mesh) = 0;
 	virtual FSceneUniformBuffer* CreateSinglePrimitiveSceneUniformBuffer(FRDGBuilder& GraphBuilder, const FViewInfo& SceneView, FMeshBatch& Mesh) = 0;
 	/** Create a Uniform Buffer containing representation for a single primitive. (For platforms that use "UniformView" path) */
diff --git a/Engine/Source/Runtime/Renderer/Private/GPUScene.cpp b/Engine/Source/Runtime/Renderer/Private/GPUScene.cpp
--- a/Engine/Source/Runtime/Renderer/Private/GPUScene.cpp
+++ b/Engine/Source/Runtime/Renderer/Private/GPUScene.cpp
@@ -193,6 +193,27 @@ static int32 PartitionUpdateRanges(FParallelUpdateRanges& Ranges, int32 ItemCoun
 	return Ranges.Range[3].ItemCount > 0 ? 4 : 3;
 }
 
+namespace UE::Renderer::Private
+{
+FORCEINLINE_GPUSCENE uint32 PackInstanceDataFlags(FInstanceDataFlags Flags)
+{
+	uint32 PackedFlags = 0x0;
+	PackedFlags |= Flags.bHasPerInstanceRandom ? INSTANCE_SCENE_DATA_FLAG_HAS_RANDOM : 0u;
+	PackedFlags |= Flags.bHasPerInstanceCustomData ? INSTANCE_SCENE_DATA_FLAG_HAS_CUSTOM_DATA : 0u;
+	PackedFlags |= Flags.bHasPerInstanceDynamicData ? INSTANCE_SCENE_DATA_FLAG_HAS_DYNAMIC_DATA : 0u;
+	PackedFlags |= Flags.bHasPerInstanceSkinningData ? INSTANCE_SCENE_DATA_FLAG_HAS_SKINNING_DATA : 0u;
+	PackedFlags |= Flags.bHasPerInstanceLMSMUVBias ? INSTANCE_SCENE_DATA_FLAG_HAS_LIGHTSHADOW_UV_BIAS : 0u;
+	PackedFlags |= Flags.bHasPerInstanceHierarchyOffset ? INSTANCE_SCENE_DATA_FLAG_HAS_HIERARCHY_OFFSET : 0u;
+	PackedFlags |= Flags.bHasPerInstanceLocalBounds ? INSTANCE_SCENE_DATA_FLAG_HAS_LOCAL_BOUNDS : 0u;
+	PackedFlags |= Flags.bHasPerInstancePayloadExtension ? INSTANCE_SCENE_DATA_FLAG_HAS_PAYLOAD_EXTENSION : 0u;
+#if WITH_EDITOR
+	PackedFlags |= Flags.bHasPerInstanceEditorData ? INSTANCE_SCENE_DATA_FLAG_HAS_EDITOR_DATA : 0u;
+#endif
+
+	return PackedFlags;
+}
+}
+
 void FGPUScenePrimitiveCollector::Add(
 	const FMeshBatchDynamicPrimitiveData* MeshBatchData,
 	const FPrimitiveUniformShaderParameters& PrimitiveShaderParams,
@@ -313,45 +334,6 @@ struct FPrimitiveUploadInfoHeader
 	const FPrimitiveSceneInfo* PrimitiveSceneInfo = nullptr;
 };
 
-/**
- * Info required by the uploader to update the instances that belong to a primitive.
- */
-struct FInstanceUploadInfo
-{
-	const FInstanceSceneDataBuffers *InstanceSceneDataBuffers = nullptr;
-	TConstArrayView<FInstanceSceneData> PrimitiveInstances;
-	int32 InstanceSceneDataOffset = INDEX_NONE;
-	int32 InstancePayloadDataOffset = INDEX_NONE;
-	int32 InstancePayloadDataStride = 0;
-	int32 InstancePayloadExtensionCount = 0;
-	int32 InstanceCustomDataCount = 0;
-
-	// Optional per-instance data views
-	TConstArrayView<FInstanceDynamicData> InstanceDynamicData;
-	TConstArrayView<uint32> InstanceSkinningData;
-	TConstArrayView<FVector4f> InstanceLightShadowUVBias;
-	TConstArrayView<float> InstanceCustomData;
-	TConstArrayView<float> InstanceRandomID;
-	TConstArrayView<uint32> InstanceHierarchyOffset;
-	TConstArrayView<FRenderBounds> InstanceLocalBounds;
-	TConstArrayView<FVector4f> InstancePayloadExtension;
-#if WITH_EDITOR
-	TConstArrayView<uint32> InstanceEditorData;
-#endif
-
-	uint32 InstanceFlags = 0x0;
-
-	FRenderTransform PrimitiveToWorld;
-	FRenderTransform PrevPrimitiveToWorld;
-	int32 PrimitiveID = INDEX_NONE;
-	uint32 LastUpdateSceneFrameNumber = ~uint32(0);
-	int32 NumInstances = 0;
-	bool bIsPrimitiveForceHidden = false;
-#if DO_CHECK
-	const FPrimitiveSceneInfo *PrimitiveSceneInfo = nullptr;
-#endif
-};
-
 #if DO_CHECK
 void ValidateInstanceUploadInfo(const FInstanceUploadInfo& UploadInfo, FRDGBuffer* InstancePayloadDataBuffer)
 {
@@ -489,24 +471,6 @@ struct FUploadDataSourceAdapterScenePrimitives
 		}
 	}
 
-	FORCEINLINE_GPUSCENE uint32 PackFlags(FInstanceDataFlags Flags) const
-	{
-		uint32 PackedFlags = 0x0;
-		PackedFlags |= Flags.bHasPerInstanceRandom				? INSTANCE_SCENE_DATA_FLAG_HAS_RANDOM              : 0u;
-		PackedFlags |= Flags.bHasPerInstanceCustomData			? INSTANCE_SCENE_DATA_FLAG_HAS_CUSTOM_DATA         : 0u;
-		PackedFlags |= Flags.bHasPerInstanceDynamicData			? INSTANCE_SCENE_DATA_FLAG_HAS_DYNAMIC_DATA        : 0u;
-		PackedFlags |= Flags.bHasPerInstanceSkinningData		? INSTANCE_SCENE_DATA_FLAG_HAS_SKINNING_DATA       : 0u;
-		PackedFlags |= Flags.bHasPerInstanceLMSMUVBias			? INSTANCE_SCENE_DATA_FLAG_HAS_LIGHTSHADOW_UV_BIAS : 0u;
-		PackedFlags |= Flags.bHasPerInstanceHierarchyOffset		? INSTANCE_SCENE_DATA_FLAG_HAS_HIERARCHY_OFFSET    : 0u;
-		PackedFlags |= Flags.bHasPerInstanceLocalBounds			? INSTANCE_SCENE_DATA_FLAG_HAS_LOCAL_BOUNDS        : 0u;
-		PackedFlags |= Flags.bHasPerInstancePayloadExtension	? INSTANCE_SCENE_DATA_FLAG_HAS_PAYLOAD_EXTENSION   : 0u;
-	#if WITH_EDITOR
-		PackedFlags |= Flags.bHasPerInstanceEditorData			? INSTANCE_SCENE_DATA_FLAG_HAS_EDITOR_DATA         : 0u;
-	#endif
-
-		return PackedFlags;
-	}
-
 	/**
 	 * Populate the primitive info for a given item index.
 	 * 
@@ -531,97 +495,7 @@ struct FUploadDataSourceAdapterScenePrimitives
 		check(PrimitiveDirtyState[PersistentPrimitiveIndex.Index] != EPrimitiveDirtyState::Removed);
 
 		FPrimitiveSceneProxy* PrimitiveSceneProxy = Scene.PrimitiveSceneProxies[PrimitiveID];
-		const FPrimitiveSceneInfo* PrimitiveSceneInfo = PrimitiveSceneProxy->GetPrimitiveSceneInfo();
-#if DO_CHECK
-		InstanceUploadInfo.PrimitiveSceneInfo = PrimitiveSceneInfo;
-#endif
-
-		InstanceUploadInfo.InstanceSceneDataOffset = PrimitiveSceneInfo->GetInstanceSceneDataOffset();
-		check(InstanceUploadInfo.InstanceSceneDataOffset >= 0);
-		InstanceUploadInfo.InstancePayloadDataOffset = PrimitiveSceneInfo->GetInstancePayloadDataOffset();
-		InstanceUploadInfo.InstancePayloadDataStride = PrimitiveSceneInfo->GetInstancePayloadDataStride();
-		InstanceUploadInfo.LastUpdateSceneFrameNumber = SceneFrameNumber;
-		InstanceUploadInfo.PrimitiveID = PersistentPrimitiveIndex.Index;
-		// HACK: ignoring IsForceHidden for non-Nanite due to issues that cropped up with water rendering.
-		// TODO: Remove the IsNaniteMesh() check
-		InstanceUploadInfo.bIsPrimitiveForceHidden = PrimitiveSceneProxy->IsNaniteMesh() && PrimitiveSceneProxy->IsForceHidden();
-
-		InstanceUploadInfo.InstanceSceneDataBuffers = PrimitiveSceneInfo->GetInstanceSceneDataBuffers();
-		if (InstanceUploadInfo.InstanceSceneDataBuffers)
-		{
-			const FInstanceSceneDataBuffers::FReadView InstanceSceneDataBuffers = InstanceUploadInfo.InstanceSceneDataBuffers->GetReadView();
-			InstanceUploadInfo.NumInstances = InstanceUploadInfo.InstanceSceneDataBuffers->GetNumInstances();
-			InstanceUploadInfo.InstanceFlags = PackFlags(InstanceSceneDataBuffers.Flags);
-			InstanceUploadInfo.InstanceLightShadowUVBias = InstanceSceneDataBuffers.InstanceLightShadowUVBias;
-			InstanceUploadInfo.InstanceCustomData = InstanceSceneDataBuffers.InstanceCustomData;
-			InstanceUploadInfo.InstanceRandomID = InstanceSceneDataBuffers.InstanceRandomIDs;
-			InstanceUploadInfo.InstanceSkinningData = InstanceSceneDataBuffers.InstanceSkinningData;
-			InstanceUploadInfo.InstanceHierarchyOffset = InstanceSceneDataBuffers.InstanceHierarchyOffset;
-			InstanceUploadInfo.InstancePayloadExtension = InstanceSceneDataBuffers.InstancePayloadExtension;
-			InstanceUploadInfo.InstanceLocalBounds = InstanceSceneDataBuffers.InstanceLocalBounds;
-		#if WITH_EDITOR
-			InstanceUploadInfo.InstanceEditorData = InstanceSceneDataBuffers.InstanceEditorData;
-		#endif
-		#if DO_CHECK
-			// This is already precomputed in the InstanceSceneDataBuffers and we don't need to do it again here, except for validation purposes
-			// TODO: this validation should (also?) move elsewhere and validate that the transform on RT matches that on GT
-			const FMatrix LocalToWorld = PrimitiveSceneProxy->GetLocalToWorld();
-			const FDFVector3 AbsoluteOrigin(LocalToWorld.GetOrigin());
-			InstanceUploadInfo.PrimitiveToWorld = FDFMatrix::MakeToRelativeWorldMatrix(AbsoluteOrigin.High, LocalToWorld).M;
-		#endif
-		}
-		else
-		{
-			// Old path, only taken for uninstanced primitives.
-			const FMatrix LocalToWorld = PrimitiveSceneProxy->GetLocalToWorld();
-			const FDFVector3 AbsoluteOrigin(LocalToWorld.GetOrigin());
-			InstanceUploadInfo.PrimitiveToWorld = FDFMatrix::MakeToRelativeWorldMatrix(AbsoluteOrigin.High, LocalToWorld).M;
-
-			InstanceUploadInfo.InstanceFlags = 0u;
-			check(InstanceUploadInfo.InstancePayloadDataOffset == INDEX_NONE && InstanceUploadInfo.InstancePayloadDataStride == 0);
-
-			// empty array signals that we should use the PrimitiveToWorld as the instance transform (and that it is therefore pre-transformed).
-			InstanceUploadInfo.PrimitiveInstances = TConstArrayView<FInstanceSceneData>();
-			InstanceUploadInfo.InstanceDynamicData = TConstArrayView<FInstanceDynamicData>();
-
-		#if 0
-			// NOTE: We only need this if not using the InstanceSceneDataBuffers and if the old path were to support dynamic data
-			{
-				bool bHasPrecomputedVolumetricLightmap{};
-				bool bOutputVelocity{};
-				int32 SingleCaptureIndex{};
-
-				FMatrix PreviousLocalToWorld;
-				Scene.GetPrimitiveUniformShaderParameters_RenderThread(PrimitiveSceneInfo, bHasPrecomputedVolumetricLightmap, PreviousLocalToWorld, SingleCaptureIndex, bOutputVelocity);
-				InstanceUploadInfo.PrevPrimitiveToWorld = FDFMatrix::MakeClampedToRelativeWorldMatrix(AbsoluteOrigin.High, PreviousLocalToWorld).M;
-			}
-		#endif
-			InstanceUploadInfo.InstanceLightShadowUVBias = TConstArrayView<FVector4f>();
-			InstanceUploadInfo.InstanceCustomData = TConstArrayView<float>();
-			InstanceUploadInfo.InstanceRandomID = TConstArrayView<float>();
-			InstanceUploadInfo.InstanceSkinningData = TConstArrayView<uint32>();
-			InstanceUploadInfo.InstanceHierarchyOffset = TConstArrayView<uint32>();
-			InstanceUploadInfo.InstancePayloadExtension = TConstArrayView<FVector4f>();
-			InstanceUploadInfo.NumInstances = 1;
-		#if WITH_EDITOR
-			InstanceUploadInfo.InstanceEditorData = TConstArrayView<uint32>();
-		#endif
-		}
-
-		InstanceUploadInfo.InstancePayloadExtensionCount = 0;
-		if (InstanceUploadInfo.InstancePayloadExtension.Num() > 0)
-		{
-			InstanceUploadInfo.InstancePayloadExtensionCount = InstanceUploadInfo.InstancePayloadExtension.Num() / InstanceUploadInfo.NumInstances;
-		}
-
-		InstanceUploadInfo.InstanceCustomDataCount = 0;
-		if (InstanceUploadInfo.InstanceCustomData.Num() > 0)
-		{
-			InstanceUploadInfo.InstanceCustomDataCount = InstanceUploadInfo.InstanceCustomData.Num() / InstanceUploadInfo.NumInstances;
-		}
-
-		// Only trigger upload if this primitive has instances
-		check(InstanceUploadInfo.NumInstances > 0);
+		InstanceUploadInfo.SetPrimitiveData(PrimitiveSceneProxy, PersistentPrimitiveIndex, SceneFrameNumber);
 	}
 
 	FORCEINLINE_GPUSCENE bool GetLightMapInfo(int32 ItemIndex, FLightMapUploadInfo &UploadInfo) const
@@ -1020,7 +894,7 @@ void FGPUScene::UpdateInternal(FRDGBuilder& GraphBuilder, FSceneUniformBuffer& S
 			const FInstanceSceneDataBuffers* InstanceSceneDataBuffers = Command.SceneInfo->GetInstanceSceneDataBuffers();
 			check(InstanceSceneDataBuffers);
 			Params.NumCustomDataFloats = InstanceSceneDataBuffers->GetNumCustomDataFloats();
-			Params.PackedInstanceSceneDataFlags = Adapter.PackFlags(InstanceSceneDataBuffers->GetFlags());
+			Params.PackedInstanceSceneDataFlags = UE::Renderer::Private::PackInstanceDataFlags(InstanceSceneDataBuffers->GetFlags());
 
 			Command.Payload.GPUSceneWriter.Execute(GraphBuilder, Params);
 		}
@@ -1470,98 +1344,7 @@ void FGPUScene::UploadGeneral(FRDGBuilder& GraphBuilder, const FRegisteredBuffer
 							FVector4f* DstPayloadData = static_cast<FVector4f*>(TaskContext.InstancePayloadUploader->Set_GetRef(UploadPayloadItemOffset, PayloadDataStart, UploadInfo.InstancePayloadDataStride));
 							TArrayView<FVector4f> InstancePayloadData = TArrayView<FVector4f>(DstPayloadData, UploadInfo.InstancePayloadDataStride);
 
-							int32 PayloadPosition = 0;
-
-							if (UploadInfo.InstanceFlags & (INSTANCE_SCENE_DATA_FLAG_HAS_HIERARCHY_OFFSET | INSTANCE_SCENE_DATA_FLAG_HAS_LOCAL_BOUNDS | INSTANCE_SCENE_DATA_FLAG_HAS_SKINNING_DATA))
-							{
-								const uint32 InstanceHierarchyOffset = (UploadInfo.InstanceFlags & INSTANCE_SCENE_DATA_FLAG_HAS_HIERARCHY_OFFSET) ? UploadInfo.InstanceHierarchyOffset[InstanceIndex] : 0;
-								InstancePayloadData[PayloadPosition].X = *(const float*)&InstanceHierarchyOffset;
-
-								const uint32 InstanceSkinningData = (UploadInfo.InstanceFlags & INSTANCE_SCENE_DATA_FLAG_HAS_SKINNING_DATA) ? UploadInfo.InstanceSkinningData[InstanceIndex] : 0;
-								InstancePayloadData[PayloadPosition].Y = *(const float*)&InstanceSkinningData;
-
-								if (UploadInfo.InstanceFlags & INSTANCE_SCENE_DATA_FLAG_HAS_LOCAL_BOUNDS)
-								{
-									check(UploadInfo.InstanceLocalBounds.Num() == UploadInfo.NumInstances);
-									const FRenderBounds& InstanceLocalBounds = UploadInfo.InstanceLocalBounds[InstanceIndex];
-									const FVector3f BoundsOrigin = InstanceLocalBounds.GetCenter();
-									const FVector3f BoundsExtent = InstanceLocalBounds.GetExtent();
-
-									InstancePayloadData[PayloadPosition + 0].Z = *(const float*)&BoundsOrigin.X;
-									InstancePayloadData[PayloadPosition + 0].W = *(const float*)&BoundsOrigin.Y;
-
-									InstancePayloadData[PayloadPosition + 1].X = *(const float*)&BoundsOrigin.Z;
-									InstancePayloadData[PayloadPosition + 1].Y = *(const float*)&BoundsExtent.X;
-									InstancePayloadData[PayloadPosition + 1].Z = *(const float*)&BoundsExtent.Y;
-									InstancePayloadData[PayloadPosition + 1].W = *(const float*)&BoundsExtent.Z;
-								}
-
-								PayloadPosition += (UploadInfo.InstanceFlags & INSTANCE_SCENE_DATA_FLAG_HAS_LOCAL_BOUNDS) ? 2 : 1;
-							}
-
-							if (UploadInfo.InstanceFlags & INSTANCE_SCENE_DATA_FLAG_HAS_DYNAMIC_DATA)
-							{
-								const FRenderTransform PrevLocalToWorld = UploadInfo.InstanceSceneDataBuffers != nullptr ?
-									UploadInfo.InstanceSceneDataBuffers->GetPrevInstanceToPrimitiveRelative(InstanceIndex) :
-									UploadInfo.InstanceDynamicData[InstanceIndex].ComputePrevLocalToWorld(UploadInfo.PrevPrimitiveToWorld);
-
-								if (bSupportsCompressedTransforms)
-								{
-									check(PayloadPosition + 1 < InstancePayloadData.Num()); // Sanity check
-									FCompressedTransform CompressedPrevLocalToWorld(PrevLocalToWorld);
-									InstancePayloadData[PayloadPosition + 0] = *(const FVector4f*)&CompressedPrevLocalToWorld.Rotation[0];
-									InstancePayloadData[PayloadPosition + 1] = *(const FVector3f*)&CompressedPrevLocalToWorld.Translation;
-									PayloadPosition += 2;
-								}
-								else
-								{
-									// Note: writes 3x float4s
-									check(PayloadPosition + 2 < InstancePayloadData.Num()); // Sanity check
-									PrevLocalToWorld.To3x4MatrixTranspose((float*)&InstancePayloadData[PayloadPosition]);
-									PayloadPosition += 3;
-								}
-							}
-
-						#if WITH_EDITOR
-							if (UploadInfo.InstanceFlags & INSTANCE_SCENE_DATA_FLAG_HAS_EDITOR_DATA)
-							{
-								check(UploadInfo.InstanceEditorData.Num() == UploadInfo.NumInstances);
-								InstancePayloadData[PayloadPosition].X = *(const float*)&UploadInfo.InstanceEditorData[InstanceIndex];
-								InstancePayloadData[PayloadPosition].Y = 0.0f; // Unused
-								InstancePayloadData[PayloadPosition].Z = 0.0f; // Unused
-								InstancePayloadData[PayloadPosition].W = 0.0f; // Unused
-								PayloadPosition += 1;
-							}
-						#endif
-
-							// TODO: Skip all this if static lighting is disabled
-							if (UploadInfo.InstanceFlags & INSTANCE_SCENE_DATA_FLAG_HAS_LIGHTSHADOW_UV_BIAS)
-							{
-								check(UploadInfo.InstanceLightShadowUVBias.Num() == UploadInfo.NumInstances);
-								InstancePayloadData[PayloadPosition] = UploadInfo.InstanceLightShadowUVBias[InstanceIndex];
-								PayloadPosition += 1;
-							}
-
-							if (UploadInfo.InstanceFlags & INSTANCE_SCENE_DATA_FLAG_HAS_PAYLOAD_EXTENSION)
-							{
-								check(UploadInfo.InstancePayloadExtension.Num() / UploadInfo.InstancePayloadExtensionCount == UploadInfo.NumInstances);
-								const int32 SrcOffset = InstanceIndex * UploadInfo.InstancePayloadExtensionCount;
-								for (int32 Idx = 0; Idx < UploadInfo.InstancePayloadExtensionCount; ++Idx)
-								{
-									InstancePayloadData[PayloadPosition++] = UploadInfo.InstancePayloadExtension[SrcOffset + Idx];
-								}
-							}
-
-							if (UploadInfo.InstanceFlags & INSTANCE_SCENE_DATA_FLAG_HAS_CUSTOM_DATA && UploadInfo.InstanceCustomDataCount > 0)
-							{
-								check(PayloadPosition + (UploadInfo.InstanceCustomDataCount >> 2u) <= InstancePayloadData.Num());
-								const float* CustomDataGetPtr = &UploadInfo.InstanceCustomData[(InstanceIndex * UploadInfo.InstanceCustomDataCount)];
-								float* CustomDataPutPtr = (float*)&InstancePayloadData[PayloadPosition];
-								for (uint32 FloatIndex = 0; FloatIndex < uint32(UploadInfo.InstanceCustomDataCount); ++FloatIndex)
-								{
-									CustomDataPutPtr[FloatIndex] = CustomDataGetPtr[FloatIndex];
-								}
-							}
+							UploadInfo.WritePayloadData(InstancePayloadData, InstanceIndex);
 						}
 					}
 				}
@@ -2696,3 +2479,250 @@ FGPUSceneWriteDelegateRef FGPUSceneWriteDelegate::CreateInternal(FMeshElementCol
 	return Result;
 }
 
+void FInstanceUploadInfo::SetPrimitiveData(FPrimitiveSceneProxy* InPrimitiveSceneProxy, const FPersistentPrimitiveIndex& PrimitiveIndex, uint32 SceneFrameNumber)
+{
+	const FPrimitiveSceneInfo* LocalPrimitiveSceneInfo = InPrimitiveSceneProxy->GetPrimitiveSceneInfo();
+#if DO_CHECK
+	PrimitiveSceneInfo = LocalPrimitiveSceneInfo;
+#endif
+
+	InstanceSceneDataOffset = LocalPrimitiveSceneInfo->GetInstanceSceneDataOffset();
+	check(InstanceSceneDataOffset >= 0);
+	InstancePayloadDataOffset = LocalPrimitiveSceneInfo->GetInstancePayloadDataOffset();
+	InstancePayloadDataStride = LocalPrimitiveSceneInfo->GetInstancePayloadDataStride();
+	LastUpdateSceneFrameNumber = SceneFrameNumber;
+	PrimitiveID = PrimitiveIndex.Index;
+	// HACK: ignoring IsForceHidden for non-Nanite due to issues that cropped up with water rendering.
+	// TODO: Remove the IsNaniteMesh() check
+	bIsPrimitiveForceHidden = InPrimitiveSceneProxy->IsNaniteMesh() && InPrimitiveSceneProxy->IsForceHidden();
+
+	const FInstanceSceneDataBuffers* InstanceSceneDataBuffersPtr = LocalPrimitiveSceneInfo->GetInstanceSceneDataBuffers();
+	SetInstancePayloadData(InstanceSceneDataBuffersPtr);
+
+	bool bSetPrimitiveToWorld = !!(DO_CHECK) || !InstanceSceneDataBuffers;
+
+	// If the primitive has isntancing data then this is already precomputed in the InstanceSceneDataBuffers and we don't need to do it again here, except for validation purposes
+	// If the primitive does not have instancing data, the primitive to world is required
+	// TODO: this validation should (also?) move elsewhere and validate that the transform on RT matches that on GT
+	if (bSetPrimitiveToWorld)
+	{
+		const FMatrix LocalToWorld = InPrimitiveSceneProxy->GetLocalToWorld();
+		const FDFVector3 AbsoluteOrigin(LocalToWorld.GetOrigin());
+		PrimitiveToWorld = FDFMatrix::MakeToRelativeWorldMatrix(AbsoluteOrigin.High, LocalToWorld).M;
+	}
+}
+
+void FInstanceUploadInfo::SetInstancePayloadData(const FInstanceSceneDataBuffers* InInstanceSceneDataBuffers)
+{
+	InstanceSceneDataBuffers = InInstanceSceneDataBuffers;
+	if (InstanceSceneDataBuffers)
+	{
+		const FInstanceSceneDataBuffers::FReadView ReadView = InstanceSceneDataBuffers->GetReadView();
+		NumInstances = InstanceSceneDataBuffers->GetNumInstances();
+		InstanceFlags = UE::Renderer::Private::PackInstanceDataFlags(ReadView.Flags);
+		InstanceLightShadowUVBias = ReadView.InstanceLightShadowUVBias;
+		InstanceCustomData = ReadView.InstanceCustomData;
+		InstanceRandomID = ReadView.InstanceRandomIDs;
+		InstanceSkinningData = ReadView.InstanceSkinningData;
+		InstanceHierarchyOffset = ReadView.InstanceHierarchyOffset;
+		InstancePayloadExtension = ReadView.InstancePayloadExtension;
+		InstanceLocalBounds = ReadView.InstanceLocalBounds;
+#if WITH_EDITOR
+		InstanceEditorData = ReadView.InstanceEditorData;
+#endif
+	}
+	else
+	{
+		InstanceFlags = 0u;
+		check(InstancePayloadDataOffset == INDEX_NONE && InstancePayloadDataStride == 0);
+
+		// empty array signals that we should use the PrimitiveToWorld as the instance transform (and that it is therefore pre-transformed).
+		PrimitiveInstances = TConstArrayView<FInstanceSceneData>();
+		InstanceDynamicData = TConstArrayView<FInstanceDynamicData>();
+
+#if 0
+		// NOTE: We only need this if not using the InstanceSceneDataBuffers and if the old path were to support dynamic data
+		{
+			bool bHasPrecomputedVolumetricLightmap{};
+			bool bOutputVelocity{};
+			int32 SingleCaptureIndex{};
+
+			FMatrix PreviousLocalToWorld;
+			Scene.GetPrimitiveUniformShaderParameters_RenderThread(PrimitiveSceneInfo, bHasPrecomputedVolumetricLightmap, PreviousLocalToWorld, SingleCaptureIndex, bOutputVelocity);
+			PrevPrimitiveToWorld = FDFMatrix::MakeClampedToRelativeWorldMatrix(AbsoluteOrigin.High, PreviousLocalToWorld).M;
+		}
+#endif
+		InstanceLightShadowUVBias = TConstArrayView<FVector4f>();
+		InstanceCustomData = TConstArrayView<float>();
+		InstanceRandomID = TConstArrayView<float>();
+		InstanceSkinningData = TConstArrayView<uint32>();
+		InstanceHierarchyOffset = TConstArrayView<uint32>();
+		InstancePayloadExtension = TConstArrayView<FVector4f>();
+		NumInstances = 1;
+#if WITH_EDITOR
+		InstanceEditorData = TConstArrayView<uint32>();
+#endif
+	}
+
+	InstancePayloadExtensionCount = 0;
+	if (InstancePayloadExtension.Num() > 0)
+	{
+		InstancePayloadExtensionCount = InstancePayloadExtension.Num() / NumInstances;
+	}
+
+	InstanceCustomDataCount = 0;
+	if (InstanceCustomData.Num() > 0)
+	{
+		InstanceCustomDataCount = InstanceCustomData.Num() / NumInstances;
+	}
+
+	// Only trigger upload if this primitive has instances
+	check(NumInstances > 0);
+}
+
+uint32 FInstanceUploadInfo::CalculatePayloadStrideFloat4() const
+{
+	const bool bSupportsCompressedTransforms = FInstanceSceneShaderData::SupportsCompressedTransforms();
+	int32 PayloadPosition = 0;
+
+	if (InstanceFlags & (INSTANCE_SCENE_DATA_FLAG_HAS_HIERARCHY_OFFSET | INSTANCE_SCENE_DATA_FLAG_HAS_LOCAL_BOUNDS | INSTANCE_SCENE_DATA_FLAG_HAS_SKINNING_DATA))
+	{
+		PayloadPosition += (InstanceFlags & INSTANCE_SCENE_DATA_FLAG_HAS_LOCAL_BOUNDS) ? 2 : 1;
+	}
+	
+	if (InstanceFlags & INSTANCE_SCENE_DATA_FLAG_HAS_DYNAMIC_DATA)
+	{
+		if (bSupportsCompressedTransforms)
+		{
+			PayloadPosition += 2;
+		}
+		else
+		{
+			PayloadPosition += 3;
+		}
+	}
+
+#if WITH_EDITOR
+	if (InstanceFlags & INSTANCE_SCENE_DATA_FLAG_HAS_EDITOR_DATA)
+	{
+		check(InstanceEditorData.Num() == NumInstances);
+		PayloadPosition += 1;
+	}
+#endif
+
+	// TODO: Skip all this if static lighting is disabled
+	if (InstanceFlags & INSTANCE_SCENE_DATA_FLAG_HAS_LIGHTSHADOW_UV_BIAS)
+	{
+		check(InstanceLightShadowUVBias.Num() == NumInstances);
+		PayloadPosition += 1;
+	}
+
+	if (InstanceFlags & INSTANCE_SCENE_DATA_FLAG_HAS_PAYLOAD_EXTENSION)
+	{
+		check(InstancePayloadExtension.Num() / InstancePayloadExtensionCount == NumInstances);
+		PayloadPosition += InstancePayloadExtensionCount;
+	}
+
+	if (InstanceFlags & INSTANCE_SCENE_DATA_FLAG_HAS_CUSTOM_DATA && InstanceCustomDataCount > 0)
+	{
+		PayloadPosition += InstanceCustomDataCount;
+	}
+	return PayloadPosition;
+}
+
+void FInstanceUploadInfo::WritePayloadData(TArrayView<FVector4f> InstancePayloadData, int32 InstanceIndex)
+{
+	const bool bSupportsCompressedTransforms = FInstanceSceneShaderData::SupportsCompressedTransforms();
+	int32 PayloadPosition = 0;
+
+	if (InstanceFlags & (INSTANCE_SCENE_DATA_FLAG_HAS_HIERARCHY_OFFSET | INSTANCE_SCENE_DATA_FLAG_HAS_LOCAL_BOUNDS | INSTANCE_SCENE_DATA_FLAG_HAS_SKINNING_DATA))
+	{
+		const uint32 InstanceHierarchyOffsetValue = (InstanceFlags & INSTANCE_SCENE_DATA_FLAG_HAS_HIERARCHY_OFFSET) ? InstanceHierarchyOffset[InstanceIndex] : 0;
+		InstancePayloadData[PayloadPosition].X = *(const float*)&InstanceHierarchyOffsetValue;
+
+		const uint32 InstanceSkinningDataValue = (InstanceFlags & INSTANCE_SCENE_DATA_FLAG_HAS_SKINNING_DATA) ? InstanceSkinningData[InstanceIndex] : 0;
+		InstancePayloadData[PayloadPosition].Y = *(const float*)&InstanceSkinningDataValue;
+
+		if (InstanceFlags & INSTANCE_SCENE_DATA_FLAG_HAS_LOCAL_BOUNDS)
+		{
+			check(InstanceLocalBounds.Num() == NumInstances);
+			const FRenderBounds& InstanceLocalBoundsValue = InstanceLocalBounds[InstanceIndex];
+			const FVector3f BoundsOrigin = InstanceLocalBoundsValue.GetCenter();
+			const FVector3f BoundsExtent = InstanceLocalBoundsValue.GetExtent();
+
+			InstancePayloadData[PayloadPosition + 0].Z = *(const float*)&BoundsOrigin.X;
+			InstancePayloadData[PayloadPosition + 0].W = *(const float*)&BoundsOrigin.Y;
+
+			InstancePayloadData[PayloadPosition + 1].X = *(const float*)&BoundsOrigin.Z;
+			InstancePayloadData[PayloadPosition + 1].Y = *(const float*)&BoundsExtent.X;
+			InstancePayloadData[PayloadPosition + 1].Z = *(const float*)&BoundsExtent.Y;
+			InstancePayloadData[PayloadPosition + 1].W = *(const float*)&BoundsExtent.Z;
+		}
+
+		PayloadPosition += (InstanceFlags & INSTANCE_SCENE_DATA_FLAG_HAS_LOCAL_BOUNDS) ? 2 : 1;
+	}
+
+	if (InstanceFlags & INSTANCE_SCENE_DATA_FLAG_HAS_DYNAMIC_DATA)
+	{
+		const FRenderTransform PrevLocalToWorld = InstanceSceneDataBuffers != nullptr ?
+			InstanceSceneDataBuffers->GetPrevInstanceToPrimitiveRelative(InstanceIndex) :
+			InstanceDynamicData[InstanceIndex].ComputePrevLocalToWorld(PrevPrimitiveToWorld);
+
+		if (bSupportsCompressedTransforms)
+		{
+			check(PayloadPosition + 1 < InstancePayloadData.Num()); // Sanity check
+			FCompressedTransform CompressedPrevLocalToWorld(PrevLocalToWorld);
+			InstancePayloadData[PayloadPosition + 0] = *(const FVector4f*)&CompressedPrevLocalToWorld.Rotation[0];
+			InstancePayloadData[PayloadPosition + 1] = *(const FVector3f*)&CompressedPrevLocalToWorld.Translation;
+			PayloadPosition += 2;
+		}
+		else
+		{
+			// Note: writes 3x float4s
+			check(PayloadPosition + 2 < InstancePayloadData.Num()); // Sanity check
+			PrevLocalToWorld.To3x4MatrixTranspose((float*)&InstancePayloadData[PayloadPosition]);
+			PayloadPosition += 3;
+		}
+	}
+
+#if WITH_EDITOR
+	if (InstanceFlags & INSTANCE_SCENE_DATA_FLAG_HAS_EDITOR_DATA)
+	{
+		check(InstanceEditorData.Num() == NumInstances);
+		InstancePayloadData[PayloadPosition].X = *(const float*)&InstanceEditorData[InstanceIndex];
+		InstancePayloadData[PayloadPosition].Y = 0.0f; // Unused
+		InstancePayloadData[PayloadPosition].Z = 0.0f; // Unused
+		InstancePayloadData[PayloadPosition].W = 0.0f; // Unused
+		PayloadPosition += 1;
+	}
+#endif
+
+	// TODO: Skip all this if static lighting is disabled
+	if (InstanceFlags & INSTANCE_SCENE_DATA_FLAG_HAS_LIGHTSHADOW_UV_BIAS)
+	{
+		check(InstanceLightShadowUVBias.Num() == NumInstances);
+		InstancePayloadData[PayloadPosition] = InstanceLightShadowUVBias[InstanceIndex];
+		PayloadPosition += 1;
+	}
+
+	if (InstanceFlags & INSTANCE_SCENE_DATA_FLAG_HAS_PAYLOAD_EXTENSION)
+	{
+		check(InstancePayloadExtension.Num() / InstancePayloadExtensionCount == NumInstances);
+		const int32 SrcOffset = InstanceIndex * InstancePayloadExtensionCount;
+		for (int32 Idx = 0; Idx < InstancePayloadExtensionCount; ++Idx)
+		{
+			InstancePayloadData[PayloadPosition++] = InstancePayloadExtension[SrcOffset + Idx];
+		}
+	}
+
+	if (InstanceFlags & INSTANCE_SCENE_DATA_FLAG_HAS_CUSTOM_DATA && InstanceCustomDataCount > 0)
+	{
+		check(PayloadPosition + (InstanceCustomDataCount >> 2u) <= InstancePayloadData.Num());
+		const float* CustomDataGetPtr = &InstanceCustomData[(InstanceIndex * InstanceCustomDataCount)];
+		float* CustomDataPutPtr = (float*)&InstancePayloadData[PayloadPosition];
+		for (uint32 FloatIndex = 0; FloatIndex < uint32(InstanceCustomDataCount); ++FloatIndex)
+		{
+			CustomDataPutPtr[FloatIndex] = CustomDataGetPtr[FloatIndex];
+		}
+	}
+}
\ No newline at end of file
diff --git a/Engine/Source/Runtime/Renderer/Private/GPUScene.h b/Engine/Source/Runtime/Renderer/Private/GPUScene.h
--- a/Engine/Source/Runtime/Renderer/Private/GPUScene.h
+++ b/Engine/Source/Runtime/Renderer/Private/GPUScene.h
@@ -3,6 +3,7 @@
 #pragma once
 
 #include "RenderResource.h"
+#include "RenderTransform.h"
 #include "RendererInterface.h"
 #include "PrimitiveUniformShaderParameters.h"
 #include "PrimitiveSceneInfo.h"
@@ -556,3 +557,48 @@ struct FBatchedPrimitiveShaderData
 private:
 	void Setup(const FPrimitiveUniformShaderParameters& PrimitiveUniformShaderParameters);
 };
+
+
+/**
+ * Info required by the uploader to update the instances that belong to a primitive.
+ */
+struct FInstanceUploadInfo
+{
+	const FInstanceSceneDataBuffers* InstanceSceneDataBuffers = nullptr;
+	TConstArrayView<FInstanceSceneData> PrimitiveInstances;
+	int32 InstanceSceneDataOffset = INDEX_NONE;
+	int32 InstancePayloadDataOffset = INDEX_NONE;
+	int32 InstancePayloadDataStride = 0;
+	int32 InstancePayloadExtensionCount = 0;
+	int32 InstanceCustomDataCount = 0;
+
+	// Optional per-instance data views
+	TConstArrayView<FInstanceDynamicData> InstanceDynamicData;
+	TConstArrayView<uint32> InstanceSkinningData;
+	TConstArrayView<FVector4f> InstanceLightShadowUVBias;
+	TConstArrayView<float> InstanceCustomData;
+	TConstArrayView<float> InstanceRandomID;
+	TConstArrayView<uint32> InstanceHierarchyOffset;
+	TConstArrayView<FRenderBounds> InstanceLocalBounds;
+	TConstArrayView<FVector4f> InstancePayloadExtension;
+#if WITH_EDITOR
+	TConstArrayView<uint32> InstanceEditorData;
+#endif
+
+	uint32 InstanceFlags = 0x0;
+
+	FRenderTransform PrimitiveToWorld;
+	FRenderTransform PrevPrimitiveToWorld;
+	int32 PrimitiveID = INDEX_NONE;
+	uint32 LastUpdateSceneFrameNumber = ~uint32(0);
+	int32 NumInstances = 0;
+	bool bIsPrimitiveForceHidden = false;
+#if DO_CHECK
+	const FPrimitiveSceneInfo* PrimitiveSceneInfo = nullptr;
+#endif
+
+	void SetPrimitiveData(FPrimitiveSceneProxy* InPrimitiveSceneProxy, const FPersistentPrimitiveIndex& PrimitiveIndex, uint32 SceneFrameNumber);
+	void SetInstancePayloadData(const FInstanceSceneDataBuffers* InInstanceSceneDataBuffers);
+	void WritePayloadData(TArrayView<FVector4f> InstancePayloadData, int32 InstanceIndex);
+	uint32 CalculatePayloadStrideFloat4() const;
+};
diff --git a/Engine/Source/Runtime/Renderer/Private/Renderer.cpp b/Engine/Source/Runtime/Renderer/Private/Renderer.cpp
--- a/Engine/Source/Runtime/Renderer/Private/Renderer.cpp
+++ b/Engine/Source/Runtime/Renderer/Private/Renderer.cpp
@@ -54,6 +54,7 @@
 #include "MeshEdgesRendering.h"
 #include "MaterialCache/MaterialCacheTagProvider.h"
 #include "Skinning/SkinnedMeshDebugView.h"
+#include "InstanceDataSceneProxy.h"
 
 DEFINE_LOG_CATEGORY(LogRenderer);
 
@@ -198,6 +199,11 @@ FSceneUniformBuffer* FRendererModule::CreateSinglePrimitiveSceneUniformBuffer(FR
 }
 
 FSceneUniformBuffer* FRendererModule::CreateSinglePrimitiveSceneUniformBuffer(FRDGBuilder& GraphBuilder, ERHIFeatureLevel::Type FeatureLevel, FMeshBatch& Mesh)
+{
+	return CreateSinglePrimitiveSceneUniformBuffer(GraphBuilder, nullptr, FeatureLevel, Mesh);
+}
+
+FSceneUniformBuffer* FRendererModule::CreateSinglePrimitiveSceneUniformBuffer(FRDGBuilder& GraphBuilder, const FPrimitiveSceneProxy* PrimitiveSceneProxy, ERHIFeatureLevel::Type FeatureLevel, FMeshBatch& Mesh)
 {
 	FSceneUniformBuffer& SceneUniforms = *GraphBuilder.AllocObject<FSceneUniformBuffer>();
 
@@ -224,6 +230,39 @@ FSceneUniformBuffer* FRendererModule::CreateSinglePrimitiveSceneUniformBuffer(FR
 			PrimitiveParams.NumInstanceSceneDataEntries = 1;
 			PrimitiveParams.InstancePayloadDataOffset = INDEX_NONE;
 			PrimitiveParams.InstancePayloadDataStride = 0;
+
+			FRDGBufferRef InstancePayloadDataBuffer = nullptr;
+
+			uint32 InstancePayloadFlags = 0;
+			uint32 CustomDataCount = 0;
+			float InstanceRandomID = 0.0f;
+
+			// If a primitive scene proxy is provided, add any extra instance data that this proxy needs.
+			if(PrimitiveSceneProxy)
+			{
+				if (const FInstanceSceneDataBuffers* InstanceSceneDataBuffersPtr = PrimitiveSceneProxy->GetInstanceSceneDataBuffers())
+				{
+					checkf(InstanceSceneDataBuffersPtr->GetNumInstances() == 1, TEXT("CreateSinglePrimitiveSceneUniformBuffer does not currently support instancing"));
+
+					FInstanceUploadInfo InstanceUploadInfo;
+					InstanceUploadInfo.SetInstancePayloadData(InstanceSceneDataBuffersPtr);
+					
+					uint32 PayloadFloat4Size = InstanceUploadInfo.CalculatePayloadStrideFloat4();
+					if (PayloadFloat4Size > 0)
+					{
+						TArray<FVector4f, SceneRenderingAllocator> InstanceUploadInfoBufferData = GraphBuilder.AllocArray<FVector4f>();
+						InstanceUploadInfoBufferData.SetNum(PayloadFloat4Size);
+						InstanceUploadInfo.WritePayloadData(InstanceUploadInfoBufferData, 0);
+						InstancePayloadDataBuffer = CreateStructuredBuffer(GraphBuilder, TEXT("InstancePayloadDataBuffer"), TConstArrayView<FVector4f>(InstanceUploadInfoBufferData));
+
+						PrimitiveParams.InstancePayloadDataOffset = 0;
+						PrimitiveParams.InstancePayloadDataStride = PayloadFloat4Size;
+
+						InstancePayloadFlags = InstanceUploadInfo.InstanceFlags;
+						CustomDataCount = InstanceUploadInfo.InstanceCustomDataCount;
+					}
+				}
+			}
 			
 			// Now we just need to fill out the first entry of primitive data in a buffer and bind it
 			FPrimitiveSceneShaderData PrimitiveSceneData(PrimitiveParams);
@@ -234,9 +273,9 @@ FSceneUniformBuffer* FRendererModule::CreateSinglePrimitiveSceneUniformBuffer(FR
 			(
 				0 /* Primitive Id */,
 				0 /* Relative Instance Id */,
-				0 /* Payload Data Flags */,
+				InstancePayloadFlags /* Payload Data Flags */,
 				INVALID_LAST_UPDATE_FRAME,
-				0 /* Custom Data Count */,
+				CustomDataCount /* Custom Data Count */,
 				0.0f /* Random ID */,
 				PrimitiveParams.LocalToRelativeWorld,
 				true,
@@ -251,7 +290,7 @@ FSceneUniformBuffer* FRendererModule::CreateSinglePrimitiveSceneUniformBuffer(FR
 			FRDGBufferRef PrimitiveSceneDataBuffer = CreateStructuredBuffer(GraphBuilder, TEXT("PrimitiveSceneDataBuffer"), TConstArrayView<FVector4f>(PrimitiveSceneData.Data));
 			FRDGBufferRef LightmapSceneDataBuffer = CreateStructuredBuffer(GraphBuilder, TEXT("LightmapSceneDataBuffer"), TConstArrayView<FVector4f>(LightmapSceneData.Data));
 			FRDGBufferRef InstanceSceneDataBuffer = CreateStructuredBuffer(GraphBuilder, TEXT("InstanceSceneDataBuffer"), TConstArrayView<FVector4f>(InstanceSceneData.Data));
-			FRDGBufferRef InstancePayloadDataBuffer = GSystemTextures.GetDefaultStructuredBuffer(GraphBuilder, sizeof(FVector4f));
+			InstancePayloadDataBuffer = InstancePayloadDataBuffer ? InstancePayloadDataBuffer : GSystemTextures.GetDefaultStructuredBuffer(GraphBuilder, sizeof(FVector4f));
 			FRDGBufferRef DummyBufferLight = GSystemTextures.GetDefaultByteAddressBuffer(GraphBuilder, sizeof(FLightSceneData));
 
 			FGPUSceneResourceParameters ShaderParameters;
diff --git a/Engine/Source/Runtime/Renderer/Private/RendererModule.h b/Engine/Source/Runtime/Renderer/Private/RendererModule.h
--- a/Engine/Source/Runtime/Renderer/Private/RendererModule.h
+++ b/Engine/Source/Runtime/Renderer/Private/RendererModule.h
@@ -49,6 +49,7 @@ public:
 	virtual uint32 GetNumDynamicLightsAffectingPrimitive(const FPrimitiveSceneInfo* PrimitiveSceneInfo,const FLightCacheInterface* LCI) override;
 	virtual void OnWorldCleanup(UWorld* World, bool bSessionEnded, bool bCleanupResources, bool bWorldChanged) override;
 	virtual void InitializeSystemTextures(FRHICommandListImmediate& RHICmdList);
+	virtual FSceneUniformBuffer* CreateSinglePrimitiveSceneUniformBuffer(FRDGBuilder& GraphBuilder, const FPrimitiveSceneProxy* PrimitiveSceneProxy, ERHIFeatureLevel::Type FeatureLevel, FMeshBatch& Mesh) override;
 	virtual FSceneUniformBuffer* CreateSinglePrimitiveSceneUniformBuffer(FRDGBuilder& GraphBuilder, ERHIFeatureLevel::Type FeatureLevel, FMeshBatch& Mesh) override;
 	virtual FSceneUniformBuffer* CreateSinglePrimitiveSceneUniformBuffer(FRDGBuilder& GraphBuilder, const FViewInfo& SceneView, FMeshBatch& Mesh) override;
 	virtual TRDGUniformBufferRef<FBatchedPrimitiveParameters> CreateSinglePrimitiveUniformView(FRDGBuilder& GraphBuilder, ERHIFeatureLevel::Type FeatureLevel, EShaderPlatform ShaderPlatform, FMeshBatch& Mesh) override;
