* Fix calculation errors for the instance payload size. Custom Data was not properly rounding up the size in its validation check.
* Work around static analysis error that did not like the usage of a preprocessor definition in a boolean expression.

diff --git a/Engine/Source/Runtime/Renderer/Private/GPUScene.cpp b/Engine/Source/Runtime/Renderer/Private/GPUScene.cpp
--- a/Engine/Source/Runtime/Renderer/Private/GPUScene.cpp
+++ b/Engine/Source/Runtime/Renderer/Private/GPUScene.cpp
@@ -2499,7 +2499,11 @@ void FInstanceUploadInfo::SetPrimitiveData(FPrimitiveSceneProxy* InPrimitiveScen
 	const FInstanceSceneDataBuffers* InstanceSceneDataBuffersPtr = LocalPrimitiveSceneInfo->GetInstanceSceneDataBuffers();
 	SetInstancePayloadData(InstanceSceneDataBuffersPtr);
 
-	bool bSetPrimitiveToWorld = !!(DO_CHECK) || !InstanceSceneDataBuffers;
+#if DO_CHECK
+	const bool bSetPrimitiveToWorld = true;
+#else
+	const bool bSetPrimitiveToWorld = !InstanceSceneDataBuffers;
+#endif
 
 	// If the primitive has isntancing data then this is already precomputed in the InstanceSceneDataBuffers and we don't need to do it again here, except for validation purposes
 	// If the primitive does not have instancing data, the primitive to world is required
@@ -2625,7 +2629,7 @@ uint32 FInstanceUploadInfo::CalculatePayloadStrideFloat4() const
 
 	if (InstanceFlags & INSTANCE_SCENE_DATA_FLAG_HAS_CUSTOM_DATA && InstanceCustomDataCount > 0)
 	{
-		PayloadPosition += InstanceCustomDataCount;
+		PayloadPosition += FMath::DivideAndRoundUp(InstanceCustomDataCount, 4);
 	}
 	return PayloadPosition;
 }
@@ -2717,7 +2721,7 @@ void FInstanceUploadInfo::WritePayloadData(TArrayView<FVector4f> InstancePayload
 
 	if (InstanceFlags & INSTANCE_SCENE_DATA_FLAG_HAS_CUSTOM_DATA && InstanceCustomDataCount > 0)
 	{
-		check(PayloadPosition + (InstanceCustomDataCount >> 2u) <= InstancePayloadData.Num());
+		check((PayloadPosition + FMath::DivideAndRoundUp(InstanceCustomDataCount, 4)) <= InstancePayloadData.Num());
 		const float* CustomDataGetPtr = &InstanceCustomData[(InstanceIndex * InstanceCustomDataCount)];
 		float* CustomDataPutPtr = (float*)&InstancePayloadData[PayloadPosition];
 		for (uint32 FloatIndex = 0; FloatIndex < uint32(InstanceCustomDataCount); ++FloatIndex)
@@ -2725,4 +2729,4 @@ void FInstanceUploadInfo::WritePayloadData(TArrayView<FVector4f> InstancePayload
 			CustomDataPutPtr[FloatIndex] = CustomDataGetPtr[FloatIndex];
 		}
 	}
-}
\ No newline at end of file
+}
