Silenced PVSStudio warning about intentional single-iteration loop.

#rnx

diff --git a/Engine/Source/Runtime/Renderer/Private/GPUScene.cpp b/Engine/Source/Runtime/Renderer/Private/GPUScene.cpp
--- a/Engine/Source/Runtime/Renderer/Private/GPUScene.cpp
+++ b/Engine/Source/Runtime/Renderer/Private/GPUScene.cpp
@@ -2096,7 +2096,7 @@ bool FGPUScene::ExecuteDeferredGPUWritePass(FRDGBuilder& GraphBuilder, TArray<FV
 
 bool FGPUScene::HasPendingGPUWrite(uint32 PrimitiveId) const
 {
-	for (uint32 PassIndex = uint32(LastDeferredGPUWritePass) + 1; PassIndex < uint32(EGPUSceneGPUWritePass::Num); ++PassIndex)
+	for (uint32 PassIndex = uint32(LastDeferredGPUWritePass) + 1; PassIndex < uint32(EGPUSceneGPUWritePass::Num); ++PassIndex) //-V1008 Disable single-iteration loop warning as this is by design in case EGPUSceneGPUWritePass is extended later.
 	{
 		if (DeferredGPUWritePassDelegates[PassIndex].FindByPredicate(
 			[PrimitiveId](const FDeferredGPUWrite& Write)
