Add support for non 0 global uniform buffer index for raytracing hit shaders

diff --git a/Engine/Source/Runtime/Renderer/Private/MeshPassProcessor.cpp b/Engine/Source/Runtime/Renderer/Private/MeshPassProcessor.cpp
--- a/Engine/Source/Runtime/Renderer/Private/MeshPassProcessor.cpp
+++ b/Engine/Source/Runtime/Renderer/Private/MeshPassProcessor.cpp
@@ -316,7 +316,7 @@ FRayTracingLocalShaderBindings* FMeshDrawShaderBindings::SetRayTracingShaderBind
 		check(LooseParameterBuffers.Num() <= 1);
 
 		const FShaderLooseParameterBufferInfo& LooseParameterBuffer = SingleShaderBindings.ParameterMapInfo.LooseParameterBuffers[0];
-		check(LooseParameterBuffer.BaseIndex == 0);
+		MaxUniformBufferIndex = FMath::Max((int32)LooseParameterBuffer.BaseIndex, MaxUniformBufferIndex);
 
 		for (int32 LooseParameterIndex = 0; LooseParameterIndex < LooseParameterBuffer.Parameters.Num(); LooseParameterIndex++)
 		{
