Prevented crash on instance upload for scenes with too many instances by skipping upload for instances outside the instance data buffer.

diff --git a/Engine/Source/Runtime/Renderer/Private/GPUScene.cpp b/Engine/Source/Runtime/Renderer/Private/GPUScene.cpp
--- a/Engine/Source/Runtime/Renderer/Private/GPUScene.cpp
+++ b/Engine/Source/Runtime/Renderer/Private/GPUScene.cpp
@@ -323,6 +323,7 @@ struct FPrimitiveUploadInfoHeader
 	uint32 PrimitiveID = INVALID_PRIMITIVE_ID;
 
 	/** Optional */
+	int32 InstanceSceneDataOffset = INDEX_NONE;
 	int32 NumInstanceUploads = 0;
 	int32 InstancePayloadDataOffset = 0;
 	int32 InstancePayloadDataStride = 0;
@@ -460,6 +461,8 @@ struct FUploadDataSourceAdapterScenePrimitives
 
 		PrimitiveUploadInfo.InstancePayloadDataOffset = PrimitiveSceneInfo->GetInstancePayloadDataOffset();
 		PrimitiveUploadInfo.InstancePayloadDataStride = PrimitiveSceneInfo->GetInstancePayloadDataStride();
+
+		PrimitiveUploadInfo.InstanceSceneDataOffset = PrimitiveSceneInfo->GetInstanceSceneDataOffset();
 		// TODO: Somehow associate the procedural update with the update entry such that we can cross validate the existence of a writer here?
 		if (PrimitiveSceneProxy->IsInstanceDataGPUOnly())
 		{
@@ -469,6 +472,7 @@ struct FUploadDataSourceAdapterScenePrimitives
 		{
 			PrimitiveUploadInfo.NumInstanceUploads = PrimitiveSceneInfo->GetNumInstanceSceneDataEntries();
 		}
+		check(PrimitiveUploadInfo.NumInstanceUploads == 0 || PrimitiveUploadInfo.InstanceSceneDataOffset >= 0);
 	}
 
 	/**
@@ -1171,6 +1175,13 @@ void FGPUScene::UploadGeneral(FRDGBuilder& GraphBuilder, const FRegisteredBuffer
 	{
 		FPrimitiveUploadInfoHeader UploadInfo;
 		UploadDataSourceAdapter.GetPrimitiveInfoHeader(ItemIndex, UploadInfo);
+
+		if (UploadInfo.InstanceSceneDataOffset + UploadInfo.NumInstanceUploads >= int32(MAX_INSTANCE_ID))
+		{
+			// clamp the uploads such that we don't try to scatter any that will land outside the buffer anyway.
+			UploadInfo.NumInstanceUploads = FMath::Max(0, int32(MAX_INSTANCE_ID) - UploadInfo.InstanceSceneDataOffset);
+		}
+
 		UploadInfo.NumPayloadUploadInstances = UploadInfo.InstancePayloadDataStride ? UploadInfo.NumInstanceUploads : 0;
 		// Figure out how much space we actually have left.
 		int32 PayloadSpaceLeft = FMath::Max(int32(0), int32(MaxInstancePayloadOffset) - int32(UploadInfo.InstancePayloadDataOffset));
@@ -1330,7 +1341,7 @@ void FGPUScene::UploadGeneral(FRDGBuilder& GraphBuilder, const FRegisteredBuffer
 
 						for (uint32 RefIndex = 0; RefIndex < InstanceDataStrideInFloat4s; ++RefIndex)
 						{
-							uint32 ScatterOffset = TaskContext.CalcInstanceDataIndex(UploadInfo.InstanceSceneDataOffset + InstanceIndex, RefIndex);
+							uint32 ScatterOffset = TaskContext.CalcInstanceDataIndex(PrimitiveHeader.InstanceSceneDataOffset + InstanceIndex, RefIndex);
 							FVector4f* DstVector = static_cast<FVector4f*>(TaskContext.InstanceSceneUploader->Set_GetRef(UploadInstanceItemOffset + RefIndex, ScatterOffset));
 							*DstVector = InstanceSceneData.Data[RefIndex];
 						}
@@ -1470,6 +1481,7 @@ struct FUploadDataSourceAdapterDynamicPrimitives
 			PrimitiveUploadInfo.InstancePayloadDataOffset = INDEX_NONE;
 			PrimitiveUploadInfo.InstancePayloadDataStride = 0;
 		}
+		PrimitiveUploadInfo.InstanceSceneDataOffset		= InstanceIDStartOffset + PrimData.LocalInstanceSceneDataOffset;
 		if (PrimData.SourceData.DataWriterGPU.IsBound())
 		{
 			// Only upload if we have data, otherwise expect the delegate to handle missing data
@@ -1512,7 +1524,6 @@ struct FUploadDataSourceAdapterDynamicPrimitives
 			InstanceUploadInfo.PrimitiveID					= PrimitiveIDStartOffset + ItemIndex;
 			InstanceUploadInfo.PrimitiveToWorld				= ShaderParams.LocalToRelativeWorld;
 			InstanceUploadInfo.PrevPrimitiveToWorld			= ShaderParams.PreviousLocalToRelativeWorld;
-			InstanceUploadInfo.InstanceSceneDataOffset		= InstanceIDStartOffset + PrimData.LocalInstanceSceneDataOffset;
 			InstanceUploadInfo.InstancePayloadDataOffset	= PrimData.LocalPayloadDataOffset == INDEX_NONE ? INDEX_NONE : PayloadStartOffset + PrimData.LocalPayloadDataOffset;
 			InstanceUploadInfo.InstancePayloadDataStride	= PrimData.SourceData.GetPayloadFloat4Stride();
 			InstanceUploadInfo.InstanceCustomDataCount		= PrimData.SourceData.NumInstanceCustomDataFloats;
@@ -2507,8 +2518,6 @@ void FInstanceUploadInfo::SetPrimitiveData(FPrimitiveSceneProxy* InPrimitiveScen
 	PrimitiveSceneInfo = LocalPrimitiveSceneInfo;
 #endif
 
-	InstanceSceneDataOffset = LocalPrimitiveSceneInfo->GetInstanceSceneDataOffset();
-	check(InstanceSceneDataOffset >= 0);
 	InstancePayloadDataOffset = LocalPrimitiveSceneInfo->GetInstancePayloadDataOffset();
 	InstancePayloadDataStride = LocalPrimitiveSceneInfo->GetInstancePayloadDataStride();
 	LastUpdateSceneFrameNumber = SceneFrameNumber;
diff --git a/Engine/Source/Runtime/Renderer/Private/GPUScene.h b/Engine/Source/Runtime/Renderer/Private/GPUScene.h
--- a/Engine/Source/Runtime/Renderer/Private/GPUScene.h
+++ b/Engine/Source/Runtime/Renderer/Private/GPUScene.h
@@ -566,7 +566,6 @@ struct FInstanceUploadInfo
 {
 	const FInstanceSceneDataBuffers* InstanceSceneDataBuffers = nullptr;
 	TConstArrayView<FInstanceSceneData> PrimitiveInstances;
-	int32 InstanceSceneDataOffset = INDEX_NONE;
 	int32 InstancePayloadDataOffset = INDEX_NONE;
 	int32 InstancePayloadDataStride = 0;
 	int32 InstancePayloadExtensionCount = 0;
