Substrate - clean up duplicated constant.

#tests pc editor dx12
#rnx

diff --git a/Engine/Source/Runtime/Renderer/Private/Substrate/Substrate.cpp b/Engine/Source/Runtime/Renderer/Private/Substrate/Substrate.cpp
--- a/Engine/Source/Runtime/Renderer/Private/Substrate/Substrate.cpp
+++ b/Engine/Source/Runtime/Renderer/Private/Substrate/Substrate.cpp
@@ -1287,13 +1287,6 @@ uint32 TileTypeDispatchIndirectArgOffset(const ESubstrateTileType Type)
 	return GetSubstrateTileTypeDispatchIndirectArgOffset_Byte(Type);
 }
 
-// Add additionnaly bits for filling/clearing stencil to ensure that the 'Substrate' bits are not corrupted by the stencil shadows 
-// when generating shadow mask. Withouth these 'trailing' bits, the incr./decr. operation would change/corrupt the 'Substrate' bits
-constexpr uint32 StencilBit_Fast_1			= StencilBit_Fast;
-constexpr uint32 StencilBit_Single_1		= StencilBit_Single;
-constexpr uint32 StencilBit_Complex_1		= StencilBit_Complex; 
-constexpr uint32 StencilBit_ComplexSpecial_1= StencilBit_ComplexSpecial; 
-
 void AddSubstrateInternalClassificationTilePass(
 	FRDGBuilder& GraphBuilder,
 	const FViewInfo& View,
@@ -1386,8 +1379,8 @@ void AddSubstrateInternalClassificationTilePass(
 						false, CF_Always,
 						true, CF_Always, SO_Keep, SO_Keep, SO_Replace,
 						false, CF_Always, SO_Keep, SO_Keep, SO_Keep,
-						0xFF, StencilBit_Fast_1>::GetRHI();
-					StencilRef = StencilBit_Fast_1;
+						0xFF, StencilBit_Fast>::GetRHI();
+					StencilRef = StencilBit_Fast;
 				}
 				break;
 				case ESubstrateTileType::ESingle:
@@ -1396,8 +1389,8 @@ void AddSubstrateInternalClassificationTilePass(
 						false, CF_Always,
 						true, CF_Always, SO_Keep, SO_Keep, SO_Replace,
 						false, CF_Always, SO_Keep, SO_Keep, SO_Keep,
-						0xFF, StencilBit_Single_1>::GetRHI();
-					StencilRef = StencilBit_Single_1;
+						0xFF, StencilBit_Single>::GetRHI();
+					StencilRef = StencilBit_Single;
 				}
 				break;
 				case ESubstrateTileType::EComplex:
@@ -1406,8 +1399,8 @@ void AddSubstrateInternalClassificationTilePass(
 						false, CF_Always,
 						true, CF_Always, SO_Keep, SO_Keep, SO_Replace,
 						false, CF_Always, SO_Keep, SO_Keep, SO_Keep,
-						0xFF, StencilBit_Complex_1>::GetRHI();
-					StencilRef = StencilBit_Complex_1;
+						0xFF, StencilBit_Complex>::GetRHI();
+					StencilRef = StencilBit_Complex;
 				}
 				break;
 				case ESubstrateTileType::EComplexSpecial:
@@ -1416,8 +1409,8 @@ void AddSubstrateInternalClassificationTilePass(
 						false, CF_Always,
 						true, CF_Always, SO_Keep, SO_Keep, SO_Replace,
 						false, CF_Always, SO_Keep, SO_Keep, SO_Keep,
-						0xFF, StencilBit_ComplexSpecial_1>::GetRHI();
-					StencilRef = StencilBit_ComplexSpecial_1;
+						0xFF, StencilBit_ComplexSpecial>::GetRHI();
+					StencilRef = StencilBit_ComplexSpecial;
 				}
 				break;
 				}
diff --git a/Engine/Source/Runtime/Renderer/Private/Substrate/Substrate.h b/Engine/Source/Runtime/Renderer/Private/Substrate/Substrate.h
--- a/Engine/Source/Runtime/Renderer/Private/Substrate/Substrate.h
+++ b/Engine/Source/Runtime/Renderer/Private/Substrate/Substrate.h
@@ -209,6 +209,10 @@ struct FSubstrateViewDebugData
 
 namespace Substrate
 {
+
+// Bits for filling/clearing stencil during shadow/lighting passes. 
+// Those are designed to not conflicts with other stencil shadow bits, when generating shadow mask. 
+// Withouth these 'trailing' bits, the incr./decr. operation would change/corrupt the 'Substrate' bits.
 constexpr uint32 StencilBit_Fast			= 0x10; // In sync with SceneRenderTargets.h - GET_STENCIL_BIT_MASK(STENCIL_SUBSTRATE_FASTPATH)
 constexpr uint32 StencilBit_Single			= 0x20; // In sync with SceneRenderTargets.h - GET_STENCIL_BIT_MASK(STENCIL_SUBSTRATE_SINGLEPATH)
 constexpr uint32 StencilBit_Complex			= 0x40; // In sync with SceneRenderTargets.h - GET_STENCIL_BIT_MASK(STENCIL_SUBSTRATE_COMPLEX)
