VolumetricCloud - enable late fog by default since this is the correct way to apply fog on cloud when lumen is enabled and r.VolumetricRenderTarget.PreferAsyncCompute=1.
And perfomance ReplayRun on a game has shown no performance regression.

#rnx
#tests PC editor dx12

diff --git a/Engine/Source/Runtime/Renderer/Private/VolumetricCloudRendering.cpp b/Engine/Source/Runtime/Renderer/Private/VolumetricCloudRendering.cpp
--- a/Engine/Source/Runtime/Renderer/Private/VolumetricCloudRendering.cpp
+++ b/Engine/Source/Runtime/Renderer/Private/VolumetricCloudRendering.cpp
@@ -97,8 +97,9 @@ static TAutoConsoleVariable<int32> CVarVolumetricCloudApplyFogOnAllPixel(
 // - Order of operations when PreferAsyncCompute=1 are (1) Cloud tracing in async compute pipe => (2) lumen update translucent lighting volume => (3) VFog upsampling (4) Cloud upsampling.
 // - But cloud tracing in (1) has bad volumetric since it is generated later, and it cannot be moved up because it depends on lumen generating the translucent lighting volume.
 // - So ApplyFogLate=1 allows to to have fog on clouds applied in (4). Which is fine for out single depth operation.
+// This is now enabled by default since perfomance ReplayRun on a game has shown no performance regression.
 static TAutoConsoleVariable<int32> CVarVolumetricCloudApplyFogLate(
-	TEXT("r.VolumetricCloud.ApplyFogLate"), 0,
+	TEXT("r.VolumetricCloud.ApplyFogLate"), 1,
 	TEXT("When true, fog will be applied on cloud later in the frame. This applies fog on cloud without visible latency due to reprojection and reconstruction. It is a bit more expenssive (upsampling pass cost is roughly x2) but the only correct look when Lumen is enabled and tracing is run on async compute using r.VolumetricRenderTarget.PreferAsyncCompute=1."),
 	ECVF_RenderThreadSafe);
 
