[Ray tracing culling] Compute view to primitive distance as a distance to bbox instead of bounding sphere. This makes it the same as how Lumen compute the distance during mesh cards culling. Previously, we can have a primitive not culled in the ray tracing scene but its mesh cards are culled in Lumen scene. When a ray hit such an object, it won't get any surface cache lighting because the object doesn't have surface cache. This can cause weird visual artifacts such as reflections being black when it is not supposed to.

#tests tested in editor and on a console
#FailedValidations Preflight

diff --git a/Engine/Shaders/Private/RayTracing/RayTracingInstanceBufferUtil.usf b/Engine/Shaders/Private/RayTracing/RayTracingInstanceBufferUtil.usf
--- a/Engine/Shaders/Private/RayTracing/RayTracingInstanceBufferUtil.usf
+++ b/Engine/Shaders/Private/RayTracing/RayTracingInstanceBufferUtil.usf
@@ -63,6 +63,7 @@ float AngleThresholdRatioSq;
 float3 ViewOrigin;
 uint CullingMode;
 uint CullUsingGroups;
+int bUseDistanceToBox;
 
 bool CullInstance(FInstanceSceneData InstanceSceneData, FPrimitiveSceneData PrimitiveData, float4x4 LocalToTranslatedWorld4x4, bool bFarFieldInstance)
 {
@@ -82,14 +83,31 @@ bool CullInstance(FInstanceSceneData InstanceSceneData, FPrimitiveSceneData Prim
 	const float Radius = length(LocalBoxExtent * NonUniformScale);
 
 	const float InstanceDrawDistSq = length2(CenterTranslatedWorld - ViewOrigin);
+
+	float DistanceToInstanceSq;
+	if (bUseDistanceToBox != 0)
+	{
+		float4 Temp = 0.0f;
+		Temp += abs(LocalBoxExtent.xxxx * LocalToTranslatedWorld[0]);
+		Temp += abs(LocalBoxExtent.yyyy * LocalToTranslatedWorld[1]);
+		Temp += abs(LocalBoxExtent.zzzz * LocalToTranslatedWorld[2]);
+
+		float3 WorldAABBExtent = Temp.xyz;
+		DistanceToInstanceSq = ComputeSquaredDistanceFromBoxToPoint(CenterTranslatedWorld, WorldAABBExtent, ViewOrigin);
+	}
+	else
+	{
+		DistanceToInstanceSq = max(sqrt(InstanceDrawDistSq) - Radius, 0.0f);
+		DistanceToInstanceSq *= DistanceToInstanceSq;
+	}
 	
 	if (bFarFieldInstance)
 	{
-		bIsVisible = InstanceDrawDistSq <= Square(FarFieldCullingRadius + Radius);
+		bIsVisible = DistanceToInstanceSq <= Square(FarFieldCullingRadius);
 	}
 	else
 	{
-		const bool bIsVisibleDistance = InstanceDrawDistSq <= Square(CullingRadius + Radius);
+		const bool bIsVisibleDistance = DistanceToInstanceSq <= Square(CullingRadius);
 		const bool bIsVisibleAngle = (Square(Radius) / InstanceDrawDistSq) >= AngleThresholdRatioSq;
 		
 		// Only culling modes 2 & 3 are supported for now
diff --git a/Engine/Source/Runtime/Renderer/Private/RayTracing/RayTracingInstanceBufferUtil.cpp b/Engine/Source/Runtime/Renderer/Private/RayTracing/RayTracingInstanceBufferUtil.cpp
--- a/Engine/Source/Runtime/Renderer/Private/RayTracing/RayTracingInstanceBufferUtil.cpp
+++ b/Engine/Source/Runtime/Renderer/Private/RayTracing/RayTracingInstanceBufferUtil.cpp
@@ -600,6 +600,7 @@ struct FRayTracingBuildInstanceBufferCS : public FGlobalShader
 		SHADER_PARAMETER(FVector3f, ViewOrigin)
 		SHADER_PARAMETER(uint32, CullingMode)
 		SHADER_PARAMETER(uint32, CullUsingGroups)
+		SHADER_PARAMETER(int32, bUseDistanceToBox)
 
 		SHADER_PARAMETER_UAV(RWStructuredBuffer<uint>, RWOutputStats)
 		SHADER_PARAMETER(uint32, OutputStatsOffset)
@@ -696,6 +697,7 @@ void BuildRayTracingInstanceBuffer(
 		PassParams.ViewOrigin = CullingParameters->TranslatedViewOrigin;
 		PassParams.CullingMode = uint32(CullingParameters->CullingMode);
 		PassParams.CullUsingGroups = uint32(CullingParameters->bCullUsingGroupIds);
+		PassParams.bUseDistanceToBox = CullingParameters->bUseDistanceToBox ? 1 : 0;
 	}
 
 	PassParams.RWOutputStats = OutputStatsUAV;
diff --git a/Engine/Source/Runtime/Renderer/Private/RayTracing/RayTracingInstanceCulling.cpp b/Engine/Source/Runtime/Renderer/Private/RayTracing/RayTracingInstanceCulling.cpp
--- a/Engine/Source/Runtime/Renderer/Private/RayTracing/RayTracingInstanceCulling.cpp
+++ b/Engine/Source/Runtime/Renderer/Private/RayTracing/RayTracingInstanceCulling.cpp
@@ -47,6 +47,12 @@ static TAutoConsoleVariable<int32> CVarRayTracingCullingGroupIds(
 	TEXT("Cull using aggregate ray tracing group id bounds when defined instead of primitive or instance bounds."),
 	ECVF_RenderThreadSafe);
 
+static TAutoConsoleVariable<int32> CVarRayTracingCullingUseDistanceToBox(
+	TEXT("r.RayTracing.Culling.UseDistanceToBox"),
+	1,
+	TEXT("Whether to compute distance to objects as a distance to their bbox or bounding sphere."),
+	ECVF_RenderThreadSafe);
+
 RayTracing::ECullingMode RayTracing::GetCullingMode(const FEngineShowFlags& ShowFlags)
 {
 	// Disable culling if path tracer is used, so that path tracer matches raster view
@@ -58,6 +64,11 @@ RayTracing::ECullingMode RayTracing::GetCullingMode(const FEngineShowFlags& Show
 	return (ECullingMode) FMath::Clamp(CVarRayTracingCulling.GetValueOnRenderThread(), 0, int32(ECullingMode::MAX) - 1);
 }
 
+static bool CullingUseDistanceToBox()
+{
+	return CVarRayTracingCullingUseDistanceToBox.GetValueOnRenderThread() != 0;
+}
+
 float GetRayTracingCullingRadius()
 {
 	return CVarRayTracingCullingRadius.GetValueOnRenderThread();
@@ -80,6 +91,7 @@ void FRayTracingCullingParameters::Init(FViewInfo& View)
 	bCullUsingGroupIds = CVarRayTracingCullingGroupIds.GetValueOnRenderThread() != 0;
 	bCullMinDrawDistance = CVarRayTracingCullingUseMinDrawDistance.GetValueOnRenderThread() != 0;
 	bUseInstanceCulling = CVarRayTracingCullingPerInstance.GetValueOnRenderThread() != 0 && bCullAllObjects;
+	bUseDistanceToBox = CullingUseDistanceToBox();
 }
 
 namespace RayTracing
@@ -121,18 +133,34 @@ bool CullBounds(const FRayTracingCullingParameters& CullingParameters, const FBo
 {
 	const float ObjectRadius = ObjectBounds.SphereRadius;
 	const FVector ObjectCenter = ObjectBounds.Origin;
-	float CameraToObjectCenterLengthSq = FVector::DistSquared(ObjectCenter, CullingParameters.ViewOrigin);
+
+	const FVector3f TranslatedObjectCenter = FVector3f(ObjectCenter - CullingParameters.ViewOrigin);
+	const float CameraToObjectCenterLengthSq = TranslatedObjectCenter.SquaredLength();
+
+	float CameraToObjectDistanceSq;
+	if (CullingParameters.bUseDistanceToBox)
+	{
+		const FVector3f BoxExtent = FVector3f(ObjectBounds.BoxExtent);
+		const FVector3f BoxMin = TranslatedObjectCenter - BoxExtent;
+		const FVector3f BoxMax = TranslatedObjectCenter + BoxExtent;
+		CameraToObjectDistanceSq = ComputeSquaredDistanceFromBoxToPoint(BoxMin, BoxMax, FVector3f::Zero());
+	}
+	else
+	{
+		const float DistanceToSphere = FMath::Max(FMath::Sqrt(CameraToObjectCenterLengthSq) - ObjectRadius, 0.0f);
+		CameraToObjectDistanceSq = FMath::Square(DistanceToSphere);
+	}
 
 	if (bIsFarFieldPrimitive)
 	{
-		if (CameraToObjectCenterLengthSq > FMath::Square(CullingParameters.FarFieldCullingRadius + ObjectRadius))
+		if (CameraToObjectDistanceSq > FMath::Square(CullingParameters.FarFieldCullingRadius))
 		{
 			return true;
 		}
 	}
 	else
 	{
-		const bool bIsFarEnoughToCull = CameraToObjectCenterLengthSq > FMath::Square(CullingParameters.CullingRadius + ObjectRadius);
+		const bool bIsFarEnoughToCull = CameraToObjectDistanceSq > FMath::Square(CullingParameters.CullingRadius);
 		const bool bIsNearEnoughToCull = CameraToObjectCenterLengthSq < FMath::Square(MinDrawDistance);
 
 		if (CullingParameters.bCullMinDrawDistance && bIsNearEnoughToCull)
diff --git a/Engine/Source/Runtime/Renderer/Private/SceneRendering.h b/Engine/Source/Runtime/Renderer/Private/SceneRendering.h
--- a/Engine/Source/Runtime/Renderer/Private/SceneRendering.h
+++ b/Engine/Source/Runtime/Renderer/Private/SceneRendering.h
@@ -1114,6 +1114,7 @@ struct FRayTracingCullingParameters
 	bool bCullUsingGroupIds;
 	bool bCullMinDrawDistance;
 	bool bUseInstanceCulling;
+	bool bUseDistanceToBox;
 
 	void Init(FViewInfo& View);
 };
