Skinning extension fixes. Fixed vertex factory dirty bit not being reset. Also fixed ensure that never fires.

diff --git a/Engine/Source/Runtime/Renderer/Private/Skinning/SkinningSceneExtension.cpp b/Engine/Source/Runtime/Renderer/Private/Skinning/SkinningSceneExtension.cpp
--- a/Engine/Source/Runtime/Renderer/Private/Skinning/SkinningSceneExtension.cpp
+++ b/Engine/Source/Runtime/Renderer/Private/Skinning/SkinningSceneExtension.cpp
@@ -596,6 +596,7 @@ void FSkinningSceneExtension::PerformSkinning(FSkinningSceneParameters& Paramete
 				const auto UpdateVertexFactory = [&](FHeaderData& Data)
 				{
 					Data.MeshObject->UpdateSceneExtensionHeader(RHICmdList, Data.Pack(), Buffers->TransformDataBufferSRV);
+					Data.bVertexFactoryDirty = 0;
 				};
 
 				if (bVertexFactoryFullUpload)
@@ -1196,7 +1197,7 @@ void FSkinningSceneExtension::FUpdater::PostMeshUpdate(FRDGBuilder& GraphBuilder
 
 		// Catch when/if no transform buffer data is allocated for a primitive we're tracking.
 		// This should be indicative of a bug.
-		ensure(Data.HierarchyBufferCount != INDEX_NONE && Data.TransformBufferCount != INDEX_NONE);
+		ensure(Data.HierarchyBufferCount != 0 && Data.TransformBufferCount != 0);
 
 		check(SceneData->Uploader.IsValid()); // Sanity check
 		SceneData->Uploader->HeaderDataUploader.Add(Data.Pack(), PersistentIndex);
