Disable PSO precaching for global RT shadows compute shaders

diff --git a/Engine/Source/Runtime/Renderer/Private/RayTracing/RayTracingShadows.cpp b/Engine/Source/Runtime/Renderer/Private/RayTracing/RayTracingShadows.cpp
--- a/Engine/Source/Runtime/Renderer/Private/RayTracing/RayTracingShadows.cpp
+++ b/Engine/Source/Runtime/Renderer/Private/RayTracing/RayTracingShadows.cpp
@@ -292,7 +292,7 @@ class FOcclusionCS : public FGlobalShader
 	static EShaderPermutationPrecacheRequest ShouldPrecachePermutation(const FGlobalShaderPermutationParameters& Parameters)
 	{
 		// Disable precaching for occlusion CS - should use project setting if it's used?
-		return EShaderPermutationPrecacheRequest::Precached;
+		return EShaderPermutationPrecacheRequest::NotPrecached;
 	}
 
 	BEGIN_SHADER_PARAMETER_STRUCT(FParameters, )
