IndirectLightingCache should require command lists from its callers.

diff --git a/Engine/Source/Runtime/Renderer/Private/IndirectLightingCache.cpp b/Engine/Source/Runtime/Renderer/Private/IndirectLightingCache.cpp
--- a/Engine/Source/Runtime/Renderer/Private/IndirectLightingCache.cpp
+++ b/Engine/Source/Runtime/Renderer/Private/IndirectLightingCache.cpp
@@ -446,14 +446,14 @@ void FIndirectLightingCache::StartUpdateCachePrimitivesTask(FScene* Scene, FScen
 	);
 }
 
-void FIndirectLightingCache::FinalizeCacheUpdates(FScene* Scene, FSceneRenderer& Renderer, FILCUpdatePrimTaskData& TaskData)
+void FIndirectLightingCache::FinalizeCacheUpdates(FRHICommandListBase& RHICmdList, FScene* Scene, FSceneRenderer& Renderer, FILCUpdatePrimTaskData& TaskData)
 {
 	SCOPE_CYCLE_COUNTER(STAT_UpdateIndirectLightingCacheFinalize);	
 	FTaskGraphInterface::Get().WaitUntilTaskCompletes(TaskData.TaskRef, ENamedThreads::GetRenderThread_Local());
-	FinalizeUpdateInternal_RenderThread(Scene, Renderer, TaskData.OutBlocksToUpdate, TaskData.OutTransitionsOverTimeToUpdate, TaskData.OutPrimitivesToUpdateStaticMeshes);
+	FinalizeUpdateInternal_RenderThread(RHICmdList, Scene, Renderer, TaskData.OutBlocksToUpdate, TaskData.OutTransitionsOverTimeToUpdate, TaskData.OutPrimitivesToUpdateStaticMeshes);
 }
 
-void FIndirectLightingCache::UpdateCache(FScene* Scene, FSceneRenderer& Renderer, bool bAllowUnbuiltPreview)
+void FIndirectLightingCache::UpdateCache(FRHICommandListBase& RHICmdList, FScene* Scene, FSceneRenderer& Renderer, bool bAllowUnbuiltPreview)
 {
 	SCOPE_CYCLE_COUNTER(STAT_UpdateIndirectLightingCache);
 
@@ -462,7 +462,7 @@ void FIndirectLightingCache::UpdateCache(FScene* Scene, FSceneRenderer& Renderer
 	TArray<FPrimitiveSceneInfo*> PrimitivesToUpdateStaticMeshes;
 
 	UpdateCachePrimitivesInternal(Scene, Renderer, bAllowUnbuiltPreview, BlocksToUpdate, TransitionsOverTimeToUpdate, PrimitivesToUpdateStaticMeshes);
-	FinalizeUpdateInternal_RenderThread(Scene, Renderer, BlocksToUpdate, TransitionsOverTimeToUpdate, PrimitivesToUpdateStaticMeshes);
+	FinalizeUpdateInternal_RenderThread(RHICmdList, Scene, Renderer, BlocksToUpdate, TransitionsOverTimeToUpdate, PrimitivesToUpdateStaticMeshes);
 }
 
 bool FIndirectLightingCache::IndirectLightingAllowed(FScene* Scene, FSceneRenderer& Renderer) const
@@ -573,7 +573,7 @@ void FIndirectLightingCache::UpdateCachePrimitivesInternal(FScene* Scene, FScene
 	}
 }
 
-void FIndirectLightingCache::FinalizeUpdateInternal_RenderThread(FScene* Scene, FSceneRenderer& Renderer, TMap<FIntVector, FBlockUpdateInfo>& BlocksToUpdate, const TArray<FIndirectLightingCacheAllocation*>& TransitionsOverTimeToUpdate, TArray<FPrimitiveSceneInfo*>& PrimitivesToUpdateStaticMeshes)
+void FIndirectLightingCache::FinalizeUpdateInternal_RenderThread(FRHICommandListBase& RHICmdList, FScene* Scene, FSceneRenderer& Renderer, TMap<FIntVector, FBlockUpdateInfo>& BlocksToUpdate, const TArray<FIndirectLightingCacheAllocation*>& TransitionsOverTimeToUpdate, TArray<FPrimitiveSceneInfo*>& PrimitivesToUpdateStaticMeshes)
 {
 	check(IsInRenderingThread());
 
@@ -581,7 +581,7 @@ void FIndirectLightingCache::FinalizeUpdateInternal_RenderThread(FScene* Scene,
 	{
 		{
 			SCOPE_CYCLE_COUNTER(STAT_UpdateIndirectLightingCacheBlocks);
-			UpdateBlocks(Scene, Renderer.Views.GetData(), BlocksToUpdate);
+			UpdateBlocks(RHICmdList, Scene, Renderer.Views.GetData(), BlocksToUpdate);
 		}
 
 		{
@@ -756,18 +756,18 @@ void FIndirectLightingCache::UpdateCachePrimitive(
 	}
 }
 
-void FIndirectLightingCache::UpdateBlocks(FScene* Scene, FViewInfo* DebugDrawingView, TMap<FIntVector, FBlockUpdateInfo>& BlocksToUpdate)
+void FIndirectLightingCache::UpdateBlocks(FRHICommandListBase& RHICmdList, FScene* Scene, FViewInfo* DebugDrawingView, TMap<FIntVector, FBlockUpdateInfo>& BlocksToUpdate)
 {
 	if (BlocksToUpdate.Num() > 0 && !IsInitialized())
 	{
-		InitResource(FRHICommandListImmediate::Get());
+		InitResource(RHICmdList);
 	}
 
 	INC_DWORD_STAT_BY(STAT_IndirectLightingCacheUpdates, BlocksToUpdate.Num());
 
 	for (TMap<FIntVector, FBlockUpdateInfo>::TIterator It(BlocksToUpdate); It; ++It)
 	{
-		UpdateBlock(Scene, DebugDrawingView, It.Value());
+		UpdateBlock(RHICmdList, Scene, DebugDrawingView, It.Value());
 	}
 }
 
@@ -838,7 +838,7 @@ void FIndirectLightingCache::SetLightingCacheDirty(FScene* Scene, const FPrecomp
 	}
 }
 
-void FIndirectLightingCache::UpdateBlock(FScene* Scene, FViewInfo* DebugDrawingView, FBlockUpdateInfo& BlockInfo)
+void FIndirectLightingCache::UpdateBlock(FRHICommandListBase& RHICmdList, FScene* Scene, FViewInfo* DebugDrawingView, FBlockUpdateInfo& BlockInfo)
 {
 	const int32 NumSamplesPerBlock = BlockInfo.Block.TexelSize * BlockInfo.Block.TexelSize * BlockInfo.Block.TexelSize;
 	FSHVectorRGB3 SingleSample;
@@ -894,9 +894,9 @@ void FIndirectLightingCache::UpdateBlock(FScene* Scene, FViewInfo* DebugDrawingV
 			BlockInfo.Block.TexelSize);
 
 		// Update the volume texture atlas
-		RHIUpdateTexture3D(GetTexture0(), 0, UpdateRegion, BlockInfo.Block.TexelSize * FormatSize, BlockInfo.Block.TexelSize * BlockInfo.Block.TexelSize * FormatSize, (const uint8*)Texture0Data.GetData());
-		RHIUpdateTexture3D(GetTexture1(), 0, UpdateRegion, BlockInfo.Block.TexelSize * FormatSize, BlockInfo.Block.TexelSize * BlockInfo.Block.TexelSize * FormatSize, (const uint8*)Texture1Data.GetData());
-		RHIUpdateTexture3D(GetTexture2(), 0, UpdateRegion, BlockInfo.Block.TexelSize * FormatSize, BlockInfo.Block.TexelSize * BlockInfo.Block.TexelSize * FormatSize, (const uint8*)Texture2Data.GetData());
+		RHICmdList.UpdateTexture3D(GetTexture0(), 0, UpdateRegion, BlockInfo.Block.TexelSize * FormatSize, BlockInfo.Block.TexelSize * BlockInfo.Block.TexelSize * FormatSize, (const uint8*)Texture0Data.GetData());
+		RHICmdList.UpdateTexture3D(GetTexture1(), 0, UpdateRegion, BlockInfo.Block.TexelSize * FormatSize, BlockInfo.Block.TexelSize * BlockInfo.Block.TexelSize * FormatSize, (const uint8*)Texture1Data.GetData());
+		RHICmdList.UpdateTexture3D(GetTexture2(), 0, UpdateRegion, BlockInfo.Block.TexelSize * FormatSize, BlockInfo.Block.TexelSize * BlockInfo.Block.TexelSize * FormatSize, (const uint8*)Texture2Data.GetData());
 	}
 	else
 	{
diff --git a/Engine/Source/Runtime/Renderer/Private/MobileShadingRenderer.cpp b/Engine/Source/Runtime/Renderer/Private/MobileShadingRenderer.cpp
--- a/Engine/Source/Runtime/Renderer/Private/MobileShadingRenderer.cpp
+++ b/Engine/Source/Runtime/Renderer/Private/MobileShadingRenderer.cpp
@@ -724,7 +724,7 @@ void FMobileSceneRenderer::InitViews(
 		// if we kicked off ILC update via task, wait and finalize.
 		if (ILCTaskData)
 		{
-			Scene->IndirectLightingCache.FinalizeCacheUpdates(Scene, *this, *ILCTaskData);
+			Scene->IndirectLightingCache.FinalizeCacheUpdates(RHICmdList, Scene, *this, *ILCTaskData);
 		}
 	}
 
diff --git a/Engine/Source/Runtime/Renderer/Private/ScenePrivate.h b/Engine/Source/Runtime/Renderer/Private/ScenePrivate.h
--- a/Engine/Source/Runtime/Renderer/Private/ScenePrivate.h
+++ b/Engine/Source/Runtime/Renderer/Private/ScenePrivate.h
@@ -2318,13 +2318,13 @@ public:
 	FIndirectLightingCacheAllocation* FindPrimitiveAllocation(FPrimitiveComponentId PrimitiveId) const;	
 
 	/** Updates indirect lighting in the cache based on visibility synchronously. */
-	void UpdateCache(FScene* Scene, FSceneRenderer& Renderer, bool bAllowUnbuiltPreview);
+	void UpdateCache(FRHICommandListBase& RHICmdList, FScene* Scene, FSceneRenderer& Renderer, bool bAllowUnbuiltPreview);
 
 	/** Starts a task to update the cache primitives.  Results and task ref returned in the FILCUpdatePrimTaskData structure */
 	void StartUpdateCachePrimitivesTask(FScene* Scene, FSceneRenderer& Renderer, bool bAllowUnbuiltPreview, FILCUpdatePrimTaskData& OutTaskData);
 
 	/** Wait on a previously started task and complete any block updates and debug draw */
-	void FinalizeCacheUpdates(FScene* Scene, FSceneRenderer& Renderer, FILCUpdatePrimTaskData& TaskData);
+	void FinalizeCacheUpdates(FRHICommandListBase& RHICmdList, FScene* Scene, FSceneRenderer& Renderer, FILCUpdatePrimTaskData& TaskData);
 
 	/** Force all primitive allocations to be re-interpolated. */
 	void SetLightingCacheDirty(FScene* Scene, const FPrecomputedLightVolume* Volume);
@@ -2344,7 +2344,7 @@ private:
 	void UpdateCachePrimitivesInternal(FScene* Scene, FSceneRenderer& Renderer, bool bAllowUnbuiltPreview, TMap<FIntVector, FBlockUpdateInfo>& OutBlocksToUpdate, TArray<FIndirectLightingCacheAllocation*>& OutTransitionsOverTimeToUpdate, TArray<FPrimitiveSceneInfo*>& OutPrimitivesToUpdateStaticMeshes);
 
 	/** Internal helper to perform blockupdates and transition updates on the results of UpdateCachePrimitivesInternal.  Must be on render thread. */
-	void FinalizeUpdateInternal_RenderThread(FScene* Scene, FSceneRenderer& Renderer, TMap<FIntVector, FBlockUpdateInfo>& BlocksToUpdate, const TArray<FIndirectLightingCacheAllocation*>& TransitionsOverTimeToUpdate, TArray<FPrimitiveSceneInfo*>& PrimitivesToUpdateStaticMeshes);
+	void FinalizeUpdateInternal_RenderThread(FRHICommandListBase& RHICmdList, FScene* Scene, FSceneRenderer& Renderer, TMap<FIntVector, FBlockUpdateInfo>& BlocksToUpdate, const TArray<FIndirectLightingCacheAllocation*>& TransitionsOverTimeToUpdate, TArray<FPrimitiveSceneInfo*>& PrimitivesToUpdateStaticMeshes);
 
 	/** Internal helper which adds an entry to the update lists for this allocation, if needed (due to movement, etc). Returns true if the allocation was updated or will be udpated */
 	bool UpdateCacheAllocation(
@@ -2370,7 +2370,7 @@ private:
 		TArray<FPrimitiveSceneInfo*>& PrimitivesToUpdateStaticMeshes);
 
 	/** Updates the contents of the volume texture blocks in BlocksToUpdate. */
-	void UpdateBlocks(FScene* Scene, FViewInfo* DebugDrawingView, TMap<FIntVector, FBlockUpdateInfo>& BlocksToUpdate);
+	void UpdateBlocks(FRHICommandListBase& RHICmdList, FScene* Scene, FViewInfo* DebugDrawingView, TMap<FIntVector, FBlockUpdateInfo>& BlocksToUpdate);
 
 	/** Updates any outstanding transitions with a new delta time. */
 	void UpdateTransitionsOverTime(const TArray<FIndirectLightingCacheAllocation*>& TransitionsOverTimeToUpdate, float DeltaWorldTime) const;
@@ -2390,7 +2390,7 @@ private:
 	 * Updates an allocation block in the cache, by re-interpolating values and uploading to the cache volume texture.
 	 * @param DebugDrawingView can be 0
 	 */
-	void UpdateBlock(FScene* Scene, FViewInfo* DebugDrawingView, FBlockUpdateInfo& Block);
+	void UpdateBlock(FRHICommandListBase& RHICmdList, FScene* Scene, FViewInfo* DebugDrawingView, FBlockUpdateInfo& Block);
 
 	/** Interpolates a single SH sample from all levels. */
 	void InterpolatePoint(
diff --git a/Engine/Source/Runtime/Renderer/Private/SceneVisibility.cpp b/Engine/Source/Runtime/Renderer/Private/SceneVisibility.cpp
--- a/Engine/Source/Runtime/Renderer/Private/SceneVisibility.cpp
+++ b/Engine/Source/Runtime/Renderer/Private/SceneVisibility.cpp
@@ -6020,13 +6020,13 @@ void FDeferredShadingSceneRenderer::EndInitViews(
 	{
 		QUICK_SCOPE_CYCLE_COUNTER(STAT_PostVisibilityFrameSetup_IndirectLightingCache_Update);
 		check(!TaskDatas.ILCUpdatePrim);
-		Scene->IndirectLightingCache.UpdateCache(Scene, *this, true);
+		Scene->IndirectLightingCache.UpdateCache(RHICmdList, Scene, *this, true);
 	}
 
 	// If we kicked off ILC update via task, wait and finalize.
 	if (TaskDatas.ILCUpdatePrim)
 	{
-		Scene->IndirectLightingCache.FinalizeCacheUpdates(Scene, *this, *TaskDatas.ILCUpdatePrim);
+		Scene->IndirectLightingCache.FinalizeCacheUpdates(RHICmdList, Scene, *this, *TaskDatas.ILCUpdatePrim);
 	}
 
 	{
