Deprecate r.Shadow.Virtual.NonNaniteVSM, replace with r.Shadow.Virtual.NonNanite.ProjectEnabled with slightly different semantics.
Remove r.Shadow.Virtual.ForceOnlyVirtualShadowMaps (already deprecated). We have not supported the VSM+classic SM for a single light path for a long while now.
Remove a few other previously deprecated cvars.
Minor cleanup of a few comments and defines.

diff --git a/Engine/Source/Runtime/Engine/Private/LightSceneProxy.cpp b/Engine/Source/Runtime/Engine/Private/LightSceneProxy.cpp
--- a/Engine/Source/Runtime/Engine/Private/LightSceneProxy.cpp
+++ b/Engine/Source/Runtime/Engine/Private/LightSceneProxy.cpp
@@ -81,10 +81,8 @@ FLightSceneProxy::FLightSceneProxy(const ULightComponent* InLightComponent)
 
 	// Currently we use virtual shadows maps for all lights when the global setting is enabled
 	bUseVirtualShadowMaps = ::UseVirtualShadowMaps(SceneInterface->GetShaderPlatform(), SceneInterface->GetFeatureLevel());
-
-	// Treat stationary lights as movable when non-nanite VSMs are enabled
-	const bool bNonNaniteVirtualShadowMaps = UseNonNaniteVirtualShadowMaps(SceneInterface->GetShaderPlatform(), SceneInterface->GetFeatureLevel());
-	if (bNonNaniteVirtualShadowMaps)
+	// Treat stationary lights as movable when VSMs are enabled
+	if (bUseVirtualShadowMaps)
 	{
 		bStaticShadowing = bStaticLighting;
 	}
diff --git a/Engine/Source/Runtime/RenderCore/Private/RenderUtils.cpp b/Engine/Source/Runtime/RenderCore/Private/RenderUtils.cpp
--- a/Engine/Source/Runtime/RenderCore/Private/RenderUtils.cpp
+++ b/Engine/Source/Runtime/RenderCore/Private/RenderUtils.cpp
@@ -60,6 +60,18 @@ static FAutoConsoleVariableRef CVarRayTracingProxiesProjectEnabled(
 	TEXT("Whether ray tracing proxies are generated by default for this project."),
 	ECVF_ReadOnly);
 
+// This is a per-project main switch for Non-Nanite support in Virtual Shadow Maps.
+// Disabling this can remove shader permutations and slightly improve performance, but Non-Nanite objects
+// will not cast shadows when Virtual Shadow Maps are enabled.
+int32 GVirtualShadowMapNonNaniteProjectEnabled = 1;
+FAutoConsoleVariableRef CVarVirtualShadowMapNonNaniteProjectEnabled(
+	TEXT("r.Shadow.Virtual.NonNanite.ProjectEnabled"),
+	GVirtualShadowMapNonNaniteProjectEnabled,
+	TEXT("This can be used to disable non-nanite rendering into Virtual Shadow Maps and related shader permutations.\n")
+	TEXT("Only use this if you do not need any non-nanite geometry to cast shadows."),
+	ECVF_RenderThreadSafe | ECVF_ReadOnly
+);
+
 const uint16 GCubeIndices[12*3] =
 {
 	0, 2, 3,
@@ -1442,14 +1454,20 @@ bool DoesPlatformSupportVirtualShadowMaps(EShaderPlatform Platform)
 
 bool DoesPlatformSupportNonNaniteVirtualShadowMaps(EShaderPlatform ShaderPlatform)
 {
-	static const auto EnableCVar = IConsoleManager::Get().FindConsoleVariable(TEXT("r.Shadow.Virtual.NonNaniteVSM"));
-	return EnableCVar->GetInt() != 0 && DoesPlatformSupportNanite(ShaderPlatform);
+	if (GVirtualShadowMapNonNaniteProjectEnabled == 0)
+	{
+		return false;
+	}
+	return DoesPlatformSupportVirtualShadowMaps(ShaderPlatform);
 }
 
 bool UseNonNaniteVirtualShadowMaps(EShaderPlatform ShaderPlatform, const FStaticFeatureLevel FeatureLevel)
 {
-	static const auto EnableCVar = IConsoleManager::Get().FindConsoleVariable(TEXT("r.Shadow.Virtual.NonNaniteVSM"));
-	return EnableCVar->GetInt() != 0 && UseVirtualShadowMaps(ShaderPlatform, FeatureLevel);
+	if (GVirtualShadowMapNonNaniteProjectEnabled == 0)
+	{
+		return false;
+	}
+	return UseVirtualShadowMaps(ShaderPlatform, FeatureLevel);
 }
 
 bool IsWaterVirtualShadowMapFilteringEnabled(const FStaticShaderPlatform Platform)
diff --git a/Engine/Source/Runtime/RenderCore/Public/RenderUtils.h b/Engine/Source/Runtime/RenderCore/Public/RenderUtils.h
--- a/Engine/Source/Runtime/RenderCore/Public/RenderUtils.h
+++ b/Engine/Source/Runtime/RenderCore/Public/RenderUtils.h
@@ -457,7 +457,7 @@ RENDERCORE_API bool UseVirtualShadowMaps(EShaderPlatform ShaderPlatform, const F
 RENDERCORE_API bool DoesPlatformSupportVirtualShadowMaps(EShaderPlatform Platform);
 
 /**
-* Returns true if non-Nanite virtual shadow maps are enabled by CVar r.Shadow.Virtual.NonNaniteVSM
+* Returns true if non-Nanite virtual shadow maps are enabled by CVar r.Shadow.Virtual.NonNanite.ProjectEnabled
 * and the runtime supports Nanite/virtual shadow maps.
 */
 RENDERCORE_API bool DoesPlatformSupportNonNaniteVirtualShadowMaps(EShaderPlatform ShaderPlatform);
diff --git a/Engine/Source/Runtime/Renderer/Private/DeferredShadingRenderer.cpp b/Engine/Source/Runtime/Renderer/Private/DeferredShadingRenderer.cpp
--- a/Engine/Source/Runtime/Renderer/Private/DeferredShadingRenderer.cpp
+++ b/Engine/Source/Runtime/Renderer/Private/DeferredShadingRenderer.cpp
@@ -1800,7 +1800,6 @@ void FDeferredShadingSceneRenderer::Render(FRDGBuilder& GraphBuilder, const FSce
 	
 				// Important that this uses consistent logic throughout the frame, so evaluate once and pass in the flag from here
 				// NOTE: Must be done after  system texture initialization
-				// TODO: This doesn't take into account the potential for split screen views with separate shadow caches
 				const bool bEnableVirtualShadowMaps = UseVirtualShadowMaps(ShaderPlatform, FeatureLevel) && ViewFamily.EngineShowFlags.DynamicShadows && !bHasRayTracedOverlay;
 				VirtualShadowMapArray.Initialize(GraphBuilder, Scene->GetVirtualShadowMapCache(), bEnableVirtualShadowMaps, ViewFamily.EngineShowFlags);
 	
diff --git a/Engine/Source/Runtime/Renderer/Private/SceneRendering.cpp b/Engine/Source/Runtime/Renderer/Private/SceneRendering.cpp
--- a/Engine/Source/Runtime/Renderer/Private/SceneRendering.cpp
+++ b/Engine/Source/Runtime/Renderer/Private/SceneRendering.cpp
@@ -4184,7 +4184,7 @@ void FSceneRenderer::OnRenderFinish(FRDGBuilder& GraphBuilder, FRDGTextureRef Vi
 		const bool bRealTimeSkyCaptureButNothingToCapture = Scene->SkyLight && Scene->SkyLight->bRealTimeCaptureEnabled && (!Scene->HasSkyAtmosphere() && !Scene->HasVolumetricCloud() && (Views.Num() > 0 && !Views[0].bSceneHasSkyMaterial));
 
 		// Point light shadows are disabled by default on mobile platforms.
-		const bool bShowPointLightWarning = !IsMobilePlatform(ShaderPlatform) ? UsedWholeScenePointLightNames.Num() > 0 && !FReadOnlyCVARCache::EnablePointLightShadows(ShaderPlatform) : false;
+		const bool bShowPointLightWarning = !IsMobilePlatform(ShaderPlatform) ? bUsedPointLightWholeSceneShadow && !FReadOnlyCVARCache::EnablePointLightShadows(ShaderPlatform) : false;
 		const bool bShowShadowedLightOverflowWarning = Scene->OverflowingDynamicShadowedLights.Num() > 0;
 
 		const bool bLocalFogVolumeInSceneButProjectDisabled = Scene->HasAnyLocalFogVolume() && !ProjectSupportsLocalFogVolumes();
@@ -4392,12 +4392,8 @@ void FSceneRenderer::OnRenderFinish(FRDGBuilder& GraphBuilder, FRDGTextureRef Vi
 						}
 						if (bShowPointLightWarning)
 						{
-							static const FText Message = NSLOCTEXT("Renderer", "PointLight", "PROJECT DOES NOT SUPPORT WHOLE SCENE POINT LIGHT SHADOWS: ");
+							static const FText Message = NSLOCTEXT("Renderer", "PointLight", "A shadow casting point light is present but the project does not support whole scene point light shadows");
 							Writer.DrawLine(Message);
-							for (const FString& LightName : UsedWholeScenePointLightNames)
-							{
-								Writer.DrawLine(FText::FromString(LightName), 35);
-							}
 						}
 						if (bShowShadowedLightOverflowWarning)
 						{
diff --git a/Engine/Source/Runtime/Renderer/Private/SceneRendering.h b/Engine/Source/Runtime/Renderer/Private/SceneRendering.h
--- a/Engine/Source/Runtime/Renderer/Private/SceneRendering.h
+++ b/Engine/Source/Runtime/Renderer/Private/SceneRendering.h
@@ -2179,8 +2179,9 @@ public:
 	bool bUsedPrecomputedVisibility;
 
 #if !(UE_BUILD_SHIPPING || UE_BUILD_TEST)
-	/** Lights added if wholescenepointlight shadow would have been rendered (ignoring r.SupportPointLightWholeSceneShadows). Used for warning about unsupported features. */
-	TArray<FString, SceneRenderingAllocator> UsedWholeScenePointLightNames;
+	// Set to true if a whole scene point light shadow is present
+	// Used for warnings about unsupported features
+	bool bUsedPointLightWholeSceneShadow = false;
 #endif
 	/** Feature level being rendered */
 	ERHIFeatureLevel::Type FeatureLevel;
diff --git a/Engine/Source/Runtime/Renderer/Private/ShadowDepthRendering.cpp b/Engine/Source/Runtime/Renderer/Private/ShadowDepthRendering.cpp
--- a/Engine/Source/Runtime/Renderer/Private/ShadowDepthRendering.cpp
+++ b/Engine/Source/Runtime/Renderer/Private/ShadowDepthRendering.cpp
@@ -145,7 +145,6 @@ TAutoConsoleVariable<int32> CVarNaniteShadowsUpdateStreaming(
 	ECVF_RenderThreadSafe);
 
 extern int32 GNaniteShowStats;
-extern int32 GEnableNonNaniteVSM;
 
 namespace Nanite
 {
@@ -2384,7 +2383,6 @@ FMeshPassProcessor* CreateVSMShadowDepthPassProcessor(ERHIFeatureLevel::Type Fea
 {
 	const EShaderPlatform ShaderPlatform = GetFeatureLevelShaderPlatform(FeatureLevel);
 
-	// Only create the mesh pass processor if VSMs are not enabled as this prevents wasting time caching the SM draw commands
 	if (UseNonNaniteVirtualShadowMaps(ShaderPlatform, FeatureLevel))
 	{
 		return new FShadowDepthPassMeshProcessor(
diff --git a/Engine/Source/Runtime/Renderer/Private/ShadowSetup.cpp b/Engine/Source/Runtime/Renderer/Private/ShadowSetup.cpp
--- a/Engine/Source/Runtime/Renderer/Private/ShadowSetup.cpp
+++ b/Engine/Source/Runtime/Renderer/Private/ShadowSetup.cpp
@@ -61,8 +61,6 @@ static FAutoConsoleVariableRef CVarMinScreenRadiusForShadowCaster(
 	ECVF_Scalability | ECVF_RenderThreadSafe
 	);
 
-extern int GEnableNonNaniteVSM;
-
 int32 GCacheWholeSceneShadows = 1;
 FAutoConsoleVariableRef CVarCacheWholeSceneShadows(
 	TEXT("r.Shadow.CacheWholeSceneShadows"),
@@ -341,28 +339,6 @@ static TAutoConsoleVariable<int32> CVarClipmapUseConservativeCulling(
 	ECVF_RenderThreadSafe
 );
 
-static TAutoConsoleVariable<int32> CVarForceOnlyVirtualShadowMaps(
-	TEXT("r.Shadow.Virtual.ForceOnlyVirtualShadowMaps"),
-	1,
-	TEXT("If enabled, disallow creation of conventional non-virtual shadow maps for any lights that get a virtual shadow map.\n")
-	TEXT("This variable should be enabled for optimal performance and memory usage, but geometric primitives that cannot be rendered into the virtual shadow map will not cast shadows.\n")
-	TEXT("Disabling may reduce performance. This variable is deprecated and may be removed in a future version."),
-	FConsoleVariableDelegate::CreateLambda(
-		[](IConsoleVariable* CVar)
-		{
-			bool bForceOnlyVirtualShadowMaps = CVar->GetInt() != 0;
-			if (!bForceOnlyVirtualShadowMaps)
-			{
-				UE_LOG(LogRenderer, Warning, 
-					   TEXT("Disabling r.Shadow.Virtual.ForceOnlyVirtualShadowMaps may reduce performance. ")
-					   TEXT("This variable is deprecated and may be removed in a future version."));
-			}
-		}
-	),
-	ECVF_RenderThreadSafe
-);
-
-
 CSV_DECLARE_CATEGORY_EXTERN(LightCount);
 
 #if !UE_BUILD_SHIPPING
@@ -2139,20 +2115,15 @@ bool FProjectedShadowInfo::TestPrimitiveFarCascadeConditions(bool bPrimitiveCast
 		return bPrimitiveCastsFarShadow;
 	}
 
-	if (GEnableNonNaniteVSM != 0 && bVSM)
+	if (bVSM && IsWholeSceneDirectionalShadow() && MaxNonFarCascadeDistance > 0.0f)
 	{
-		const bool bWholeSceneDirectionalShadow = IsWholeSceneDirectionalShadow();
+		check(DependentView);
 
-		if (bWholeSceneDirectionalShadow && MaxNonFarCascadeDistance > 0.0f)
+		FVector ViewOrigin = DependentView->GetShadowViewMatrices().GetViewOrigin();
+		float ViewDistance = (Bounds.Origin - ViewOrigin) | DependentView->GetViewDirection();
+		if (ViewDistance - Bounds.SphereRadius > MaxNonFarCascadeDistance)
 		{
-			check(DependentView);
-
-			FVector ViewOrigin = DependentView->GetShadowViewMatrices().GetViewOrigin();
-			float ViewDistance = (Bounds.Origin - ViewOrigin) | DependentView->GetViewDirection();
-			if (ViewDistance - Bounds.SphereRadius > MaxNonFarCascadeDistance)
-			{
-				return bPrimitiveCastsFarShadow;
-			}
+			return bPrimitiveCastsFarShadow;
 		}
 	}
 	return true;
@@ -3235,6 +3206,8 @@ void FSceneRenderer::SetupInteractionShadows(
 	const TArray<FProjectedShadowInfo*,SceneRenderingAllocator>& ViewDependentWholeSceneShadows,
 	TArray<FProjectedShadowInfo*,SceneRenderingAllocator>& PreShadows)
 {
+	check(!UseVirtualShadowMaps(ShaderPlatform));
+
 	// too high on hit count to leave on
 	// SCOPE_CYCLE_COUNTER(STAT_SetupInteractionShadows);
 
@@ -3254,7 +3227,7 @@ void FSceneRenderer::SetupInteractionShadows(
 	if (!bShadowHandledByParent)
 	{
 		const bool bCreateTranslucentObjectShadow = GUseTranslucencyShadowDepths && Interaction->HasTranslucentObjectShadow();
-		const bool bCreateInsetObjectShadow = Interaction->HasInsetObjectShadow() && !UseNonNaniteVirtualShadowMaps(ShaderPlatform, FeatureLevel);
+		const bool bCreateInsetObjectShadow = Interaction->HasInsetObjectShadow();
 		const bool bCreateObjectShadowForStationaryLight = ShouldCreateObjectShadowForStationaryLight(Interaction->GetLight(), PrimitiveSceneInfo->Proxy, Interaction->IsShadowMapped());
 
 		if (Interaction->HasShadow() 
@@ -3715,7 +3688,6 @@ void ComputeWholeSceneShadowCacheModes(
 	float ActualDesiredResolution,
 	const FIntPoint& MaxShadowResolution,
 	FScene* Scene,
-	bool bNeedsVirtualShadowMap,
 	FWholeSceneProjectedShadowInitializer& InOutProjectedShadowInitializer,
 	const int64 CachedShadowMapsSize,
 	FIntPoint& InOutShadowMapSize,
@@ -3757,7 +3729,7 @@ void ComputeWholeSceneShadowCacheModes(
 			FCachedShadowMapData* CachedShadowMapData = &(*CachedShadowMapDatas)[0];
 
 			if (InOutProjectedShadowInitializer.IsCachedShadowValid(CachedShadowMapData->Initializer) 
-				&& CachedShadowMapData->bCachedShadowMapHasNaniteGeometry == !bNeedsVirtualShadowMap)
+				&& CachedShadowMapData->bCachedShadowMapHasNaniteGeometry)
 			{
 				if (CachedShadowMapData->ShadowMap.IsValid() && CachedShadowMapData->ShadowMap.GetSize() == InOutShadowMapSize)
 				{
@@ -4220,12 +4192,7 @@ void FSceneRenderer::CreateWholeSceneProjectedShadow(
 			for (int32 ShadowIndex = 0, ShadowCount = ProjectedShadowInitializers.Num(); ShadowIndex < ShadowCount; ShadowIndex++)
 			{
 				FWholeSceneProjectedShadowInitializer& ProjectedShadowInitializer = ProjectedShadowInitializers[ShadowIndex];
-
-				const bool bNeedsVirtualShadowMap =
-					LightSceneInfo->Proxy->UseVirtualShadowMaps() &&
-					VirtualShadowMapArray.IsEnabled() &&
-					!ProjectedShadowInitializer.bRayTracedDistanceField;
-
+				
 				// Consolidate repeated operation into a single site.
 				auto AddInteractingPrimitives = [bStaticSceneOnly, &Views = this->Views, FeatureLevel = this->FeatureLevel, &TaskData]
 				(FLightPrimitiveInteraction* InteractionList, 
@@ -4266,51 +4233,53 @@ void FSceneRenderer::CreateWholeSceneProjectedShadow(
 						}
 					}
 				};
-
-				if (bNeedsVirtualShadowMap)
+		
+				if (LightSceneInfo->Proxy->UseVirtualShadowMaps())
 				{
-					// Create the projected shadow info.
-					FProjectedShadowInfo* ProjectedShadowInfo = Allocator.Create<FProjectedShadowInfo>();
+					// VSMs should be enabled globally and proxy should never ask for more than one
+					if (ensure(VirtualShadowMapArray.IsEnabled() &&	ProjectedShadowInitializers.Num() == 1))
+					{
+						// Create the projected shadow info.
+						FProjectedShadowInfo* ProjectedShadowInfo = Allocator.Create<FProjectedShadowInfo>();
 
-					int32 VirtualSizeX = FVirtualShadowMap::VirtualMaxResolutionXY;
-					int32 VirtualSizeY = FVirtualShadowMap::VirtualMaxResolutionXY;						
+						int32 VirtualSizeX = FVirtualShadowMap::VirtualMaxResolutionXY;
+						int32 VirtualSizeY = FVirtualShadowMap::VirtualMaxResolutionXY;						
 					
-					ProjectedShadowInfo->SetupWholeSceneProjection(
-						LightSceneInfo,
-						NULL,
-						ProjectedShadowInitializer,
-						VirtualSizeX,
-						VirtualSizeY,
-						VirtualSizeX,
-						VirtualSizeY,
-						0 // no border
-					);
-					FShadowSceneRenderer& ShadowSceneRenderer = GetSceneExtensionsRenderers().GetRenderer<FShadowSceneRenderer>();
+						ProjectedShadowInfo->SetupWholeSceneProjection(
+							LightSceneInfo,
+							NULL,
+							ProjectedShadowInitializer,
+							VirtualSizeX,
+							VirtualSizeY,
+							VirtualSizeX,
+							VirtualSizeY,
+							0 // no border
+						);
+						FShadowSceneRenderer& ShadowSceneRenderer = GetSceneExtensionsRenderers().GetRenderer<FShadowSceneRenderer>();
 
-					TSharedPtr<FVirtualShadowMapPerLightCacheEntry> VirtualSmPerLightCacheEntry = ShadowSceneRenderer.AddLocalLightShadow(ProjectedShadowInitializer, ProjectedShadowInfo, LightSceneInfo, MaxScreenRadius);
+						TSharedPtr<FVirtualShadowMapPerLightCacheEntry> VirtualSmPerLightCacheEntry = ShadowSceneRenderer.AddLocalLightShadow(ProjectedShadowInitializer, ProjectedShadowInfo, LightSceneInfo, MaxScreenRadius);
 
-					ProjectedShadowInfo->bVSM = true;
-					ProjectedShadowInfo->MeshSelectionMask = EShadowMeshSelection::VSM;
+						ProjectedShadowInfo->bVSM = true;
+						ProjectedShadowInfo->MeshSelectionMask = EShadowMeshSelection::VSM;
 
-					// Initialize to true if we are not creating LPIs for Nanite.
-					bool bContainsNaniteSubjects = ShouldSkipNaniteLPIs(TaskData.Scene->GetShaderPlatform());
+						// Initialize to true if we are not creating LPIs for Nanite.
+						bool bContainsNaniteSubjects = ShouldSkipNaniteLPIs(TaskData.Scene->GetShaderPlatform());
 
-					// Skip mesh setup if it won't be rendered anyway
-					if (ProjectedShadowInfo->bShouldRenderVSM)
-					{
-						// Skip convex hull tests for VSM since this causes artifacts due to caching (potentially check if caching is enabled but may lead to race condition).
-						// The interaction setup has already tested the light bounds (by calling FLightSceneProxy::AffectsBounds)
-						FLightViewFrustumConvexHulls EmptyHull{};
-						AddInteractingPrimitives(LightSceneInfo->GetDynamicInteractionOftenMovingPrimitiveList(), ProjectedShadowInfo, bContainsNaniteSubjects, EmptyHull, VirtualSmPerLightCacheEntry);
-						AddInteractingPrimitives(LightSceneInfo->GetDynamicInteractionStaticPrimitiveList(), ProjectedShadowInfo, bContainsNaniteSubjects, EmptyHull, VirtualSmPerLightCacheEntry);
-						ProjectedShadowInfo->bContainsNaniteSubjects = bContainsNaniteSubjects;
-					}
+						// Skip mesh setup if it won't be rendered anyway
+						if (ProjectedShadowInfo->bShouldRenderVSM)
+						{
+							// Skip convex hull tests for VSM since this causes artifacts due to caching (potentially check if caching is enabled but may lead to race condition).
+							// The interaction setup has already tested the light bounds (by calling FLightSceneProxy::AffectsBounds)
+							FLightViewFrustumConvexHulls EmptyHull{};
+							AddInteractingPrimitives(LightSceneInfo->GetDynamicInteractionOftenMovingPrimitiveList(), ProjectedShadowInfo, bContainsNaniteSubjects, EmptyHull, VirtualSmPerLightCacheEntry);
+							AddInteractingPrimitives(LightSceneInfo->GetDynamicInteractionStaticPrimitiveList(), ProjectedShadowInfo, bContainsNaniteSubjects, EmptyHull, VirtualSmPerLightCacheEntry);
+							ProjectedShadowInfo->bContainsNaniteSubjects = bContainsNaniteSubjects;
+						}
 
-					VisibleLightInfo.AllProjectedShadows.Add(ProjectedShadowInfo);
+						VisibleLightInfo.AllProjectedShadows.Add(ProjectedShadowInfo);
+					}
 				}
-				
-				// If using *only* virtual shadow maps via forced CVar, skip creation of the non-virtual shadow map
-				if (!bNeedsVirtualShadowMap || CVarForceOnlyVirtualShadowMaps.GetValueOnRenderThread() == 0)
+				else // No VSM, fall back to classic shadow map path
 				{
 					// Round down to the nearest power of two so that resolution changes are always doubling or halving the resolution, which increases filtering stability
 					// Use the max resolution if the desired resolution is larger than that
@@ -4341,7 +4310,6 @@ void FSceneRenderer::CreateWholeSceneProjectedShadow(
 							MaxDesiredResolution,
 							FIntPoint(MaxShadowResolution, MaxShadowResolutionY),
 							Scene,
-							bNeedsVirtualShadowMap,
 							// Below are in-out or out parameters. They can change
 							ProjectedShadowInitializer,
 							CachedShadowMapsSize,
@@ -4374,15 +4342,6 @@ void FSceneRenderer::CreateWholeSceneProjectedShadow(
 						ProjectedShadowInfo->CacheMode = CacheMode[CacheModeIndex];
 						ProjectedShadowInfo->FadeAlphas = FadeAlphas;
 
-						// If we have a virtual shadow map, disable nanite rendering into the regular shadow map or else we'd get double-shadowing
-						// and also filter out any VSM-supporting meshes
-						if (bNeedsVirtualShadowMap)
-						{
-							ProjectedShadowInfo->bNaniteGeometry = false;
-							// Note: the default is all types
-							ProjectedShadowInfo->MeshSelectionMask = EShadowMeshSelection::SM;
-						}
-
 						bool bContainsNaniteSubjects = ShouldSkipNaniteLPIs(TaskData.Scene->GetShaderPlatform());
 
 						// Ray traced shadows use the GPU managed distance field object buffers, no CPU culling should be used
@@ -4916,7 +4875,7 @@ struct FGatherShadowPrimitivesPacket
 
 	bool DoesPrimitiveCastInsetShadow(FDynamicShadowsTaskData& TaskData, const FPrimitiveSceneInfo* PrimitiveSceneInfo, const FPrimitiveSceneProxy* PrimitiveProxy) const
 	{
-		if (UseNonNaniteVirtualShadowMaps(TaskData.ShaderPlatform, TaskData.FeatureLevel))
+		if (UseVirtualShadowMaps(TaskData.ShaderPlatform, TaskData.FeatureLevel))
 		{
 			return false;
 		}
@@ -5478,13 +5437,7 @@ void FSceneRenderer::AddViewDependentWholeSceneShadowsForView(
 {
 	SCOPE_CYCLE_COUNTER(STAT_AddViewDependentWholeSceneShadowsForView);
 
-	// Note: it is possible for a non-directional light to set up a proxy that requests a view-dependent SM (or several)
-	// Unmodified UE does not use this, so the virtual SM ignores this path for the time being (more likely to be removed).
-	const bool bDirectionalLight = LightSceneInfo.Proxy->GetLightType() == LightType_Directional;
-	const bool bNeedsVirtualShadowMap =
-		LightSceneInfo.Proxy->UseVirtualShadowMaps() &&
-		VirtualShadowMapArray.IsEnabled() &&
-		bDirectionalLight;
+	const bool bNeedsVirtualShadowMap = LightSceneInfo.Proxy->UseVirtualShadowMaps();
 
 	// Allow each view to create a whole scene view dependent shadow
 	for (int32 ViewIndex = 0; ViewIndex < Views.Num(); ViewIndex++)
@@ -5524,10 +5477,8 @@ void FSceneRenderer::AddViewDependentWholeSceneShadowsForView(
 
 			float MaxNonFarCascadeDistance = 0.0f;
 
-
-			bool bForceOnlyVirtualShadowMaps = CVarForceOnlyVirtualShadowMaps.GetValueOnRenderThread() != 0;
 			bool bVsmUseFarShadowRules = CVarVsmUseFarShadowRules.GetValueOnRenderThread() != 0;
-			if (!bNeedsVirtualShadowMap || !bForceOnlyVirtualShadowMaps || bVsmUseFarShadowRules)
+			if (!bNeedsVirtualShadowMap || bVsmUseFarShadowRules)
 			{
 				for (int32 Index = 0; Index < ProjectionCount; Index++)
 				{
@@ -5549,9 +5500,9 @@ void FSceneRenderer::AddViewDependentWholeSceneShadowsForView(
 							MaxNonFarCascadeDistance = FMath::Max(MaxNonFarCascadeDistance, ProjectedShadowInitializer.CascadeSettings.SplitFar);
 						}
 
-						if (bNeedsVirtualShadowMap && bForceOnlyVirtualShadowMaps)
+						if (bNeedsVirtualShadowMap)
 						{
-							// ForceOnlyVirtualShadowMaps disables CSM, but DF shadows can still be enabled
+							// VSM disables CSM, but DF shadows can still be enabled
 							if (!ProjectedShadowInitializer.bRayTracedDistanceField)
 							{
 								continue;
@@ -5603,14 +5554,6 @@ void FSceneRenderer::AddViewDependentWholeSceneShadowsForView(
 							VisibleLightInfo.AllProjectedShadows.Add(ProjectedShadowInfo);
 							ShadowInfos.Add(ProjectedShadowInfo);
 
-							if (bNeedsVirtualShadowMap && !ProjectedShadowInitializer.bRayTracedDistanceField)
-							{
-								// If we have a virtual shadow map, disable nanite rendering into the regular shadow map or else we'd get double-shadowing
-								ProjectedShadowInfo->bNaniteGeometry = false;
-								// Note: the default is all types
-								ProjectedShadowInfo->MeshSelectionMask = EShadowMeshSelection::SM;
-							}
-
 							// Ray traced shadows use the GPU managed distance field object buffers, no CPU culling needed
 							if (!ProjectedShadowInfo->bRayTracedDistanceField)
 							{
@@ -5704,10 +5647,16 @@ void FSceneRenderer::AddViewDependentWholeSceneShadowsForView(
 				}
 			}
 
-			if (bNeedsVirtualShadowMap)
+			// Note: it is possible for a non-directional light to set up a proxy that requests a view-dependent SM (or several)
+			// Unmodified UE does not use this, so the virtual SM ignores this path for the time being (more likely to be removed).
+			const bool bDirectionalLight = LightSceneInfo.Proxy->GetLightType() == LightType_Directional;
+			if (bNeedsVirtualShadowMap && bDirectionalLight)
 			{
-				FShadowSceneRenderer& ShadowSceneRenderer = GetSceneExtensionsRenderers().GetRenderer<FShadowSceneRenderer>();
-				ShadowSceneRenderer.AddDirectionalLightShadow(LightSceneInfo, View, MaxNonFarCascadeDistance, ShadowInfosThatNeedCulling);
+				if (ensure(VirtualShadowMapArray.IsEnabled()))
+				{
+					FShadowSceneRenderer& ShadowSceneRenderer = GetSceneExtensionsRenderers().GetRenderer<FShadowSceneRenderer>();
+					ShadowSceneRenderer.AddDirectionalLightShadow(LightSceneInfo, View, MaxNonFarCascadeDistance, ShadowInfosThatNeedCulling);
+				}
 			}
 		}
 	}
@@ -6294,11 +6243,7 @@ void FSceneRenderer::CreateDynamicShadows(FDynamicShadowsTaskData& TaskData)
 							&& (!bMobile || bDirectionalLightShadow);
 
 #if !(UE_BUILD_SHIPPING || UE_BUILD_TEST)
-						const bool bPointLightWholeSceneShadow = (bShouldCreateShadowForMovableLight || bShouldCreateShadowForOverflowStaticShadowing || bShouldCreateShadowToPreviewStaticLight) && bPointLightShadow;
-						if (bPointLightWholeSceneShadow)
-						{
-							UsedWholeScenePointLightNames.Add(LightSceneInfoCompact.LightSceneInfo->Proxy->GetOwnerNameOrLabel());
-						}
+						bUsedPointLightWholeSceneShadow = bUsedPointLightWholeSceneShadow || ((bShouldCreateShadowForMovableLight || bShouldCreateShadowForOverflowStaticShadowing || bShouldCreateShadowToPreviewStaticLight) && bPointLightShadow);
 #endif
 
 						if (bCreateShadowForMovableLight || bCreateShadowToPreviewStaticLight || bCreateShadowForOverflowStaticShadowing)
@@ -6327,33 +6272,36 @@ void FSceneRenderer::CreateDynamicShadows(FDynamicShadowsTaskData& TaskData)
 							
 							if (!bMobile || bMobileInsetShadows || IsMobileCapsuleShadowsEnabled(ShaderPlatform) || (LightSceneInfo->Proxy->CastsModulatedShadows() && !LightSceneInfo->Proxy->UseCSMForDynamicObjects() && LightSceneInfo->Proxy->HasStaticShadowing()))
 							{
-								const TArray<FLightPrimitiveInteraction*>* InteractionShadowPrimitives = LightSceneInfo->GetInteractionShadowPrimitives();
-
-								if (InteractionShadowPrimitives)
+								if (!LightSceneInfo->Proxy->UseVirtualShadowMaps())
 								{
-									const int32 NumPrims = InteractionShadowPrimitives->Num();
-									for (int32 Idx = 0; Idx < NumPrims; ++Idx)
+									const TArray<FLightPrimitiveInteraction*>* InteractionShadowPrimitives = LightSceneInfo->GetInteractionShadowPrimitives();
+
+									if (InteractionShadowPrimitives)
 									{
-										SetupInteractionShadows(TaskData, (*InteractionShadowPrimitives)[Idx], VisibleLightInfo, bStaticSceneOnly, ViewDependentWholeSceneShadows, PreShadows);
+										const int32 NumPrims = InteractionShadowPrimitives->Num();
+										for (int32 Idx = 0; Idx < NumPrims; ++Idx)
+										{
+											SetupInteractionShadows(TaskData, (*InteractionShadowPrimitives)[Idx], VisibleLightInfo, bStaticSceneOnly, ViewDependentWholeSceneShadows, PreShadows);
+										}
 									}
-								}
-								else
-								{
-									// Look for individual primitives with a dynamic shadow.
-									for (FLightPrimitiveInteraction* Interaction = LightSceneInfo->GetDynamicInteractionOftenMovingPrimitiveList();
-										Interaction;
-										Interaction = Interaction->GetNextPrimitive()
-										)
+									else
 									{
-										SetupInteractionShadows(TaskData, Interaction, VisibleLightInfo, bStaticSceneOnly, ViewDependentWholeSceneShadows, PreShadows);
-									}
+										// Look for individual primitives with a dynamic shadow.
+										for (FLightPrimitiveInteraction* Interaction = LightSceneInfo->GetDynamicInteractionOftenMovingPrimitiveList();
+											Interaction;
+											Interaction = Interaction->GetNextPrimitive()
+											)
+										{
+											SetupInteractionShadows(TaskData, Interaction, VisibleLightInfo, bStaticSceneOnly, ViewDependentWholeSceneShadows, PreShadows);
+										}
 
-									for (FLightPrimitiveInteraction* Interaction = LightSceneInfo->GetDynamicInteractionStaticPrimitiveList();
-										Interaction;
-										Interaction = Interaction->GetNextPrimitive()
-										)
-									{
-										SetupInteractionShadows(TaskData, Interaction, VisibleLightInfo, bStaticSceneOnly, ViewDependentWholeSceneShadows, PreShadows);
+										for (FLightPrimitiveInteraction* Interaction = LightSceneInfo->GetDynamicInteractionStaticPrimitiveList();
+											Interaction;
+											Interaction = Interaction->GetNextPrimitive()
+											)
+										{
+											SetupInteractionShadows(TaskData, Interaction, VisibleLightInfo, bStaticSceneOnly, ViewDependentWholeSceneShadows, PreShadows);
+										}
 									}
 								}
 							}
diff --git a/Engine/Source/Runtime/Renderer/Private/Shadows/ShadowSceneRenderer.cpp b/Engine/Source/Runtime/Renderer/Private/Shadows/ShadowSceneRenderer.cpp
--- a/Engine/Source/Runtime/Renderer/Private/Shadows/ShadowSceneRenderer.cpp
+++ b/Engine/Source/Runtime/Renderer/Private/Shadows/ShadowSceneRenderer.cpp
@@ -43,24 +43,6 @@ static TAutoConsoleVariable<int32> CVarDistantLightMode(
 	ECVF_RenderThreadSafe
 );
 
-static TAutoConsoleVariable<float> CVarDistantLightForceCacheFootprintFraction(
-	TEXT("r.Shadow.Virtual.DistantLightForceCacheFootprintFraction_WillBeRemoved"),
-	0.0f,
-	TEXT("Will be removed."),
-	ECVF_Scalability | ECVF_RenderThreadSafe
-);
-static FAutoConsoleVariableDeprecated CVarDistantLightForceCacheFootprintFraction_Deprecated(TEXT("r.Shadow.Virtual.DistantLightForceCacheFootprintFraction"), TEXT("r.Shadow.Virtual.DistantLightForceCacheFootprintFraction_WillBeRemoved"), TEXT("5.7"));
-
-static TAutoConsoleVariable<bool> CVarUseConservativeDistantLightThreshold(
-	TEXT("r.Shadow.Virtual.UseConservativeDistantLightThreshold_WillBeRemoved"),
-	true,
-	TEXT("Note: Will be removed in a future release!\n")
-	TEXT("  Base the distant light cutoff on the minimum mip level instead of the shadow resolution calculated through the old path.\n")
-	TEXT("  Disabling this causes problems with long narrow spot lights, due to the use of an inscribed sphere in the original heuristic, but also typically leads to somewhat more lights being classified as \"distant\" which reduces cost."),
-	ECVF_Scalability | ECVF_RenderThreadSafe
-);
-static FAutoConsoleVariableDeprecated CVarUseConservativeDistantLightThreshold_Deprecated(TEXT("r.Shadow.Virtual.UseConservativeDistantLightThreshold"), TEXT("r.Shadow.Virtual.UseConservativeDistantLightThreshold_WillBeRemoved"), TEXT("5.7"));
-
 static TAutoConsoleVariable<float> CVarNaniteShadowsLODBias(
 	TEXT("r.Shadow.NaniteLODBias"),
 	1.0f,
@@ -151,7 +133,6 @@ FShadowSceneRenderer::FShadowSceneRenderer(FDeferredShadingSceneRenderer& InScen
 	, Scene(InShadowScene.Scene)
 	, ShadowScene(InShadowScene)
 	, VirtualShadowMapArray(InSceneRenderer.VirtualShadowMapArray)
-	, bUseConservativeDistantLightThreshold(CVarUseConservativeDistantLightThreshold.GetValueOnAnyThread())
 	, DistantLightMode(CVarDistantLightMode.GetValueOnAnyThread())
 {
 }
@@ -413,25 +394,11 @@ TSharedPtr<FVirtualShadowMapPerLightCacheEntry> FShadowSceneRenderer::AddLocalLi
 	}
 	
 	bool bIsDistantLight = DistantLightMode == 2;
-	bool bShouldForceTimeSliceDistantUpdate = false;
 
 	if (DistantLightMode == 1)
 	{
-		if (bUseConservativeDistantLightThreshold)
-		{
-			// use distant light only if we are sure that there's only one mip level.
-			bIsDistantLight = MinMipLevel == (FVirtualShadowMap::MaxMipLevels - 1);
-			bShouldForceTimeSliceDistantUpdate = false;// TODO: (bIsDistantLight && MaxScreenRadius <= BiasedFootprintThreshold * DistantLightForceCacheFootprintFraction); ??
-		}
-		else
-		{
-			// Single page res, at this point we force the VSM to be single page
-			const float BiasedFootprintThreshold = float(FVirtualShadowMap::PageSize) * FMath::Exp2(ResolutionLODBiasLocal - LightSceneProxy->GetVSMResolutionLodBias());
-			bIsDistantLight = MaxScreenRadius <= BiasedFootprintThreshold;
-			
-			const float DistantLightForceCacheFootprintFraction = FMath::Clamp(CVarDistantLightForceCacheFootprintFraction.GetValueOnRenderThread(), 0.0f, 1.0f);
-			bShouldForceTimeSliceDistantUpdate = (bIsDistantLight && MaxScreenRadius <= BiasedFootprintThreshold * DistantLightForceCacheFootprintFraction);
-		}
+		// use distant light only if we are sure that there's only one mip level.
+		bIsDistantLight = MinMipLevel == (FVirtualShadowMap::MaxMipLevels - 1);
 	}
 	
 	const int32 NumMaps = ProjectedShadowInitializer.bOnePassPointLightShadow ? 6 : 1;
@@ -442,7 +409,6 @@ TSharedPtr<FVirtualShadowMapPerLightCacheEntry> FShadowSceneRenderer::AddLocalLi
 	const int32 ForceInvalidateLocalMode = CVarForceInvalidateLocalVSM.GetValueOnRenderThread();
 	// In mode 1 we force invalidate all non-distant lights
 	// In mode 2 we force invalidate all lights, even distant ones
-	const bool bAllowInvalidation = !bShouldForceTimeSliceDistantUpdate;
 	bool bForceInvalidate = !CacheManager->IsCacheEnabled();
 	if (ForceInvalidateLocalMode > 1)
 	{
@@ -459,7 +425,6 @@ TSharedPtr<FVirtualShadowMapPerLightCacheEntry> FShadowSceneRenderer::AddLocalLi
 		LightSceneProxy->GetRadius(),
 		bIsDistantLight,
 		bForceInvalidate,
-		bAllowInvalidation,
 		IsVirtualShadowMapLocalReceiverMaskEnabled());
 
 	if (bIsDistantLight && PerLightCacheEntry->GetLastScheduledFrameNumber() == Scene.GetFrameNumber())
diff --git a/Engine/Source/Runtime/Renderer/Private/Shadows/ShadowSceneRenderer.h b/Engine/Source/Runtime/Renderer/Private/Shadows/ShadowSceneRenderer.h
--- a/Engine/Source/Runtime/Renderer/Private/Shadows/ShadowSceneRenderer.h
+++ b/Engine/Source/Runtime/Renderer/Private/Shadows/ShadowSceneRenderer.h
@@ -285,9 +285,7 @@ private:
 
 	// One pass projection stuff. Set up in RenderVitualShadowMapProjectionMaskBits
 	bool bShouldUseVirtualShadowMapOnePassProjection = false;
-
-	// Base the distant light cutoff on the minimum mip level instead of the shadow resolution calculated through the old path.
-	bool bUseConservativeDistantLightThreshold = false;
+		
 	int32 DistantLightMode = 0;
 
 	// Tracking for a given frame/render of which passes we need - clear in BeginRender
diff --git a/Engine/Source/Runtime/Renderer/Private/VirtualShadowMaps/VirtualShadowMapArray.cpp b/Engine/Source/Runtime/Renderer/Private/VirtualShadowMaps/VirtualShadowMapArray.cpp
--- a/Engine/Source/Runtime/Renderer/Private/VirtualShadowMaps/VirtualShadowMapArray.cpp
+++ b/Engine/Source/Runtime/Renderer/Private/VirtualShadowMaps/VirtualShadowMapArray.cpp
@@ -373,15 +373,6 @@ static TAutoConsoleVariable<int32> CVarCullBackfacingPixels(
 	ECVF_RenderThreadSafe
 );
 
-int32 GEnableNonNaniteVSM = 1;
-FAutoConsoleVariableRef CVarEnableNonNaniteVSM(
-	TEXT("r.Shadow.Virtual.NonNaniteVSM"),
-	GEnableNonNaniteVSM,
-	TEXT("Enable support for non-nanite Virtual Shadow Maps.")
-	TEXT("Read-only and to be set in a config file (requires restart)."),
-	ECVF_RenderThreadSafe | ECVF_ReadOnly
-);
-
 static TAutoConsoleVariable<int32> CVarNonNaniteVsmUseHzb(
 	TEXT("r.Shadow.Virtual.NonNanite.UseHZB"),
 	1,
@@ -1107,7 +1098,6 @@ void FVirtualShadowMapArray::UpdateCachedUniformBuffers(FRDGBuilder& GraphBuilde
 void FVirtualShadowMapArray::SetShaderDefines(FShaderCompilerEnvironment& OutEnvironment)
 {
 	static_assert(FVirtualShadowMap::Log2Level0DimPagesXY * 2U + NANITE_MAX_VIEWS_PER_CULL_RASTERIZE_PASS_BITS <= 32U, "Page indirection plus view index must fit into 32-bits for page-routing storage!");
-	OutEnvironment.SetDefine(TEXT("ENABLE_NON_NANITE_VSM"), GEnableNonNaniteVSM);
 	OutEnvironment.SetDefine(TEXT("MAX_PAGE_AREA_DIAGNOSTIC_SLOTS"), MaxPageAreaDiagnosticSlots);
 	OutEnvironment.SetDefine(TEXT("MAX_NPF_DIAGNOSTIC_SLOTS_SINGLEPAGE"), MaxNPFDiagnosticSlotsSinglePage);
 	OutEnvironment.SetDefine(TEXT("MAX_NPF_DIAGNOSTIC_SLOTS_MULTIPAGE"), MaxNPFDiagnosticSlotsMultiPage);
diff --git a/Engine/Source/Runtime/Renderer/Private/VirtualShadowMaps/VirtualShadowMapCacheManager.cpp b/Engine/Source/Runtime/Renderer/Private/VirtualShadowMaps/VirtualShadowMapCacheManager.cpp
--- a/Engine/Source/Runtime/Renderer/Private/VirtualShadowMaps/VirtualShadowMapCacheManager.cpp
+++ b/Engine/Source/Runtime/Renderer/Private/VirtualShadowMaps/VirtualShadowMapCacheManager.cpp
@@ -356,7 +356,6 @@ void FVirtualShadowMapPerLightCacheEntry::UpdateLocal(
 	const float NewLightRadius,
 	bool bNewIsDistantLight,
 	bool bForceInvalidate,
-	bool bAllowInvalidation,
 	bool bInUseReceiverMask)
 {
 	// TODO: The logic in this function is needlessly convoluted... clean up
@@ -369,7 +368,7 @@ void FVirtualShadowMapPerLightCacheEntry::UpdateLocal(
 
 	// Check cache validity based of shadow setup
 	// If it is a distant light, we want to let the time-share perform the invalidation.
-	if (bForceInvalidate || (bAllowInvalidation && bLightMoved))
+	if (bForceInvalidate || bLightMoved)
 	{
 		RenderedFrameNumber = -1;
 	}
diff --git a/Engine/Source/Runtime/Renderer/Private/VirtualShadowMaps/VirtualShadowMapCacheManager.h b/Engine/Source/Runtime/Renderer/Private/VirtualShadowMaps/VirtualShadowMapCacheManager.h
--- a/Engine/Source/Runtime/Renderer/Private/VirtualShadowMaps/VirtualShadowMapCacheManager.h
+++ b/Engine/Source/Runtime/Renderer/Private/VirtualShadowMaps/VirtualShadowMapCacheManager.h
@@ -136,7 +136,6 @@ public:
 		const float NewLightRadius,
 		bool bNewIsDistantLight,
 		bool bForceInvalidate,
-		bool bAllowInvalidation,
 		bool bInUseReceiverMask);
 
 	bool AffectsBounds(const FBoxSphereBounds& Bounds) const
