Add option to run Substrate's stencil writing for classification during the IndirectLighting composition pass.

This allows to merge the cost of stencil marking. This is disabled by default. This saves around 0.15ms on Sw2.

diff --git a/Engine/Source/Runtime/Renderer/Private/IndirectLightRendering.cpp b/Engine/Source/Runtime/Renderer/Private/IndirectLightRendering.cpp
--- a/Engine/Source/Runtime/Renderer/Private/IndirectLightRendering.cpp
+++ b/Engine/Source/Runtime/Renderer/Private/IndirectLightRendering.cpp
@@ -996,6 +996,7 @@ void FDeferredShadingSceneRenderer::RenderDiffuseIndirectAndAmbientOcclusion(
 
 	FSceneTextureParameters SceneTextureParameters = GetSceneTextureParameters(GraphBuilder, SceneTextures.UniformBuffer);
 	FRDGTextureRef SceneColorTexture = SceneTextures.Color.Target;
+	FRDGTextureRef SceneDepthTexture = SceneTextures.Depth.Target;
 
 	const FRDGSystemTextures& SystemTextures = FRDGSystemTextures::Get(GraphBuilder);
 
@@ -1397,6 +1398,12 @@ void FDeferredShadingSceneRenderer::RenderDiffuseIndirectAndAmbientOcclusion(
 			PSParameters->RenderTargets[0] = FRenderTargetBinding(
 				SceneColorTexture, ERenderTargetLoadAction::ELoad);
 
+			const bool bWriteStencil = !bCompositeRegularLumenOnly && Substrate::IsSubstrateEnabled() && !Substrate::UsesStandaloneClassificationStencilPass();
+			if (bWriteStencil)
+			{
+				PSParameters->RenderTargets.DepthStencil = FDepthStencilBinding(SceneDepthTexture, ERenderTargetLoadAction::ELoad, ERenderTargetLoadAction::ELoad, FExclusiveDepthStencil::DepthRead_StencilWrite);
+			}
+
 			{
 				FRDGTextureDesc Desc = FRDGTextureDesc::Create2D(
 					SceneColorTexture->Desc.Extent,
@@ -1568,7 +1575,7 @@ void FDeferredShadingSceneRenderer::RenderDiffuseIndirectAndAmbientOcclusion(
 						View.ViewRect.Width(), View.ViewRect.Height()),
 					PassParameters,
 					ERDGPassFlags::Raster,
-					[&View, TileVertexShader, PixelShader, PassParameters, TileType, BlendState, PrimitiveType](FRDGAsyncTask, FRHICommandList& RHICmdList)
+					[&View, TileVertexShader, PixelShader, PassParameters, TileType, BlendState, PrimitiveType, bWriteStencil](FRDGAsyncTask, FRHICommandList& RHICmdList)
 				{
 					FGraphicsPipelineStateInitializer GraphicsPSOInit;
 					RHICmdList.ApplyCachedRenderTargets(GraphicsPSOInit);
@@ -1583,7 +1590,12 @@ void FDeferredShadingSceneRenderer::RenderDiffuseIndirectAndAmbientOcclusion(
 					GraphicsPSOInit.BoundShaderState.VertexShaderRHI = TileVertexShader.GetVertexShader();
 					GraphicsPSOInit.BoundShaderState.PixelShaderRHI = PixelShader.GetPixelShader();		
 					GraphicsPSOInit.DepthStencilState = TStaticDepthStencilState<false, CF_Always>::GetRHI();
-					SetGraphicsPipelineState(RHICmdList, GraphicsPSOInit, 0x0);
+					uint32 StencilRef = 0;
+					if (bWriteStencil)
+					{
+						Substrate::GetSubstrateTiledDepthStencilStateWrite(TileType, GraphicsPSOInit, StencilRef);
+					}
+					SetGraphicsPipelineState(RHICmdList, GraphicsPSOInit, StencilRef);
 					SetShaderParameters(RHICmdList, PixelShader, PixelShader.GetPixelShader(), PassParameters->PS);
 					SetShaderParameters(RHICmdList, TileVertexShader, TileVertexShader.GetVertexShader(), PassParameters->VS);
 					RHICmdList.DrawPrimitiveIndirect(PassParameters->VS.TileIndirectBuffer->GetIndirectRHICallBuffer(), Substrate::TileTypeDrawIndirectArgOffset(TileType));
diff --git a/Engine/Source/Runtime/Renderer/Private/LightRendering.cpp b/Engine/Source/Runtime/Renderer/Private/LightRendering.cpp
--- a/Engine/Source/Runtime/Renderer/Private/LightRendering.cpp
+++ b/Engine/Source/Runtime/Renderer/Private/LightRendering.cpp
@@ -2445,24 +2445,10 @@ static void InternalSetBoundingGeometryRasterizerState(FGraphicsPipelineStateIni
 	}
 }
 
-template<ECompareFunction CompareFunction>
-void GetSubstrateTiledDepthStencilState(ESubstrateTileType InTileType, FGraphicsPipelineStateInitializer& GraphicsPSOInit, uint32& OutStencilRef)
-{
-	OutStencilRef = 0;
-	check(Substrate::IsSubstrateEnabled());
-	switch (InTileType)
-	{
-	case ESubstrateTileType::ESimple:			OutStencilRef = Substrate::StencilBit_Fast;				GraphicsPSOInit.DepthStencilState = TStaticDepthStencilState<false, CompareFunction, true, CF_Equal, SO_Keep, SO_Keep, SO_Keep, true, CF_Equal, SO_Keep, SO_Keep, SO_Keep, Substrate::StencilBit_Fast, 0x0>::GetRHI(); break;
-	case ESubstrateTileType::ESingle:			OutStencilRef = Substrate::StencilBit_Single;			GraphicsPSOInit.DepthStencilState = TStaticDepthStencilState<false, CompareFunction, true, CF_Equal, SO_Keep, SO_Keep, SO_Keep, true, CF_Equal, SO_Keep, SO_Keep, SO_Keep, Substrate::StencilBit_Single, 0x0>::GetRHI(); break;
-	case ESubstrateTileType::EComplex:			OutStencilRef = Substrate::StencilBit_Complex;			GraphicsPSOInit.DepthStencilState = TStaticDepthStencilState<false, CompareFunction, true, CF_Equal, SO_Keep, SO_Keep, SO_Keep, true, CF_Equal, SO_Keep, SO_Keep, SO_Keep, Substrate::StencilBit_Complex, 0x0>::GetRHI(); break;
-	case ESubstrateTileType::EComplexSpecial:	OutStencilRef = Substrate::StencilBit_ComplexSpecial; 	GraphicsPSOInit.DepthStencilState = TStaticDepthStencilState<false, CompareFunction, true, CF_Equal, SO_Keep, SO_Keep, SO_Keep, true, CF_Equal, SO_Keep, SO_Keep, SO_Keep, Substrate::StencilBit_ComplexSpecial, 0x0>::GetRHI(); break;
-	default: check(false);
-	}
-}
 
 void GetSubstrateTiledDepthStencilStateNoDepthTest(ESubstrateTileType InTileType, FGraphicsPipelineStateInitializer& GraphicsPSOInit, uint32& OutStencilRef)
 {
-	return GetSubstrateTiledDepthStencilState<CF_Always>(InTileType, GraphicsPSOInit, OutStencilRef);
+	return Substrate::GetSubstrateTiledDepthStencilStateRead(InTileType, CF_Always, GraphicsPSOInit, OutStencilRef);
 }
 
 template<ECompareFunction CompareFunction>
@@ -2473,7 +2459,7 @@ static uint32 InternalSetBoundingGeometryDepthState(FGraphicsPipelineStateInitia
 	uint32 StencilRef = 0u;
 	if (TileType != ESubstrateTileType::ECount)
 	{
-		GetSubstrateTiledDepthStencilState<CompareFunction>(TileType, GraphicsPSOInit, StencilRef);
+		Substrate::GetSubstrateTiledDepthStencilStateRead(TileType, CompareFunction, GraphicsPSOInit, StencilRef);
 	}
 	else
 	{
diff --git a/Engine/Source/Runtime/Renderer/Private/Substrate/Substrate.cpp b/Engine/Source/Runtime/Renderer/Private/Substrate/Substrate.cpp
--- a/Engine/Source/Runtime/Renderer/Private/Substrate/Substrate.cpp
+++ b/Engine/Source/Runtime/Renderer/Private/Substrate/Substrate.cpp
@@ -89,6 +89,12 @@ static TAutoConsoleVariable<int32> CVarSubstrateUseStochasticLightingClassificat
 	TEXT("Run Substrate tile classification within the Stochastic lighting classification. Only available when using Substrate Blendable GBuffer."),
 	ECVF_RenderThreadSafe);
 
+static TAutoConsoleVariable<int32> CVarSubstrateWriteStencilDuringIndirectLightingPass(
+	TEXT("r.Substrate.UseStandaloneClassificationStencilPass"),
+	1,
+	TEXT("Run Substrate tile stencil marking within a standalone pass prior to the lighting passes."),
+	ECVF_RenderThreadSafe);
+
 IMPLEMENT_GLOBAL_SHADER_PARAMETER_STRUCT(FSubstrateGlobalUniformParameters, "Substrate");
 
 void FSubstrateViewData::Reset()
@@ -117,6 +123,11 @@ namespace MegaLights
 namespace Substrate
 {
 
+bool UsesStandaloneClassificationStencilPass()
+{
+	return CVarSubstrateWriteStencilDuringIndirectLightingPass.GetValueOnRenderThread() > 0;
+}
+
 uint32 GetMaxDownsampleFactor()
 {
 	return 2;
@@ -1441,6 +1452,11 @@ void AddSubstrateStencilPass(
 	const TArray<FViewInfo>& Views,
 	const FMinimalSceneTextures& SceneTextures)
 {
+	if (!UsesStandaloneClassificationStencilPass())
+	{
+		return;
+	}
+
 	for (int32 i = 0; i < Views.Num(); ++i)
 	{
 		RDG_EVENT_SCOPE_CONDITIONAL(GraphBuilder, Views.Num() > 1, "View%d", i);
@@ -1999,4 +2015,43 @@ bool ShouldCompileSubstrateTileTypePermutations(const int32 SubstrateTileType, c
 	return true;
 }
 
+template<ECompareFunction CompareFunction>
+void InternalGetSubstrateTiledDepthStencilStateRead(ESubstrateTileType InTileType, FGraphicsPipelineStateInitializer& GraphicsPSOInit, uint32& OutStencilRef)
+{
+	OutStencilRef = 0;
+	check(Substrate::IsSubstrateEnabled());
+	switch (InTileType)
+	{
+		case ESubstrateTileType::ESimple:			OutStencilRef = Substrate::StencilBit_Fast;				GraphicsPSOInit.DepthStencilState = TStaticDepthStencilState<false, CompareFunction, true, CF_Equal, SO_Keep, SO_Keep, SO_Keep, true, CF_Equal, SO_Keep, SO_Keep, SO_Keep, Substrate::StencilBit_Fast, 0x0>::GetRHI(); break;
+		case ESubstrateTileType::ESingle:			OutStencilRef = Substrate::StencilBit_Single;			GraphicsPSOInit.DepthStencilState = TStaticDepthStencilState<false, CompareFunction, true, CF_Equal, SO_Keep, SO_Keep, SO_Keep, true, CF_Equal, SO_Keep, SO_Keep, SO_Keep, Substrate::StencilBit_Single, 0x0>::GetRHI(); break;
+		case ESubstrateTileType::EComplex:			OutStencilRef = Substrate::StencilBit_Complex;			GraphicsPSOInit.DepthStencilState = TStaticDepthStencilState<false, CompareFunction, true, CF_Equal, SO_Keep, SO_Keep, SO_Keep, true, CF_Equal, SO_Keep, SO_Keep, SO_Keep, Substrate::StencilBit_Complex, 0x0>::GetRHI(); break;
+		case ESubstrateTileType::EComplexSpecial:	OutStencilRef = Substrate::StencilBit_ComplexSpecial; 	GraphicsPSOInit.DepthStencilState = TStaticDepthStencilState<false, CompareFunction, true, CF_Equal, SO_Keep, SO_Keep, SO_Keep, true, CF_Equal, SO_Keep, SO_Keep, SO_Keep, Substrate::StencilBit_ComplexSpecial, 0x0>::GetRHI(); break;
+		default: check(false);
+	}
+}
+
+void GetSubstrateTiledDepthStencilStateWrite(ESubstrateTileType InTileType, FGraphicsPipelineStateInitializer& GraphicsPSOInit, uint32& OutStencilRef)
+{
+	OutStencilRef = 0;
+	check(Substrate::IsSubstrateEnabled());
+	switch (InTileType)
+	{
+		case ESubstrateTileType::ESimple:			OutStencilRef = Substrate::StencilBit_Fast;				GraphicsPSOInit.DepthStencilState = TStaticDepthStencilState<false, CF_Always, true, CF_Always, SO_Keep, SO_Keep, SO_Replace, false, CF_Always, SO_Keep, SO_Keep, SO_Keep, 0xFF, Substrate::StencilBit_Fast>::GetRHI(); break;
+		case ESubstrateTileType::ESingle:			OutStencilRef = Substrate::StencilBit_Single;			GraphicsPSOInit.DepthStencilState = TStaticDepthStencilState<false, CF_Always, true, CF_Always, SO_Keep, SO_Keep, SO_Replace, false, CF_Always, SO_Keep, SO_Keep, SO_Keep, 0xFF, Substrate::StencilBit_Single>::GetRHI(); break;
+		case ESubstrateTileType::EComplex:			OutStencilRef = Substrate::StencilBit_Complex;			GraphicsPSOInit.DepthStencilState = TStaticDepthStencilState<false, CF_Always, true, CF_Always, SO_Keep, SO_Keep, SO_Replace, false, CF_Always, SO_Keep, SO_Keep, SO_Keep, 0xFF, Substrate::StencilBit_Complex>::GetRHI(); break;
+		case ESubstrateTileType::EComplexSpecial:	OutStencilRef = Substrate::StencilBit_ComplexSpecial; 	GraphicsPSOInit.DepthStencilState = TStaticDepthStencilState<false, CF_Always, true, CF_Always, SO_Keep, SO_Keep, SO_Replace, false, CF_Always, SO_Keep, SO_Keep, SO_Keep, 0xFF, Substrate::StencilBit_ComplexSpecial>::GetRHI(); break;
+		default: check(false);
+	}
+}
+
+void GetSubstrateTiledDepthStencilStateRead(ESubstrateTileType InTileType, ECompareFunction InCompareFunction, FGraphicsPipelineStateInitializer& OutGraphicsPSOInit, uint32& OutStencilRef)
+{
+	switch (InCompareFunction)
+	{
+		case CF_DepthNearOrEqual: 	InternalGetSubstrateTiledDepthStencilStateRead<CF_DepthNearOrEqual>(InTileType, OutGraphicsPSOInit, OutStencilRef); break;
+		case CF_Always: 			InternalGetSubstrateTiledDepthStencilStateRead<CF_Always>(InTileType, OutGraphicsPSOInit, OutStencilRef); break;
+		default: check(false); // Not implemented
+	}
+}
+
 } // namespace Substrate
diff --git a/Engine/Source/Runtime/Renderer/Private/Substrate/Substrate.h b/Engine/Source/Runtime/Renderer/Private/Substrate/Substrate.h
--- a/Engine/Source/Runtime/Renderer/Private/Substrate/Substrate.h
+++ b/Engine/Source/Runtime/Renderer/Private/Substrate/Substrate.h
@@ -225,6 +225,7 @@ bool GetSubstrateUsesComplexSpecialPath(const FViewInfo& View);
 bool GetSubstrateUsesAnisotropy(const FViewInfo& View);
 bool UsesSubstrateMaterialBuffer(EShaderPlatform In);
 bool UsesStochasticLightingClassification(EShaderPlatform InPlatform);
+bool UsesStandaloneClassificationStencilPass();
 bool GetSubstrateUsesTileType(const FViewInfo& View, ESubstrateTileType TileType);
 
 void InitialiseSubstrateFrameSceneData(FRDGBuilder& GraphBuilder, FSceneRenderer& SceneRenderer);
@@ -301,4 +302,7 @@ void AddSubstrateDBufferBasePass(
 class FSubstrateTileType : SHADER_PERMUTATION_INT("SUBSTRATE_TILETYPE", 4);
 bool ShouldCompileSubstrateTileTypePermutations(const int32 SubstrateTileType, const EShaderPlatform Platform);
 
+void GetSubstrateTiledDepthStencilStateRead(ESubstrateTileType InTileType, ECompareFunction InCompareFunction, FGraphicsPipelineStateInitializer& OutGraphicsPSOInit, uint32& OutStencilRef);
+void GetSubstrateTiledDepthStencilStateWrite(ESubstrateTileType InTileType, FGraphicsPipelineStateInitializer& GraphicsPSOInit, uint32& OutStencilRef);
+
 }; // namespace Substrate
