Hardened header initialization to include updated primitives as well. This fixes timing issues where the mesh object update occurs before the primitive is added to the scene.

diff --git a/Engine/Source/Runtime/Engine/Public/SkinningSceneExtensionProxy.h b/Engine/Source/Runtime/Engine/Public/SkinningSceneExtensionProxy.h
--- a/Engine/Source/Runtime/Engine/Public/SkinningSceneExtensionProxy.h
+++ b/Engine/Source/Runtime/Engine/Public/SkinningSceneExtensionProxy.h
@@ -99,14 +99,13 @@ public:
 
 	bool SetLOD(int32 InLODIndex)
 	{
-		const bool bInitForBoneMap = BoneTransformStorageMode == EBoneTransformStorageMode::BoneMap && LODIndex != InLODIndex;
+		const bool bLODChange = InLODIndex != LODIndex;
 		LODIndex = (uint16)InLODIndex;
-		if (bInitForBoneMap)
+		if (bLODChange && BoneTransformStorageMode == EBoneTransformStorageMode::BoneMap)
 		{
 			InitForBoneMapStorage();
-			return true;
 		}
-		return false;
+		return bLODChange;
 	}
 
 	virtual void CreateRenderThreadResources(FSceneInterface& Scene, FRHICommandListBase& RHICmdList) {}
diff --git a/Engine/Source/Runtime/Renderer/Private/Skinning/SkinningSceneExtension.cpp b/Engine/Source/Runtime/Renderer/Private/Skinning/SkinningSceneExtension.cpp
--- a/Engine/Source/Runtime/Renderer/Private/Skinning/SkinningSceneExtension.cpp
+++ b/Engine/Source/Runtime/Renderer/Private/Skinning/SkinningSceneExtension.cpp
@@ -514,6 +514,8 @@ void FSkinningSceneExtension::PerformSkinning(FSkinningSceneParameters& Paramete
 			for (const int32 HeaderDataIndex : HeaderDataIndices)
 			{
 				FHeaderData& Header = HeaderDatas[HeaderDataIndex];
+				Header.Validate();
+
 				int32 RangeIndex = 0;
 
 				if (!EnumHasAnyFlags(Header.DirtyBoneTransforms, EDirtyBoneTransforms::All) && !Header.bAlwaysUpdateTransforms)
@@ -749,24 +751,14 @@ FSkinningSceneExtension::FUpdater::FUpdater(FSkinningSceneExtension& InSceneData
 : SceneData(&InSceneData)
 , bEnableAsync(GSkinningBuffersAsyncUpdate)
 {
-	SceneData->UpdateCounters.Full++;
-	SceneData->UpdateCounters.Header++;
-}
-
-void FSkinningSceneExtension::FUpdater::AddFullUpdate(FHeaderData& HeaderData, int32 PrimitiveIndex)
-{
-	if (HeaderData.UpdateCounters.Full != SceneData->UpdateCounters.Full)
-	{
-		HeaderData.UpdateCounters.Full = SceneData->UpdateCounters.Full;
-		FullUpdateList.Emplace(PrimitiveIndex);
-	}
+	SceneData->UpdateCounter++;
 }
 
 void FSkinningSceneExtension::FUpdater::AddHeaderUpdate(FHeaderData& HeaderData, int32 PrimitiveIndex)
 {
-	if (HeaderData.UpdateCounters.Header != SceneData->UpdateCounters.Header)
+	if (HeaderData.UpdateCounters.Header != SceneData->UpdateCounter)
 	{
-		HeaderData.UpdateCounters.Header = SceneData->UpdateCounters.Header;
+		HeaderData.UpdateCounters.Header = SceneData->UpdateCounter;
 		HeaderDataUpdateList.Emplace(PrimitiveIndex);
 	}
 }
@@ -869,7 +861,7 @@ void FSkinningSceneExtension::FUpdater::PostSceneUpdate(FRDGBuilder& GraphBuilde
 			{
 				TRACE_CPUPROFILER_EVENT_SCOPE(Skinning::InitHeaderData);
 
-				for (auto PrimitiveSceneInfo : AddedList)
+				for (FPrimitiveSceneInfo* PrimitiveSceneInfo : AddedList)
 				{
 					if (!PrimitiveSceneInfo->Proxy->IsSkinnedMesh())
 					{
@@ -928,22 +920,38 @@ void FSkinningSceneExtension::FUpdater::PostMeshUpdate(FRDGBuilder& GraphBuilder
 	SceneData->TaskHandles[InitUpdateListTask] = GraphBuilder.AddSetupTask([this, UpdatedSceneInfoList, SkeletalMeshSubmitData]
 		{
 			TRACE_CPUPROFILER_EVENT_SCOPE(Skinning::InitUpdateList);
-			TConstArrayView<FSkeletalMeshUpdateEvent> UpdateEvents = SkeletalMeshSubmitData->GetUpdateEvents();
 
-			FullUpdateList.Reserve(UpdatedSceneInfoList.Num() + UpdateEvents.Num());
-			HeaderDataUpdateList.Reserve(UpdatedSceneInfoList.Num() + UpdateEvents.Num());
+			int32 NumUpdates = 0;
+			int32 NumLODChanges = 0;
 
-			for (FPrimitiveSceneInfo* PrimitiveSceneInfo : UpdatedSceneInfoList)
+			const auto UpdateLOD = [this, &NumLODChanges] (FHeaderData& HeaderData, int32 PrimitiveIndex, int32 LODIndex)
 			{
-				const int32 Index = PrimitiveSceneInfo->GetPersistentIndex().Index;
-				if (SceneData->HeaderDatas.IsValidIndex(Index))
+				HeaderData.DirtyBoneTransforms |= EDirtyBoneTransforms::Current;
+
+				if (!HeaderData.bAlwaysUpdateTransforms)
 				{
-					AddFullUpdate(SceneData->HeaderDatas[Index], Index);
+					HeaderData.CurrentTransformSlot = !HeaderData.CurrentTransformSlot;
 				}
-			}
 
-			int32 NumUpdates = 0;
-			int32 NumLODChanges = 0;
+				if (HeaderData.Proxy->SetLOD(LODIndex))
+				{
+					NumLODChanges++;
+
+					HeaderData.MaxTransformCount = HeaderData.Proxy->GetMaxBoneTransformCount();
+					HeaderData.MaxHierarchyCount = HeaderData.Proxy->GetMaxBoneHierarchyCount();
+					HeaderData.MaxObjectSpaceCount = HeaderData.Proxy->GetMaxBoneObjectSpaceCount();
+
+					FullUpdateList.Emplace(PrimitiveIndex);
+				}
+				else
+				{
+					AddHeaderUpdate(HeaderData, PrimitiveIndex);
+				}
+			};
+
+			TConstArrayView<FSkeletalMeshUpdateEvent> UpdateEvents = SkeletalMeshSubmitData->GetUpdateEvents();
+			FullUpdateList.Reserve(UpdatedSceneInfoList.Num() + UpdateEvents.Num());
+			HeaderDataUpdateList.Reserve(UpdatedSceneInfoList.Num() + UpdateEvents.Num());
 
 			for (FSkeletalMeshUpdateEvent UpdateEvent : UpdateEvents)
 			{
@@ -969,30 +977,22 @@ void FSkinningSceneExtension::FUpdater::PostMeshUpdate(FRDGBuilder& GraphBuilder
 					continue;
 				}
 
-				HeaderData.DirtyBoneTransforms |= EDirtyBoneTransforms::Current;
-				HeaderData.bVertexFactoryDirty = 1;
-
-				if (!HeaderData.bAlwaysUpdateTransforms)
-				{
-					HeaderData.CurrentTransformSlot = !HeaderData.CurrentTransformSlot;
-				}
-
-				if (HeaderData.Proxy->SetLOD(UpdateEvent.LODIndex))
-				{
-					NumLODChanges++;
-
-					HeaderData.MaxTransformCount = HeaderData.Proxy->GetMaxBoneTransformCount();
-					HeaderData.MaxHierarchyCount = HeaderData.Proxy->GetMaxBoneHierarchyCount();
-					HeaderData.MaxObjectSpaceCount = HeaderData.Proxy->GetMaxBoneObjectSpaceCount();
-
-					AddFullUpdate(HeaderData, PrimitiveIndex);
-				}
-				else
+				UpdateLOD(HeaderData, PrimitiveIndex, UpdateEvent.LODIndex);
+				HeaderData.UpdateCounters.UpdateLOD = SceneData->UpdateCounter;
+				NumUpdates++;
+			}
+			
+			for (FPrimitiveSceneInfo* PrimitiveSceneInfo : UpdatedSceneInfoList)
+			{
+				const int32 Index = PrimitiveSceneInfo->GetPersistentIndex().Index;
+				if (SceneData->HeaderDatas.IsValidIndex(Index))
 				{
-					AddHeaderUpdate(HeaderData, PrimitiveIndex);
+					FHeaderData& HeaderData = SceneData->HeaderDatas[Index];
+					if (HeaderData.UpdateCounters.UpdateLOD != SceneData->UpdateCounter)
+					{
+						UpdateLOD(HeaderData, Index, HeaderData.MeshObject->GetLOD());
+					}
 				}
-
-				NumUpdates++;
 			}
 
 			INC_DWORD_STAT_BY(STAT_SkinningSceneExtension_NumUpdates, NumUpdates);
@@ -1191,7 +1191,7 @@ void FSkinningSceneExtension::FUpdater::PostMeshUpdate(FRDGBuilder& GraphBuilder
 		bEnableAsync
 	);
 
-	auto UploadHeaderData = [this](const FHeaderData& Data)
+	auto UploadHeaderData = [this](FHeaderData& Data)
 	{
 		const int32 PersistentIndex = Data.PrimitiveSceneInfo->GetPersistentIndex().Index;
 
@@ -1201,6 +1201,8 @@ void FSkinningSceneExtension::FUpdater::PostMeshUpdate(FRDGBuilder& GraphBuilder
 
 		check(SceneData->Uploader.IsValid()); // Sanity check
 		SceneData->Uploader->HeaderDataUploader.Add(Data.Pack(), PersistentIndex);
+
+		Data.bVertexFactoryDirty = 1;
 	};
 
 	// Kick off the header data upload task (synced when accessing the buffer)
diff --git a/Engine/Source/Runtime/Renderer/Private/Skinning/SkinningSceneExtension.h b/Engine/Source/Runtime/Renderer/Private/Skinning/SkinningSceneExtension.h
--- a/Engine/Source/Runtime/Renderer/Private/Skinning/SkinningSceneExtension.h
+++ b/Engine/Source/Runtime/Renderer/Private/Skinning/SkinningSceneExtension.h
@@ -27,8 +27,8 @@ class FSkinningSceneExtension : public ISceneExtension
 
 	struct FUpdateCounters
 	{
-		uint32 Full = 0;
 		uint32 Header = 0;
+		uint32 UpdateLOD = 0;
 	};
 
 public:
@@ -46,7 +46,6 @@ public:
 		void PostMeshUpdate(FRDGBuilder& GraphBuilder, const TConstArrayView<FPrimitiveSceneInfo*>& SceneInfosWithStaticDrawListUpdate, const FSkeletalMeshUpdater::FSubmitData* SkeletalMeshSubmitData);
 
 	private:
-		void AddFullUpdate(FHeaderData& HeaderData, int32 PrimitiveIndex);
 		void AddHeaderUpdate(FHeaderData& HeaderData, int32 PrimitiveIndex);
 
 		FSkinningSceneExtension* SceneData = nullptr;
@@ -138,11 +137,7 @@ private:
 
 		FSkinningHeader Pack() const
 		{
-			// Verify that values all fit within the encoded range prior to packing
-			check(HierarchyBufferOffset	<= SKINNING_BUFFER_HIERARCHY_OFFSET_MAX);
-			check(TransformBufferOffset	<= SKINNING_BUFFER_TRANSFORM_OFFSET_MAX);
-			check(ObjectSpaceBufferOffset == INDEX_NONE || ObjectSpaceBufferOffset <= SKINNING_BUFFER_OBJECT_SPACE_OFFSET_MAX);
-			check(MaxInfluenceCount		<= SKINNING_BUFFER_INFLUENCE_MAX);
+			Validate();
 
 			FSkinningHeader Output;
 			Output.HierarchyBufferOffset	= HierarchyBufferOffset;
@@ -155,6 +150,14 @@ private:
 			Output.CurrentTransformSlot		= CurrentTransformSlot;
 			return Output;
 		}
+
+		void Validate() const
+		{
+			check(HierarchyBufferOffset	<= SKINNING_BUFFER_HIERARCHY_OFFSET_MAX);
+			check(TransformBufferOffset	<= SKINNING_BUFFER_TRANSFORM_OFFSET_MAX);
+			check(ObjectSpaceBufferOffset == INDEX_NONE || ObjectSpaceBufferOffset <= SKINNING_BUFFER_OBJECT_SPACE_OFFSET_MAX);
+			check(MaxInfluenceCount		<= SKINNING_BUFFER_INFLUENCE_MAX);
+		}
 	};
 
 	class FBuffers
@@ -212,7 +215,7 @@ private:
 	TUniquePtr<FBuffers> Buffers;
 	TUniquePtr<FUploader> Uploader;
 	TStaticArray<UE::Tasks::FTask, NumTasks> TaskHandles;
-	FUpdateCounters UpdateCounters;
+	uint32 UpdateCounter = 0;
 	int32 VertexFactoryUpdateRequests = 0;
 	bool bVertexFactoryFullUpload = false;
 
