[LitTranslucency] Release history resources when temporal filter is disabled.
- also fix r.TranslucencyLightingVolume.Debug.FreezePosition when temporal filter is disabled.

diff --git a/Engine/Source/Runtime/Renderer/Private/TranslucentLighting.cpp b/Engine/Source/Runtime/Renderer/Private/TranslucentLighting.cpp
--- a/Engine/Source/Runtime/Renderer/Private/TranslucentLighting.cpp
+++ b/Engine/Source/Runtime/Renderer/Private/TranslucentLighting.cpp
@@ -2312,7 +2312,6 @@ void FilterTranslucencyLightingVolume(
 				{
 					TranslucencyLightingViewState.HistoryAmbient[Index] = nullptr;
 					TranslucencyLightingViewState.HistoryDirectional[Index] = nullptr;
-					
 					TranslucencyLightingViewState.HistoryMark[Index] = nullptr;
 				}
 			}
@@ -2468,15 +2467,48 @@ void FilterTranslucencyLightingVolume(
 					{
 						TranslucencyLightingViewState.HistoryMark[VolumeCascadeIndex] = nullptr;
 					}
+				}
+			}
+		}
+	}
+	else
+	{
+		// release history resources if temporal is not enabled
 
-					TranslucencyLightingViewState.HistoryVolumeMin[VolumeCascadeIndex] = View.TranslucencyLightingVolumeMin[VolumeCascadeIndex];
-					TranslucencyLightingViewState.HistoryVoxelSize[VolumeCascadeIndex] = View.TranslucencyVolumeVoxelSize[VolumeCascadeIndex];
-					TranslucencyLightingViewState.HistoryVolumeSize[VolumeCascadeIndex] = View.TranslucencyLightingVolumeSize[VolumeCascadeIndex];
+		for (int32 ViewIndex = 0; ViewIndex < Views.Num(); ++ViewIndex)
+		{
+			const FViewInfo& View = Views[ViewIndex];
+
+			if (View.ViewState && !View.bStatePrevViewInfoIsReadOnly)
+			{
+				for (int32 VolumeCascadeIndex = 0; VolumeCascadeIndex < TVC_MAX; VolumeCascadeIndex++)
+				{
+					FTranslucencyLightingViewState& TranslucencyLightingViewState = View.ViewState->TranslucencyLighting;
+					TranslucencyLightingViewState.HistoryAmbient[VolumeCascadeIndex] = nullptr;
+					TranslucencyLightingViewState.HistoryDirectional[VolumeCascadeIndex] = nullptr;
+					TranslucencyLightingViewState.HistoryMark[VolumeCascadeIndex] = nullptr;
 				}
 			}
 		}
 	}
 
+	// update HistoryVolumeMin/Size regardless of temporal filter being enabled or not so r.TranslucencyLightingVolume.Debug.FreezePosition always works
+	for (int32 ViewIndex = 0; ViewIndex < Views.Num(); ++ViewIndex)
+	{
+		const FViewInfo& View = Views[ViewIndex];
+
+		if (View.ViewState && !View.bStatePrevViewInfoIsReadOnly)
+		{
+			for (int32 VolumeCascadeIndex = 0; VolumeCascadeIndex < TVC_MAX; VolumeCascadeIndex++)
+			{
+				FTranslucencyLightingViewState& TranslucencyLightingViewState = View.ViewState->TranslucencyLighting;
+				TranslucencyLightingViewState.HistoryVolumeMin[VolumeCascadeIndex] = View.TranslucencyLightingVolumeMin[VolumeCascadeIndex];
+				TranslucencyLightingViewState.HistoryVoxelSize[VolumeCascadeIndex] = View.TranslucencyVolumeVoxelSize[VolumeCascadeIndex];
+				TranslucencyLightingViewState.HistoryVolumeSize[VolumeCascadeIndex] = View.TranslucencyLightingVolumeSize[VolumeCascadeIndex];
+			}
+		}
+	}
+
 	if (bSpatialFilter)
 	{
 		for (int32 ViewIndex = 0; ViewIndex < Views.Num(); ++ViewIndex)
