Temporarily disabled a check because of unmet assumptions.
#rnx

diff --git a/Engine/Source/Runtime/Renderer/Private/ShadowSetup.cpp b/Engine/Source/Runtime/Renderer/Private/ShadowSetup.cpp
--- a/Engine/Source/Runtime/Renderer/Private/ShadowSetup.cpp
+++ b/Engine/Source/Runtime/Renderer/Private/ShadowSetup.cpp
@@ -1945,7 +1945,6 @@ bool FProjectedShadowInfo::ShouldDrawStaticMeshes_AnyThread(
 bool FProjectedShadowInfo::AddSubjectPrimitive(FPrimitiveSceneInfo* PrimitiveSceneInfo, TArrayView<FViewInfo> ViewArray, bool bShouldRecordShadowSubjectsForMobileShading, bool bInvertViewShadowRelevance)
 {
 	// Ray traced shadows use the GPU managed distance field object buffers, no CPU culling should be used
-	check(!bRayTracedDistanceField);
 	bool bWasAdded = false;
 	if (!ReceiverPrimitives.Contains(PrimitiveSceneInfo)
 		// Far cascade only casts from primitives marked for it
