Fix issue with local lights that opt in to distance field shadows when VSMs are enabled.

#rnx

diff --git a/Engine/Source/Runtime/Renderer/Private/ShadowSetup.cpp b/Engine/Source/Runtime/Renderer/Private/ShadowSetup.cpp
--- a/Engine/Source/Runtime/Renderer/Private/ShadowSetup.cpp
+++ b/Engine/Source/Runtime/Renderer/Private/ShadowSetup.cpp
@@ -1945,6 +1945,7 @@ bool FProjectedShadowInfo::ShouldDrawStaticMeshes_AnyThread(
 bool FProjectedShadowInfo::AddSubjectPrimitive(FPrimitiveSceneInfo* PrimitiveSceneInfo, TArrayView<FViewInfo> ViewArray, bool bShouldRecordShadowSubjectsForMobileShading, bool bInvertViewShadowRelevance)
 {
 	// Ray traced shadows use the GPU managed distance field object buffers, no CPU culling should be used
+	check(!bRayTracedDistanceField);
 	bool bWasAdded = false;
 	if (!ReceiverPrimitives.Contains(PrimitiveSceneInfo)
 		// Far cascade only casts from primitives marked for it
@@ -4233,7 +4234,9 @@ void FSceneRenderer::CreateWholeSceneProjectedShadow(
 					}
 				};
 		
-				if (LightSceneInfo->Proxy->UseVirtualShadowMaps())
+				// NOTE: Allow lights that explicitly opt in to distance field shadows to take precidence
+				// We may want to change this directly in the proxy, but for now this maintains existing behavior
+				if (LightSceneInfo->Proxy->UseVirtualShadowMaps() && !ProjectedShadowInitializer.bRayTracedDistanceField)
 				{
 					// VSMs should be enabled globally and proxy should never ask for more than one
 					if (ensure(VirtualShadowMapArray.IsEnabled() &&	ProjectedShadowInitializers.Num() == 1))
