[RHI] Add support for writing GPU addresses of FRHIResources into a buffer.
- replaces GetAccelerationStructureAddress().

diff --git a/Engine/Source/Runtime/Apple/MetalRHI/Private/MetalDynamicRHI.h b/Engine/Source/Runtime/Apple/MetalRHI/Private/MetalDynamicRHI.h
--- a/Engine/Source/Runtime/Apple/MetalRHI/Private/MetalDynamicRHI.h
+++ b/Engine/Source/Runtime/Apple/MetalRHI/Private/MetalDynamicRHI.h
@@ -129,6 +129,7 @@ public:
 	[[nodiscard]] virtual FRHIBufferInitializer RHICreateBufferInitializer(FRHICommandListBase& RHICmdList, const FRHIBufferCreateDesc& CreateDesc) final override;
 
 	virtual void RHIReplaceResources(FRHICommandListBase& RHICmdList, TArray<FRHIResourceReplaceInfo>&& ReplaceInfos) final override;
+	virtual void RHIWriteResourceGPUAddresses(FRHICommandListBase& RHICmdList, uint32 GPUIndex, TConstArrayView<FRHIResource*> Resources, uint64* AddressesPtr) final override;
 	
 	virtual void * RHILockBuffer(class FRHICommandListBase& RHICmdList, FRHIBuffer* Buffer, uint32 Offset, uint32 SizeRHI, EResourceLockMode LockMode) final override;
 	virtual void RHIUnlockBuffer(class FRHICommandListBase& RHICmdList, FRHIBuffer* Buffer) final override;
diff --git a/Engine/Source/Runtime/Apple/MetalRHI/Private/MetalRHI.cpp b/Engine/Source/Runtime/Apple/MetalRHI/Private/MetalRHI.cpp
--- a/Engine/Source/Runtime/Apple/MetalRHI/Private/MetalRHI.cpp
+++ b/Engine/Source/Runtime/Apple/MetalRHI/Private/MetalRHI.cpp
@@ -1888,6 +1888,38 @@ void FMetalDynamicRHI::RHIReplaceResources(FRHICommandListBase& RHICmdList, TArr
 	RHICmdList.RHIThreadFence(true);
 }
 
+void FMetalDynamicRHI::RHIWriteResourceGPUAddresses(FRHICommandListBase& RHICmdList, uint32 GPUIndex, TConstArrayView<FRHIResource*> Resources, uint64* AddressesPtr)
+{
+	RHICmdList.EnqueueLambda([AddressesPtr, Resources](FRHICommandListBase&)
+		{
+			TRACE_CPUPROFILER_EVENT_SCOPE(RHIWriteResourceGPUAddresses);
+
+			for (int32 ResourceIndex = 0; ResourceIndex < Resources.Num(); ++ResourceIndex)
+			{
+				const FRHIResource* Resource = Resources[ResourceIndex];
+
+				if (Resource->GetType() == RRT_Buffer)
+				{
+					const FMetalRHIBuffer* Buffer = ResourceCast(static_cast<const FRHIBuffer*>(Resource));
+					AddressesPtr[ResourceIndex] = Buffer->GetCurrentBufferOrNull()->GetGPUAddress();
+				}
+				else if (Resource->GetType() == RRT_Texture)
+				{
+					unimplemented();
+				}
+				else if (Resource->GetType() == RRT_RayTracingGeometry)
+				{
+					const FMetalRayTracingGeometry* RayTracingGeometry = ResourceCast(static_cast<const FRHIRayTracingGeometry*>(Resource));
+					AddressesPtr[ResourceIndex] = RayTracingGeometry->GetGPUVirtualAddress();
+				}
+				else
+				{
+					checkNoEntry();
+				}
+			}
+		});
+}
+
 #if PLATFORM_SUPPORTS_BINDLESS_RENDERING
 FRHIResourceCollectionRef FMetalDynamicRHI::RHICreateResourceCollection(FRHICommandListBase& RHICmdList, TConstArrayView<FRHIResourceCollectionMember> InMembers)
 {
diff --git a/Engine/Source/Runtime/Apple/MetalRHI/Private/MetalRayTracing.cpp b/Engine/Source/Runtime/Apple/MetalRHI/Private/MetalRayTracing.cpp
--- a/Engine/Source/Runtime/Apple/MetalRHI/Private/MetalRayTracing.cpp
+++ b/Engine/Source/Runtime/Apple/MetalRHI/Private/MetalRayTracing.cpp
@@ -866,9 +866,9 @@ void FMetalRayTracingScene::BindBuffer(FRHIBuffer* InBuffer, uint32 InBufferOffs
 	}
 }
 
-FRayTracingAccelerationStructureAddress FMetalRayTracingGeometry::GetAccelerationStructureAddress(uint64 GPUIndex) const
+uint64 FMetalRayTracingGeometry::GetGPUAddress() const
 {
-	return AccelerationStructure->GetAccelerationStructure()->gpuResourceID()._impl; 
+	return AccelerationStructure->GetAccelerationStructure()->gpuResourceID()._impl;
 }
 
 void FMetalRHICommandContext::BuildAccelerationStructure(FMetalBufferPtr CurInstanceBuffer, uint32 InstanceBufferOffset, 
diff --git a/Engine/Source/Runtime/Apple/MetalRHI/Private/MetalRayTracing.h b/Engine/Source/Runtime/Apple/MetalRHI/Private/MetalRayTracing.h
--- a/Engine/Source/Runtime/Apple/MetalRHI/Private/MetalRayTracing.h
+++ b/Engine/Source/Runtime/Apple/MetalRHI/Private/MetalRayTracing.h
@@ -28,9 +28,13 @@ public:
 
 	void ReleaseUnderlyingResource();
 	void ReleaseDescriptors();
+
+	uint64 GetGPUAddress() const;
 	
 	/** FRHIRayTracingGeometry Interface */
-	virtual FRayTracingAccelerationStructureAddress GetAccelerationStructureAddress(uint64 GPUIndex) const final override;
+	PRAGMA_DISABLE_DEPRECATION_WARNINGS
+	virtual FRayTracingAccelerationStructureAddress GetAccelerationStructureAddress(uint64 GPUIndex) const final override { return GetGPUAddress(); }
+	PRAGMA_ENABLE_DEPRECATION_WARNINGS
 	/** FRHIRayTracingGeometry Interface */
 
 	void Swap(FMetalRayTracingGeometry& Other);
diff --git a/Engine/Source/Runtime/D3D12RHI/Private/D3D12RHIPrivate.h b/Engine/Source/Runtime/D3D12RHI/Private/D3D12RHIPrivate.h
--- a/Engine/Source/Runtime/D3D12RHI/Private/D3D12RHIPrivate.h
+++ b/Engine/Source/Runtime/D3D12RHI/Private/D3D12RHIPrivate.h
@@ -307,6 +307,13 @@ public:
 		return Object ? static_cast<TReturnType*>(Object->GetLinkedObject(GPUIndex)) : nullptr;
 	}
 
+	template<typename TRHIType, typename TReturnType = typename TD3D12ResourceTraits<TRHIType>::TConcreteType>
+	static FORCEINLINE_DEBUGGABLE const TReturnType* ResourceCast(const TRHIType* Resource, uint32 GPUIndex)
+	{
+		const TReturnType* Object = ResourceCast<const TRHIType, const TReturnType>(Resource);
+		return Object ? static_cast<const TReturnType*>(Object->GetLinkedObject(GPUIndex)) : nullptr;
+	}
+
 	virtual bool QueueSupportsTileMapping(ED3D12QueueType /*InQueueType*/) { return true; }
 
 	virtual FD3D12CommandContext* CreateCommandContext(FD3D12Device* InParent, ED3D12QueueType InQueueType, bool InIsDefaultContext);
@@ -349,6 +356,8 @@ public:
 	[[nodiscard]] virtual FRHIBufferInitializer RHICreateBufferInitializer(FRHICommandListBase& RHICmdList, const FRHIBufferCreateDesc& CreateDesc) override;
 
 	virtual void RHIReplaceResources(FRHICommandListBase& RHICmdList, TArray<FRHIResourceReplaceInfo>&& ReplaceInfos) final override;
+	virtual void RHIWriteResourceGPUAddresses(FRHICommandListBase& RHICmdList, uint32 GPUIndex, TConstArrayView<FRHIResource*> Resources, uint64* AddressesPtr) final override;
+
 	virtual void* RHILockBuffer(FRHICommandListBase& RHICmdList, FRHIBuffer* Buffer, uint32 Offset, uint32 Size, EResourceLockMode LockMode) final override;
 	virtual void* RHILockBufferMGPU(FRHICommandListBase& RHICmdList, FRHIBuffer* Buffer, uint32 GPUIndex, uint32 Offset, uint32 Size, EResourceLockMode LockMode) final override;
 	virtual void RHIUnlockBuffer(FRHICommandListBase& RHICmdList, FRHIBuffer* Buffer) final override;
diff --git a/Engine/Source/Runtime/D3D12RHI/Private/D3D12RayTracing.h b/Engine/Source/Runtime/D3D12RHI/Private/D3D12RayTracing.h
--- a/Engine/Source/Runtime/D3D12RHI/Private/D3D12RayTracing.h
+++ b/Engine/Source/Runtime/D3D12RHI/Private/D3D12RayTracing.h
@@ -130,14 +130,21 @@ public:
 
 	FD3D12RayTracingGeometry(FRHICommandListBase& RHICmdList, FD3D12Adapter* Adapter, const FRayTracingGeometryInitializer& Initializer);
 	~FD3D12RayTracingGeometry();
-
-	virtual FRayTracingAccelerationStructureAddress GetAccelerationStructureAddress(uint64 GPUIndex) const final override
+	
+	D3D12_GPU_VIRTUAL_ADDRESS GetGPUVirtualAddress(uint64 GPUIndex) const
 	{
-		checkf(AccelerationStructureBuffers[GPUIndex], 
+		checkf(AccelerationStructureBuffers[GPUIndex],
 			TEXT("Trying to get address of acceleration structure '%s' without allocated memory."), *DebugName.ToString());
 		return AccelerationStructureBuffers[GPUIndex]->ResourceLocation.GetGPUVirtualAddress();
 	}
 
+	PRAGMA_DISABLE_DEPRECATION_WARNINGS
+	virtual FRayTracingAccelerationStructureAddress GetAccelerationStructureAddress(uint64 GPUIndex) const final override
+	PRAGMA_ENABLE_DEPRECATION_WARNINGS
+	{
+		return GetGPUVirtualAddress(GPUIndex);
+	}
+
 	void SetupHitGroupSystemParameters(uint32 InGPUIndex);
 	void UpdateResidency(FD3D12CommandContext& CommandContext);
 	void CompactAccelerationStructure(FD3D12CommandContext& CommandContext, uint32 InGPUIndex, uint64 InSizeAfterCompaction);
diff --git a/Engine/Source/Runtime/D3D12RHI/Private/D3D12Resources.cpp b/Engine/Source/Runtime/D3D12RHI/Private/D3D12Resources.cpp
--- a/Engine/Source/Runtime/D3D12RHI/Private/D3D12Resources.cpp
+++ b/Engine/Source/Runtime/D3D12RHI/Private/D3D12Resources.cpp
@@ -1677,3 +1677,35 @@ void FD3D12DynamicRHI::RHIReplaceResources(FRHICommandListBase& RHICmdList, TArr
 		}
 	);
 }
+
+void FD3D12DynamicRHI::RHIWriteResourceGPUAddresses(FRHICommandListBase& RHICmdList, uint32 GPUIndex, TConstArrayView<FRHIResource*> Resources, uint64* AddressesPtr)
+{
+	RHICmdList.EnqueueLambda([AddressesPtr, Resources, GPUIndex](FRHICommandListBase&)
+		{
+			TRACE_CPUPROFILER_EVENT_SCOPE(RHIWriteResourceGPUAddresses);
+
+			for (int32 ResourceIndex = 0; ResourceIndex < Resources.Num(); ++ResourceIndex)
+			{
+				const FRHIResource* Resource = Resources[ResourceIndex];
+
+				if (Resource->GetType() == RRT_Buffer)
+				{
+					const FD3D12Buffer* Buffer = ResourceCast(static_cast<const FRHIBuffer*>(Resource), GPUIndex);
+					AddressesPtr[ResourceIndex] = Buffer->GetGPUVirtualAddress();
+				}
+				else if (Resource->GetType() == RRT_Texture)
+				{
+					unimplemented();
+				}
+				else if (Resource->GetType() == RRT_RayTracingGeometry)
+				{
+					const FD3D12RayTracingGeometry* RayTracingGeometry = ResourceCast(static_cast<const FRHIRayTracingGeometry*>(Resource));
+					AddressesPtr[ResourceIndex] = RayTracingGeometry->GetGPUVirtualAddress(GPUIndex);
+				}
+				else
+				{
+					checkNoEntry();
+				}
+			}
+		});
+}
diff --git a/Engine/Source/Runtime/D3D12RHI/Private/D3D12Resources.h b/Engine/Source/Runtime/D3D12RHI/Private/D3D12Resources.h
--- a/Engine/Source/Runtime/D3D12RHI/Private/D3D12Resources.h
+++ b/Engine/Source/Runtime/D3D12RHI/Private/D3D12Resources.h
@@ -894,6 +894,11 @@ public:
 		}
 	}
 
+	D3D12_GPU_VIRTUAL_ADDRESS GetGPUVirtualAddress() const
+	{
+		return ResourceLocation.GetGPUVirtualAddress();
+	}
+
 	FD3D12ResourceLocation ResourceLocation;
 
 public:
diff --git a/Engine/Source/Runtime/RHI/Private/RHIStrings.cpp b/Engine/Source/Runtime/RHI/Private/RHIStrings.cpp
--- a/Engine/Source/Runtime/RHI/Private/RHIStrings.cpp
+++ b/Engine/Source/Runtime/RHI/Private/RHIStrings.cpp
@@ -432,7 +432,11 @@ PRAGMA_ENABLE_DEPRECATION_WARNINGS
 	RHI_RESOURCE_TYPE_DEF(Viewport),
 	RHI_RESOURCE_TYPE_DEF(UnorderedAccessView),
 	RHI_RESOURCE_TYPE_DEF(ShaderResourceView),
+PRAGMA_DISABLE_DEPRECATION_WARNINGS
 	RHI_RESOURCE_TYPE_DEF(RayTracingAccelerationStructure),
+PRAGMA_ENABLE_DEPRECATION_WARNINGS
+	RHI_RESOURCE_TYPE_DEF(RayTracingGeometry),
+	RHI_RESOURCE_TYPE_DEF(RayTracingScene),
 	RHI_RESOURCE_TYPE_DEF(StagingBuffer),
 	RHI_RESOURCE_TYPE_DEF(CustomPresent),
 	RHI_RESOURCE_TYPE_DEF(ShaderLibrary),
diff --git a/Engine/Source/Runtime/RHI/Public/DynamicRHI.h b/Engine/Source/Runtime/RHI/Public/DynamicRHI.h
--- a/Engine/Source/Runtime/RHI/Public/DynamicRHI.h
+++ b/Engine/Source/Runtime/RHI/Public/DynamicRHI.h
@@ -1050,6 +1050,11 @@ public:
 		return nullptr;
 	}
 
+	virtual void RHIWriteResourceGPUAddresses(FRHICommandListBase& RHICmdList, uint32 GPUIndex, TConstArrayView<FRHIResource*> Resources, uint64* AddressesPtr)
+	{
+		checkNoEntry();
+	}
+
 protected:
 	TArray<uint32> PixelFormatBlockBytes;
 	friend class FValidationRHI;
diff --git a/Engine/Source/Runtime/RHI/Public/RHICommandList.h b/Engine/Source/Runtime/RHI/Public/RHICommandList.h
--- a/Engine/Source/Runtime/RHI/Public/RHICommandList.h
+++ b/Engine/Source/Runtime/RHI/Public/RHICommandList.h
@@ -1238,6 +1238,31 @@ PRAGMA_ENABLE_DEPRECATION_WARNINGS
 		return new (MemManager) FRHIBatchedShaderParametersAllocator(ShaderParameterState.AllocatorsRoot, *this, PageSize);
 	}
 
+	/*
+	* Writes the GPU addresses of the provided resources into a GPU buffer that can be accessed in shaders.
+	* Addresses should be queried immediately before use, as certain RHI operations such as lock/update/defragment may move the resources in memory.
+	*/
+	void WriteResourceGPUAddresses(TConstArrayView<FRHIResource*> Resources, FRHIBuffer* OutputBuffer)
+	{
+		if (!Bypass())
+		{
+			Resources = AllocArray(Resources);
+		}
+
+		for (uint32 GPUIndex : GetGPUMask())
+		{
+			uint64* AddressesPtr = (uint64*)LockBufferMGPU(
+				OutputBuffer,
+				GPUIndex,
+				0,
+				Resources.Num() * sizeof(uint64), RLM_WriteOnly);
+
+			GDynamicRHI->RHIWriteResourceGPUAddresses(*this, GPUIndex, Resources, AddressesPtr);
+
+			UnlockBufferMGPU(OutputBuffer, GPUIndex);
+		}
+	}
+
 protected:
 	FMemStackBase& GetAllocator() { return MemManager; }
 
diff --git a/Engine/Source/Runtime/RHI/Public/RHIDefinitions.h b/Engine/Source/Runtime/RHI/Public/RHIDefinitions.h
--- a/Engine/Source/Runtime/RHI/Public/RHIDefinitions.h
+++ b/Engine/Source/Runtime/RHI/Public/RHIDefinitions.h
@@ -1062,7 +1062,9 @@ enum ERHIResourceType : uint8
 	RRT_Viewport,
 	RRT_UnorderedAccessView,
 	RRT_ShaderResourceView,
-	RRT_RayTracingAccelerationStructure,
+	RRT_RayTracingAccelerationStructure UE_DEPRECATED(5.7, "RRT_RayTracingAccelerationStructure is deprecated. Use RRT_RayTracingGeometry or RRT_RayTracingScene instead."),
+	RRT_RayTracingGeometry,
+	RRT_RayTracingScene,
 	RRT_RayTracingShaderBindingTable,
 	RRT_StagingBuffer,
 	RRT_CustomPresent,
diff --git a/Engine/Source/Runtime/RHI/Public/RHIResources.h b/Engine/Source/Runtime/RHI/Public/RHIResources.h
--- a/Engine/Source/Runtime/RHI/Public/RHIResources.h
+++ b/Engine/Source/Runtime/RHI/Public/RHIResources.h
@@ -3730,29 +3730,39 @@ class FRHIRayTracingAccelerationStructure
 #endif
 {
 public:
-	FRHIRayTracingAccelerationStructure() : FRHIResource(RRT_RayTracingAccelerationStructure) {}
-
 	FRayTracingAccelerationStructureSize GetSizeInfo() const
 	{
 		return SizeInfo;
 	};
 
 protected:
+	FRHIRayTracingAccelerationStructure(bool bIsRayTracingScene)
+		: FRHIResource(bIsRayTracingScene ? RRT_RayTracingScene : RRT_RayTracingGeometry)
+	{}
+
 	FRayTracingAccelerationStructureSize SizeInfo = {};
 };
 
-using FRayTracingAccelerationStructureAddress = uint64;
+using FRayTracingAccelerationStructureAddress UE_DEPRECATED(5.7, "FRayTracingAccelerationStructureAddress is deprecated and should no longer be used.") = uint64;
 
 /** Bottom level ray tracing acceleration structure (contains triangles). */
 class FRHIRayTracingGeometry : public FRHIRayTracingAccelerationStructure
 {
 public:
-	FRHIRayTracingGeometry() = default;
+	FRHIRayTracingGeometry()
+		: FRHIRayTracingAccelerationStructure(false)
+	{}
+
 	FRHIRayTracingGeometry(const FRayTracingGeometryInitializer& InInitializer)
-		: Initializer(InInitializer)
+		: FRHIRayTracingAccelerationStructure(false)
+		, Initializer(InInitializer)
 	{}
 
+	PRAGMA_DISABLE_DEPRECATION_WARNINGS
+	UE_DEPRECATED(5.7, "Use RHICmdList.WriteResourceGPUAddresses(...) instead.")
 	virtual FRayTracingAccelerationStructureAddress GetAccelerationStructureAddress(uint64 GPUIndex) const = 0;
+	PRAGMA_ENABLE_DEPRECATION_WARNINGS
+
 	virtual bool IsCompressed() const { return false; }
 
 	const FRayTracingGeometryInitializer& GetInitializer() const
@@ -3773,6 +3783,10 @@ class FRHIRayTracingScene
 	: public FRHIRayTracingAccelerationStructure
 {
 public:
+	FRHIRayTracingScene()
+		: FRHIRayTracingAccelerationStructure(true)
+	{}
+
 	virtual const FRayTracingSceneInitializer& GetInitializer() const = 0;
 };
 
diff --git a/Engine/Source/Runtime/Renderer/Private/RayTracing/RayTracingInstanceBufferUtil.cpp b/Engine/Source/Runtime/Renderer/Private/RayTracing/RayTracingInstanceBufferUtil.cpp
--- a/Engine/Source/Runtime/Renderer/Private/RayTracing/RayTracingInstanceBufferUtil.cpp
+++ b/Engine/Source/Runtime/Renderer/Private/RayTracing/RayTracingInstanceBufferUtil.cpp
@@ -991,7 +991,7 @@ void FRayTracingInstanceBufferBuilder::FillAccelerationStructureAddressesBuffer(
 	{
 		// Round to PoT to avoid resizing too often
 		const uint32 NumGeometriesAligned = FMath::RoundUpToPowerOfTwo(NumGeometries);
-		const uint32 AccelerationStructureAddressesBufferSize = NumGeometriesAligned * sizeof(FRayTracingAccelerationStructureAddress);
+		const uint32 AccelerationStructureAddressesBufferSize = NumGeometriesAligned * sizeof(uint64);
 
 		if (AccelerationStructureAddressesBuffer.NumBytes < AccelerationStructureAddressesBufferSize)
 		{
@@ -1000,28 +1000,10 @@ void FRayTracingInstanceBufferBuilder::FillAccelerationStructureAddressesBuffer(
 		}
 	}
 
-	for (uint32 GPUIndex : RHICmdList.GetGPUMask())
-	{
-		FRayTracingAccelerationStructureAddress* AddressesPtr = (FRayTracingAccelerationStructureAddress*)RHICmdList.LockBufferMGPU(
-			AccelerationStructureAddressesBuffer.Buffer,
-			GPUIndex,
-			0,
-			NumGeometries * sizeof(FRayTracingAccelerationStructureAddress), RLM_WriteOnly);
-
-		const TArrayView<FRHIRayTracingGeometry*> ReferencedGeometries = RHICmdList.AllocArray(MakeConstArrayView(Data.ReferencedGeometries));
-
-		RHICmdList.EnqueueLambda([AddressesPtr, ReferencedGeometries, GPUIndex](FRHICommandListBase&)
-			{
-				TRACE_CPUPROFILER_EVENT_SCOPE(GetAccelerationStructuresAddresses);
-
-				for (int32 GeometryIndex = 0; GeometryIndex < ReferencedGeometries.Num(); ++GeometryIndex)
-				{
-					AddressesPtr[GeometryIndex] = ReferencedGeometries[GeometryIndex]->GetAccelerationStructureAddress(GPUIndex);
-				}
-			});
+	TConstArrayView<FRHIResource*> ReferencedResources = MakeConstArrayView(reinterpret_cast<FRHIResource**>(Data.ReferencedGeometries.GetData()), Data.ReferencedGeometries.Num());
+	
 
-		RHICmdList.UnlockBufferMGPU(AccelerationStructureAddressesBuffer.Buffer, GPUIndex);
-	}
+	RHICmdList.WriteResourceGPUAddresses(ReferencedResources, AccelerationStructureAddressesBuffer.Buffer);
 }
 
 void FRayTracingInstanceBufferBuilder::BuildRayTracingInstanceBuffer(
diff --git a/Engine/Source/Runtime/VulkanRHI/Private/VulkanDynamicRHI.h b/Engine/Source/Runtime/VulkanRHI/Private/VulkanDynamicRHI.h
--- a/Engine/Source/Runtime/VulkanRHI/Private/VulkanDynamicRHI.h
+++ b/Engine/Source/Runtime/VulkanRHI/Private/VulkanDynamicRHI.h
@@ -152,6 +152,8 @@ public:
 	[[nodiscard]] virtual FRHIBufferInitializer RHICreateBufferInitializer(FRHICommandListBase& RHICmdList, const FRHIBufferCreateDesc& CreateDesc) final override;
 
 	virtual void RHIReplaceResources(FRHICommandListBase& RHICmdList, TArray<FRHIResourceReplaceInfo>&& ReplaceInfos) final override;
+	virtual void RHIWriteResourceGPUAddresses(FRHICommandListBase& RHICmdList, uint32 GPUIndex, TConstArrayView<FRHIResource*> Resources, uint64* AddressesPtr) final override;
+
 	virtual void* LockBuffer_BottomOfPipe(FRHICommandListBase& RHICmdList, FRHIBuffer* Buffer, uint32 Offset, uint32 Size, EResourceLockMode LockMode) final override;
 	virtual void UnlockBuffer_BottomOfPipe(FRHICommandListBase& RHICmdList, FRHIBuffer* Buffer) final override;
 	virtual void* RHILockBuffer(FRHICommandListBase& RHICmdList, FRHIBuffer* BufferRHI, uint32 Offset, uint32 Size, EResourceLockMode LockMode) final override;
diff --git a/Engine/Source/Runtime/VulkanRHI/Private/VulkanRHI.cpp b/Engine/Source/Runtime/VulkanRHI/Private/VulkanRHI.cpp
--- a/Engine/Source/Runtime/VulkanRHI/Private/VulkanRHI.cpp
+++ b/Engine/Source/Runtime/VulkanRHI/Private/VulkanRHI.cpp
@@ -2164,6 +2164,38 @@ void FVulkanDynamicRHI::RHIReplaceResources(FRHICommandListBase& RHICmdList, TAr
 	);
 }
 
+void FVulkanDynamicRHI::RHIWriteResourceGPUAddresses(FRHICommandListBase& RHICmdList, uint32 GPUIndex, TConstArrayView<FRHIResource*> Resources, uint64* AddressesPtr)
+{
+	RHICmdList.EnqueueLambda([AddressesPtr, Resources](FRHICommandListBase&)
+		{
+			TRACE_CPUPROFILER_EVENT_SCOPE(RHIWriteResourceGPUAddresses);
+
+			for (int32 ResourceIndex = 0; ResourceIndex < Resources.Num(); ++ResourceIndex)
+			{
+				const FRHIResource* Resource = Resources[ResourceIndex];
+
+				if (Resource->GetType() == RRT_Buffer)
+				{
+					const FVulkanBuffer* Buffer = ResourceCast(static_cast<const FRHIBuffer*>(Resource));
+					AddressesPtr[ResourceIndex] = Buffer->GetDeviceAddress();
+				}
+				else if (Resource->GetType() == RRT_Texture)
+				{
+					unimplemented();
+				}
+				else if (Resource->GetType() == RRT_RayTracingGeometry)
+				{
+					const FVulkanRayTracingGeometry* RayTracingGeometry = ResourceCast(static_cast<const FRHIRayTracingGeometry*>(Resource));
+					AddressesPtr[ResourceIndex] = RayTracingGeometry->GetDeviceAddress();
+				}
+				else
+				{
+					checkNoEntry();
+				}
+			}
+		});
+}
+
 #if PLATFORM_SUPPORTS_BINDLESS_RENDERING
 FRHIResourceCollectionRef FVulkanDynamicRHI::RHICreateResourceCollection(FRHICommandListBase& RHICmdList, TConstArrayView<FRHIResourceCollectionMember> InMembers)
 {
diff --git a/Engine/Source/Runtime/VulkanRHI/Private/VulkanRayTracing.h b/Engine/Source/Runtime/VulkanRHI/Private/VulkanRayTracing.h
--- a/Engine/Source/Runtime/VulkanRHI/Private/VulkanRayTracing.h
+++ b/Engine/Source/Runtime/VulkanRHI/Private/VulkanRayTracing.h
@@ -187,7 +187,11 @@ public:
 	FVulkanRayTracingGeometry(FRHICommandListBase& RHICmdList, const FRayTracingGeometryInitializer& Initializer, FVulkanDevice* InDevice);
 	~FVulkanRayTracingGeometry();
 
-	virtual FRayTracingAccelerationStructureAddress GetAccelerationStructureAddress(uint64 GPUIndex) const final override { return Address; }
+	virtual VkDeviceAddress GetDeviceAddress() const { return Address; }
+
+	PRAGMA_DISABLE_DEPRECATION_WARNINGS
+	virtual FRayTracingAccelerationStructureAddress GetAccelerationStructureAddress(uint64 GPUIndex) const final override { return GetDeviceAddress(); }
+	PRAGMA_ENABLE_DEPRECATION_WARNINGS
 
 	void Swap(FVulkanRayTracingGeometry& Other);
 
