Force alpha writes to remove clear from tonemap for PF_A2B10G10R10 outputs

diff --git a/Engine/Source/Runtime/Renderer/Private/PostProcess/PostProcessTonemap.cpp b/Engine/Source/Runtime/Renderer/Private/PostProcess/PostProcessTonemap.cpp
--- a/Engine/Source/Runtime/Renderer/Private/PostProcess/PostProcessTonemap.cpp
+++ b/Engine/Source/Runtime/Renderer/Private/PostProcess/PostProcessTonemap.cpp
@@ -549,26 +549,25 @@ bool ShouldWriteAlphaChannel(const FViewInfo& View, const FTonemapInputs& Inputs
 	// @todo: Remove this once Oculus fix the bug in their runtime that requires alpha here.
 	const bool bIsStereo = IStereoRendering::IsStereoEyeView(View);
 	const bool bFormatNeedsAlphaWrite = Format == PF_R9G9B9EXP5;
-	return (Inputs.bWriteAlphaChannel || bIsStereo || bFormatNeedsAlphaWrite);
+	const bool bAlphaIsOnlyTwoBitsSoAvoidClear = Format == PF_A2B10G10R10;
+	return (Inputs.bWriteAlphaChannel || bIsStereo || bFormatNeedsAlphaWrite || bAlphaIsOnlyTwoBitsSoAvoidClear);
 }
 
-bool ShouldOverrideOutputLoadActionToFastClear(const FScreenPassRenderTarget& Output, bool bShouldWriteAlphaChannel)
+bool ShouldOverrideOutputLoadActionToClear(const FScreenPassRenderTarget& Output, bool bShouldWriteAlphaChannel)
 {
 	bool bShouldOverride = false;
-	
-	EPixelFormatChannelFlags OutputFlags = GetPixelFormatValidChannels(Output.Texture->Desc.Format);
-	// If we do not write through alpha channel but the output texture has alpha channel
-	// need to override to fast clear load action ERenderTargetLoadAction::Clear, otherwise,
-	// the alpha channel can be garbage data in terms of different driver implementation.
-	if ( (!bShouldWriteAlphaChannel) && EnumHasAnyFlags(OutputFlags, EPixelFormatChannelFlags::A))
-	{
-		bShouldOverride = true;
-	}
 
 	// If the load action is already load, there should be content in the texture. Disable clear overriding.
-	if (Output.LoadAction == ERenderTargetLoadAction::ELoad)
+	if (Output.LoadAction == ERenderTargetLoadAction::ENoAction)
 	{
-		bShouldOverride = false;
+		EPixelFormatChannelFlags OutputFlags = GetPixelFormatValidChannels(Output.Texture->Desc.Format);
+		// If we do not write through alpha channel but the output texture has alpha channel
+		// need to override to fast clear load action ERenderTargetLoadAction::EClear, otherwise,
+		// the alpha channel can be garbage data in terms of different driver implementation.
+		if ((!bShouldWriteAlphaChannel) && EnumHasAnyFlags(OutputFlags, EPixelFormatChannelFlags::A))
+		{
+			bShouldOverride = true;
+		}
 	}
 
 	return bShouldOverride;
@@ -655,7 +654,7 @@ FScreenPassTexture AddTonemapPass(FRDGBuilder& GraphBuilder, const FViewInfo& Vi
 	}
 
 	const bool bShouldWriteAlphaChannel = ShouldWriteAlphaChannel(View, Inputs, Output.Texture->Desc.Format);
-	if (ShouldOverrideOutputLoadActionToFastClear(Output, bShouldWriteAlphaChannel))
+	if (ShouldOverrideOutputLoadActionToClear(Output, bShouldWriteAlphaChannel))
 	{
 		Output.LoadAction = ERenderTargetLoadAction::EClear;
 	}
