Reduce shader permutation

diff --git a/Engine/Source/Runtime/Renderer/Private/IndirectLightRendering.cpp b/Engine/Source/Runtime/Renderer/Private/IndirectLightRendering.cpp
--- a/Engine/Source/Runtime/Renderer/Private/IndirectLightRendering.cpp
+++ b/Engine/Source/Runtime/Renderer/Private/IndirectLightRendering.cpp
@@ -383,6 +383,12 @@ class FReflectionEnvironmentSkyLightingPS : public FGlobalShader
 			PermutationVector.Set<Substrate::FSubstrateTileType>(0);
 		}
 
+		// Hair shading is classified only for complex & complex special tiles
+		if (Substrate::IsSubstrateEnabled() && PermutationVector.Get<FHairComplexTransmittance>() && PermutationVector.Get<Substrate::FSubstrateTileType>() < ESubstrateTileType::EComplex) 
+		{
+			PermutationVector.Set<Substrate::FSubstrateTileType>(0);
+		}
+
 		return PermutationVector;
 	}
 
