Fixed crash with null GPU skin vertex factory contents due to header not being updated when the bone transform allocation doesn't change. Re-enabled the skin extension for GPU skin.

diff --git a/Engine/Source/Runtime/Engine/Private/SkeletalRenderGPUSkin.cpp b/Engine/Source/Runtime/Engine/Private/SkeletalRenderGPUSkin.cpp
--- a/Engine/Source/Runtime/Engine/Private/SkeletalRenderGPUSkin.cpp
+++ b/Engine/Source/Runtime/Engine/Private/SkeletalRenderGPUSkin.cpp
@@ -130,7 +130,7 @@ EBoneTransformStorageMode GetBoneTransformStorageMode()
 #endif
 }
 
-bool GGPUSkinSceneExtensionEnabled = false;
+bool GGPUSkinSceneExtensionEnabled = true;
 static FAutoConsoleVariableRef CVarGPUSkinUseSceneExtension(
 	TEXT("r.GPUSkin.UseSceneExtension"),
 	GGPUSkinSceneExtensionEnabled,
diff --git a/Engine/Source/Runtime/Renderer/Private/Skinning/SkinningSceneExtension.cpp b/Engine/Source/Runtime/Renderer/Private/Skinning/SkinningSceneExtension.cpp
--- a/Engine/Source/Runtime/Renderer/Private/Skinning/SkinningSceneExtension.cpp
+++ b/Engine/Source/Runtime/Renderer/Private/Skinning/SkinningSceneExtension.cpp
@@ -944,10 +944,8 @@ void FSkinningSceneExtension::FUpdater::PostMeshUpdate(FRDGBuilder& GraphBuilder
 
 					FullUpdateList.Emplace(PrimitiveIndex);
 				}
-				else
-				{
-					AddHeaderUpdate(HeaderData, PrimitiveIndex);
-				}
+
+				AddHeaderUpdate(HeaderData, PrimitiveIndex);
 			};
 
 			TConstArrayView<FSkeletalMeshUpdateEvent> UpdateEvents = SkeletalMeshSubmitData->GetUpdateEvents();
