[MegaLights] Fix incorrect guiding due to LightSceneId reuse.
- Moved LightSceneIdToForwardLightIndex to a SceneExtension (mapping is still view-specific based on PersistentViewId).
- Clear entries in LightSceneIdToForwardLightIndex when lights are removed from the scene.

diff --git a/Engine/Source/Runtime/Renderer/Private/LightGrid.h b/Engine/Source/Runtime/Renderer/Private/LightGrid.h
new file mode 100644
--- /dev/null
+++ b/Engine/Source/Runtime/Renderer/Private/LightGrid.h
@@ -0,0 +1,37 @@
+// Copyright Epic Games, Inc. All Rights Reserved.
+
+#pragma once
+
+#include "SceneExtensions.h"
+#include "PrimitiveSceneProxy.h"
+
+class FLightGridSceneExtension : public ISceneExtension
+{
+	DECLARE_SCENE_EXTENSION(RENDERER_API, FLightGridSceneExtension);
+
+	friend class FLightGridSceneExtensionUpdater;
+
+public:
+	using ISceneExtension::ISceneExtension;
+
+	//~ Begin ISceneExtension Interface.
+	virtual ISceneExtensionUpdater* CreateUpdater() override;
+	//~ End ISceneExtension Interface.
+
+	struct FViewData
+	{
+		// Map from Light ID in GPU Scene to index in the View's ForwardLightData array
+		// This is cached between frames so we can access previous frame mapping
+		TArray<int32> LightSceneIdToForwardLightIndex;
+
+		// Similar to LightSceneIdToForwardLightIndex but for simple lights
+		TMap<FSimpleLightId, int32> SimpleLightIdToForwardLightIndex;
+	};
+
+	FViewData* GetViewData(const FViewInfo& View);
+
+private:
+
+	// Indexed by persistent view ID
+	TSparseArray<FViewData> ViewDatas;
+};
diff --git a/Engine/Source/Runtime/Renderer/Private/LightGridInjection.cpp b/Engine/Source/Runtime/Renderer/Private/LightGridInjection.cpp
--- a/Engine/Source/Runtime/Renderer/Private/LightGridInjection.cpp
+++ b/Engine/Source/Runtime/Renderer/Private/LightGridInjection.cpp
@@ -34,6 +34,7 @@ LightGridInjection.cpp
 #include "ShaderPrintParameters.h"
 #include "RenderUtils.h"
 #include "MegaLights/MegaLights.h"
+#include "LightGrid.h"
 #include "LightGridDefinitions.h"
 #include "VolumetricFog.h"
 #include "LightViewData.h"
@@ -290,6 +291,84 @@ FForwardLightUniformParameters::FForwardLightUniformParameters()
 	StaticShadowmapSampler = TStaticSamplerState<SF_Bilinear,AM_Clamp,AM_Clamp,AM_Clamp>::GetRHI();
 }
 
+IMPLEMENT_SCENE_EXTENSION(FLightGridSceneExtension);
+
+class FLightGridSceneExtensionUpdater : public ISceneExtensionUpdater
+{
+	DECLARE_SCENE_EXTENSION_UPDATER(FLightGridSceneExtensionUpdater, FLightGridSceneExtension);
+
+public:
+	FLightGridSceneExtensionUpdater(FLightGridSceneExtension& InSceneExtension) : SceneExtension(InSceneExtension) {}
+
+	virtual void PreLightsUpdate(FRDGBuilder& GraphBuilder, const FLightSceneChangeSet& LightSceneChangeSet) override;
+	virtual void PreSceneUpdate(FRDGBuilder& GraphBuilder, const FScenePreUpdateChangeSet& ChangeSet, FSceneUniformBuffer& SceneUniforms) override;
+	virtual void PostSceneUpdate(FRDGBuilder& GraphBuilder, const FScenePostUpdateChangeSet& ChangeSet) override;
+
+private:
+	FLightGridSceneExtension& SceneExtension;
+
+	const FScenePostUpdateChangeSet* PostUpdateChangeSet = nullptr;
+};
+
+void FLightGridSceneExtensionUpdater::PreLightsUpdate(FRDGBuilder& GraphBuilder, const FLightSceneChangeSet& LightSceneChangeSet)
+{
+	if (LightSceneChangeSet.RemovedLightIds.Num() > 0)
+	{
+		for (FLightGridSceneExtension::FViewData& ViewData : SceneExtension.ViewDatas)
+		{
+			for (int32 Id : LightSceneChangeSet.RemovedLightIds)
+			{
+				if (ViewData.LightSceneIdToForwardLightIndex.IsValidIndex(Id))
+				{
+					ViewData.LightSceneIdToForwardLightIndex[Id] = INDEX_NONE;
+				}
+			}
+		}
+	}
+}
+
+void FLightGridSceneExtensionUpdater::PreSceneUpdate(FRDGBuilder& GraphBuilder, const FScenePreUpdateChangeSet& ChangeSet, FSceneUniformBuffer& SceneUniforms)
+{
+	if (ChangeSet.ViewUpdateChangeSet)
+	{
+		for (FPersistentViewId ViewId : ChangeSet.ViewUpdateChangeSet->RemovedViewIds)
+		{
+			check(SceneExtension.ViewDatas.IsAllocated(ViewId.Index));
+			SceneExtension.ViewDatas.RemoveAt(ViewId.Index);
+		}
+	}
+}
+
+void FLightGridSceneExtensionUpdater::PostSceneUpdate(FRDGBuilder& GraphBuilder, const FScenePostUpdateChangeSet& ChangeSet)
+{
+	if (ChangeSet.ViewUpdateChangeSet)
+	{
+		for (FPersistentViewId ViewId : ChangeSet.ViewUpdateChangeSet->AddedViewIds)
+		{
+			if (!SceneExtension.ViewDatas.IsValidIndex(ViewId.Index))
+			{
+				SceneExtension.ViewDatas.EmplaceAt(ViewId.Index);
+			}
+		}
+	}
+}
+
+ISceneExtensionUpdater* FLightGridSceneExtension::CreateUpdater()
+{
+	return new FLightGridSceneExtensionUpdater(*this);
+}
+
+FLightGridSceneExtension::FViewData* FLightGridSceneExtension::GetViewData(const FViewInfo& View)
+{
+	if (View.PersistentViewId.IsValid())
+	{
+		check(ViewDatas.IsAllocated(View.PersistentViewId.Index));
+		return &ViewDatas[View.PersistentViewId.Index];
+	}
+	
+	return nullptr;
+}
+
 int32 NumCulledLightsGridStride = 2;
 int32 NumCulledGridPrimitiveTypes = 2;
 int32 LightLinkStride = 2;
@@ -847,10 +926,14 @@ FComputeLightGridOutput FSceneRenderer::ComputeLightGrid(FRDGBuilder& GraphBuild
 	static const auto RayEndBiasCVar = IConsoleManager::Get().FindTConsoleVariableDataFloat(TEXT("r.MegaLights.HardwareRayTracing.EndBias"));
 	const float GlobalRayEndBias = RayEndBiasCVar ? RayEndBiasCVar->GetValueOnRenderThread() : 0.0f;
 
+	FLightGridSceneExtension& LightGridExtension = Scene->GetExtension<FLightGridSceneExtension>();
+
 	for (int32 ViewIndex = 0; ViewIndex < Views.Num(); ViewIndex++)
 	{
 		FViewInfo& View = Views[ViewIndex];
 
+		FLightGridSceneExtension::FViewData* LightGridViewData = LightGridExtension.GetViewData(View);
+
 		// Get the associated secondary view if this is an instanced primary view, or vice/versa
 		const FViewInfo* AssociatedSecondaryView = nullptr;
 		const FViewInfo* AssociatedPrimaryView = nullptr;
@@ -926,9 +1009,9 @@ FComputeLightGridOutput FSceneRenderer::ComputeLightGrid(FRDGBuilder& GraphBuild
 			DirectionalMegaLightsSupportedStart = MAX_int32;
 
 			TMap<FSimpleLightId, int32> NewSimpleLightIdToForwardLightIndex;
-			if (View.ViewState)
+			if (LightGridViewData)
 			{
-				NewSimpleLightIdToForwardLightIndex.Reserve(View.ViewState->SimpleLightIdToForwardLightIndex.Num());
+				NewSimpleLightIdToForwardLightIndex.Reserve(LightGridViewData->SimpleLightIdToForwardLightIndex.Num());
 			}
 
 			// Next add all the other lights, track the end index for clustered supporting lights
@@ -966,9 +1049,9 @@ FComputeLightGridOutput FSceneRenderer::ComputeLightGrid(FRDGBuilder& GraphBuild
 					const FSimpleLightPerViewEntry& SimpleLightPerViewData = SortedLightSet.SimpleLights.GetViewDependentData(SimpleLightIndex, ViewIndex, Views.Num());
 
 					int32 PrevForwardLightIndex = INDEX_NONE;
-					if (View.ViewState && SimpleLight.LightId.IsValid())
+					if (LightGridViewData && SimpleLight.LightId.IsValid())
 					{
-						PrevForwardLightIndex = View.ViewState->SimpleLightIdToForwardLightIndex.FindRef(SimpleLight.LightId, INDEX_NONE);
+						PrevForwardLightIndex = LightGridViewData->SimpleLightIdToForwardLightIndex.FindRef(SimpleLight.LightId, INDEX_NONE);
 						NewSimpleLightIdToForwardLightIndex.Add(SimpleLight.LightId, IndexInBuffer);
 					}
 
@@ -1024,10 +1107,25 @@ FComputeLightGridOutput FSceneRenderer::ComputeLightGrid(FRDGBuilder& GraphBuild
 						FForwardLightData& LightData = ForwardLightData[IndexInBuffer];
 
 						int32 PrevForwardLightIndex = INDEX_NONE;
-						if (View.ViewState)
+						if (LightGridViewData)
 						{
-							PrevForwardLightIndex = View.ViewState->LightSceneIdToForwardLightIndex.FindOrAdd(LightSceneInfo->Id, INDEX_NONE);
-							View.ViewState->LightSceneIdToForwardLightIndex[LightSceneInfo->Id] = IndexInBuffer;
+							if (LightGridViewData->LightSceneIdToForwardLightIndex.IsValidIndex(LightSceneInfo->Id))
+							{
+								PrevForwardLightIndex = LightGridViewData->LightSceneIdToForwardLightIndex[LightSceneInfo->Id];
+							}
+							else
+							{
+								PrevForwardLightIndex = INDEX_NONE;
+
+								const int32 OldSize = LightGridViewData->LightSceneIdToForwardLightIndex.Num();
+								LightGridViewData->LightSceneIdToForwardLightIndex.SetNumUninitialized(LightSceneInfo->Id + 1);
+
+								for (int32 Index = OldSize; Index < LightGridViewData->LightSceneIdToForwardLightIndex.Num(); ++Index)
+								{
+									LightGridViewData->LightSceneIdToForwardLightIndex[Index] = INDEX_NONE;
+								}
+							}
+							LightGridViewData->LightSceneIdToForwardLightIndex[LightSceneInfo->Id] = IndexInBuffer;
 						}
 
 						if (SortedLightInfo.SortKey.Fields.LightType != LightType_Directional) // only local lights go into the grid
@@ -1252,9 +1350,9 @@ FComputeLightGridOutput FSceneRenderer::ComputeLightGrid(FRDGBuilder& GraphBuild
 				}
 			}
 
-			if (View.ViewState)
+			if (LightGridViewData)
 			{
-				View.ViewState->SimpleLightIdToForwardLightIndex = MoveTemp(NewSimpleLightIdToForwardLightIndex);
+				LightGridViewData->SimpleLightIdToForwardLightIndex = MoveTemp(NewSimpleLightIdToForwardLightIndex);
 			}
 
 			// 3. add simple lights into ForwardLightData and fill uninitialized ViewSpacePosAndRadiusData/IndirectionIndices
diff --git a/Engine/Source/Runtime/Renderer/Private/ScenePrivate.h b/Engine/Source/Runtime/Renderer/Private/ScenePrivate.h
--- a/Engine/Source/Runtime/Renderer/Private/ScenePrivate.h
+++ b/Engine/Source/Runtime/Renderer/Private/ScenePrivate.h
@@ -1019,13 +1019,6 @@ public:
 
 	TMap<int32, TRDGUniformBufferRef<FBeerShadowMapUniformBufferParameters>> BeerShadowMapUniformBufferMap;
 
-	// Map from Light ID in GPU Scene to index in the View's ForwardLightData array
-	// This is stored in ViewState so we can access previous frame mapping
-	TMap<int32, int32> LightSceneIdToForwardLightIndex;
-
-	// Similar to LightSceneIdToForwardLightIndex but for simple lights
-	TMap<FSimpleLightId, int32> SimpleLightIdToForwardLightIndex;
-
 	// Pre-computed filter in spectral (i.e. FFT) domain along with data to determine if we need to up date it
 	struct {
 		/// @cond DOXYGEN_WARNINGS
diff --git a/Engine/Source/Runtime/Renderer/Private/ShadowRendering.cpp b/Engine/Source/Runtime/Renderer/Private/ShadowRendering.cpp
--- a/Engine/Source/Runtime/Renderer/Private/ShadowRendering.cpp
+++ b/Engine/Source/Runtime/Renderer/Private/ShadowRendering.cpp
@@ -2307,7 +2307,9 @@ void FDeferredShadingSceneRenderer::CollectLightForTranslucencyLightingVolumeInj
 	const FLightSceneInfo* LightSceneInfo,
 	bool bSupportShadowMaps,
 	FTranslucentLightInjectionCollector& Collector)
-{	
+{
+	FLightGridSceneExtension& LightGridExtension = Scene->GetExtension<FLightGridSceneExtension>();
+
 	bool bInjectedShadowedLight = false;
 	if (bSupportShadowMaps)
 	{
@@ -2335,7 +2337,9 @@ void FDeferredShadingSceneRenderer::CollectLightForTranslucencyLightingVolumeInj
 						const FViewInfo* DependentView = ProjectedShadowInfo->DependentView;
 						if (DependentView == nullptr || DependentView == &View)
 						{
-							Collector.AddLightForInjection(View, ViewIndex, VisibleLightInfos, *LightSceneInfo, ProjectedShadowInfo);
+							FLightGridSceneExtension::FViewData* LightGridViewData = LightGridExtension.GetViewData(View);
+
+							Collector.AddLightForInjection(View, ViewIndex, LightGridViewData, VisibleLightInfos, *LightSceneInfo, ProjectedShadowInfo);
 							bInjectedShadowedLight = true;
 						}
 					}
@@ -2352,7 +2356,9 @@ void FDeferredShadingSceneRenderer::CollectLightForTranslucencyLightingVolumeInj
 			FViewInfo& View = Views[ViewIndex];
 			if (View.IsPrimarySceneView() && LightSceneInfo->ShouldRenderLight(View))
 			{
-				Collector.AddLightForInjection(View, ViewIndex, VisibleLightInfos, *LightSceneInfo);
+				FLightGridSceneExtension::FViewData* LightGridViewData = LightGridExtension.GetViewData(View);
+
+				Collector.AddLightForInjection(View, ViewIndex, LightGridViewData, VisibleLightInfos, *LightSceneInfo);
 			}
 		}
 	}
diff --git a/Engine/Source/Runtime/Renderer/Private/TranslucentLighting.cpp b/Engine/Source/Runtime/Renderer/Private/TranslucentLighting.cpp
--- a/Engine/Source/Runtime/Renderer/Private/TranslucentLighting.cpp
+++ b/Engine/Source/Runtime/Renderer/Private/TranslucentLighting.cpp
@@ -38,6 +38,7 @@
 #include "ClearQuad.h"
 #include "ScenePrivate.h"
 #include "OneColorShader.h"
+#include "LightGrid.h"
 #include "LightFunctionRendering.h"
 #include "LightRendering.h"
 #include "ScreenRendering.h"
@@ -1599,6 +1600,7 @@ FTranslucentLightInjectionCollector::FTranslucentLightInjectionCollector(
 void FTranslucentLightInjectionCollector::AddLightForInjection(
 	const FViewInfo& View,
 	const uint32 ViewIndex,
+	FLightGridSceneExtension::FViewData* LightGridViewData,
 	TArrayView<const FVisibleLightInfo> VisibleLightInfos,
 	const FLightSceneInfo& LightSceneInfo,
 	const FProjectedShadowInfo* InProjectedShadowInfo)
@@ -1659,13 +1661,17 @@ void FTranslucentLightInjectionCollector::AddLightForInjection(
 			!bUseAdaptiveVolumetricShadowMap;
 
 		// If it would otherwise be supported, see if we can find the local light index in the light grid
-		if (bSupportsBatching && View.ForwardLightingResources.ForwardLightUniformBuffer && View.ViewState)
+		if (bSupportsBatching && View.ForwardLightingResources.ForwardLightUniformBuffer && LightGridViewData)
 		{
-			const int32* Value = View.ViewState->LightSceneIdToForwardLightIndex.Find(LightSceneInfo.Id);
-			if (Value)
+			int32 ForwardLightIndex = INDEX_NONE;
+			
+			if (LightGridViewData->LightSceneIdToForwardLightIndex.IsValidIndex(LightSceneInfo.Id))
 			{
-				const int32 ForwardLightIndex = *Value;
+				ForwardLightIndex = LightGridViewData->LightSceneIdToForwardLightIndex[LightSceneInfo.Id];
+			}
 
+			if (ForwardLightIndex != INDEX_NONE)
+			{
 				FPerViewData& InjectionData = InjectionDataPerView[ViewIndex];
 				// There shouldn't be any duplication here
 				++InjectionData.BatchedLocalLightCount;
diff --git a/Engine/Source/Runtime/Renderer/Private/TranslucentLighting.h b/Engine/Source/Runtime/Renderer/Private/TranslucentLighting.h
--- a/Engine/Source/Runtime/Renderer/Private/TranslucentLighting.h
+++ b/Engine/Source/Runtime/Renderer/Private/TranslucentLighting.h
@@ -5,6 +5,7 @@
 #include "SceneRendering.h"
 #include "RenderGraph.h"
 #include "VolumeRendering.h"
+#include "LightGrid.h"
 
 struct FSortedLightSceneInfo;
 struct FCloudShadowAOData;
@@ -105,6 +106,7 @@ public:
 	void AddLightForInjection(
 		const FViewInfo& View,
 		const uint32 ViewIndex,
+		FLightGridSceneExtension::FViewData* LightGridViewData,
 		TArrayView<const FVisibleLightInfo> VisibleLightInfos,
 		const FLightSceneInfo& LightSceneInfo,
 		const FProjectedShadowInfo* InProjectedShadowInfo = nullptr);
