Always use volume texture and compute shader for tonemapping LUT

Also disable UAV capability for some pixel formats on Apple2 devices (A8 SoC) to match reality.

diff --git a/Engine/Shaders/Private/PostProcessCombineLUTs.usf b/Engine/Shaders/Private/PostProcessCombineLUTs.usf
--- a/Engine/Shaders/Private/PostProcessCombineLUTs.usf
+++ b/Engine/Shaders/Private/PostProcessCombineLUTs.usf
@@ -156,37 +156,8 @@ float BlueCorrection;
 float ExpandGamut;
 float ToneCurveAmount;
 
-float4 CombineLUTsCommon(float2 InUV, uint InLayerIndex)
+float4 CombineLUTsCommon(float3 LUTEncodedColor)
 {
-#if USE_VOLUME_LUT == 1
-	// construct the neutral color from a 3d position volume texture	
-	float4 Neutral;
-	{
-		float2 UV = InUV - float2(0.5f / LUTSize, 0.5f / LUTSize);
-
-		Neutral = float4(UV * LUTSize / (LUTSize - 1), InLayerIndex / (LUTSize - 1), 0);
-	}
-#else
-	// construct the neutral color from a 2d position in 256x16
-	float4 Neutral;
-	{ 
-		float2 UV = InUV;
-
-		// 0.49999f instead of 0.5f to avoid getting into negative values
-		UV -= float2(0.49999f / (LUTSize * LUTSize), 0.49999f / LUTSize);
-
-		float Scale = LUTSize / (LUTSize - 1);
-
-		float3 RGB;
-		
-		RGB.r = frac(UV.x * LUTSize);
-		RGB.b = UV.x - RGB.r / LUTSize;
-		RGB.g = UV.y;
-
-		Neutral = float4(RGB * Scale, 0);
-	}
-#endif
-
 	float4 OutColor = 0;
 	
 	const float3x3 AP1_2_sRGB = mul( XYZ_2_sRGB_MAT, mul( D60_2_D65_CAT, AP1_2_XYZ_MAT ) );
@@ -195,7 +166,6 @@ float4 CombineLUTsCommon(float2 InUV, uint InLayerIndex)
 
 	const float3x3 AP1_2_Output  = OuputGamutMappingMatrix( OutputGamut );
 
-	float3 LUTEncodedColor = Neutral.rgb;
 	float3 LinearColor;
 	// Decode texture values as ST-2084 (Dolby PQ)
 	BRANCH
@@ -443,50 +413,16 @@ float4 CombineLUTsCommon(float2 InUV, uint InLayerIndex)
 	return OutColor;
 }
 
-// todo: Weight[0] should be used for neutral, Texture* name should start with 1, color correction should apply on top of that
-#if USE_VOLUME_LUT == 1
-void MainPS(FWriteToSliceGeometryOutput Input, out float4 OutColor : SV_Target0)
-{
-	OutColor = CombineLUTsCommon(Input.Vertex.UV, Input.LayerIndex);
-}
-#else
-void MainPS(noperspective float4 InUV : TEXCOORD0, out float4 OutColor : SV_Target0)
-{
-	OutColor = CombineLUTsCommon(InUV.xy, 0);
-}
-#endif
-
-#if COMPUTESHADER
-// xy: Unused, zw: ThreadToUVScale
-float2 OutputExtentInverse;
+float InvLUTSizeMinusOne;
 
-#if USE_VOLUME_LUT == 1
 RWTexture3D<float4> RWOutputTexture;
 
 [numthreads(THREADGROUP_SIZE, THREADGROUP_SIZE, THREADGROUP_SIZE)]
 void MainCS(uint3 DispatchThreadId : SV_DispatchThreadID) 
 {
-	float2 UV = ((float2)DispatchThreadId.xy + 0.5f) * OutputExtentInverse;
-	uint LayerIndex = DispatchThreadId.z;
-	
-	float4 OutColor = CombineLUTsCommon(UV, LayerIndex);
+	float3 LUTEncodedColor = (float3)DispatchThreadId * InvLUTSizeMinusOne;
+	float4 OutColor = CombineLUTsCommon(LUTEncodedColor);
 
 	uint3 PixelPos = DispatchThreadId;
 	RWOutputTexture[PixelPos] = OutColor;
 }
-#else
-RWTexture2D<float4> RWOutputTexture;
-
-[numthreads(THREADGROUP_SIZE, THREADGROUP_SIZE, 1)]
-void MainCS(uint2 DispatchThreadId : SV_DispatchThreadID)
-{
-	float2 UV = ((float2)DispatchThreadId + 0.5f) * OutputExtentInverse;
-	uint LayerIndex = 0;
-	
-	float4 OutColor = CombineLUTsCommon(UV, LayerIndex);
-
-	uint2 PixelPos = DispatchThreadId;
-	RWOutputTexture[PixelPos] = OutColor;
-}
-#endif
-#endif
diff --git a/Engine/Shaders/Private/PostProcessTonemap.usf b/Engine/Shaders/Private/PostProcessTonemap.usf
--- a/Engine/Shaders/Private/PostProcessTonemap.usf
+++ b/Engine/Shaders/Private/PostProcessTonemap.usf
@@ -159,12 +159,8 @@ float ComputeFilmGrainIntensity(float3 LinearColorRGBGamut)
 
 
 // LUT for color grading
-#if USE_VOLUME_LUT == 1
-	Texture3D		ColorGradingLUT;
-#else
-	Texture2D       ColorGradingLUT;
-#endif 
-SamplerState        ColorGradingLUTSampler;
+Texture3D ColorGradingLUT;
+SamplerState ColorGradingLUTSampler;
 
 float LUTSize;
 float InvLUTSize; // 1 / LUTSize
@@ -190,12 +186,7 @@ half3 ColorLookupTable( half3 LinearColor )
 
 	float3 UVW = LUTEncodedColor * LUTScale + LUTOffset;
 
-#if USE_VOLUME_LUT == 1
 	half3 OutDeviceColor = Texture3DSample( ColorGradingLUT, ColorGradingLUTSampler, UVW ).rgb;
-#else
-	half3 OutDeviceColor = UnwrappedTexture3DSample( ColorGradingLUT, ColorGradingLUTSampler, UVW, LUTSize, InvLUTSize ).rgb;
-#endif
-	
 	return OutDeviceColor * 1.05;
 }
 
diff --git a/Engine/Source/Runtime/Apple/MetalRHI/Private/MetalRHI.cpp b/Engine/Source/Runtime/Apple/MetalRHI/Private/MetalRHI.cpp
--- a/Engine/Source/Runtime/Apple/MetalRHI/Private/MetalRHI.cpp
+++ b/Engine/Source/Runtime/Apple/MetalRHI/Private/MetalRHI.cpp
@@ -1135,6 +1135,15 @@ FMetalDynamicRHI::FMetalDynamicRHI(ERHIFeatureLevel::Type RequestedFeatureLevel)
 
 	RHIInitDefaultPixelFormatCapabilities();
 
+	if (MTLDevice->supportsFamily(MTL::GPUFamilyApple2) && !MTLDevice->supportsFamily(MTL::GPUFamilyApple3))
+	{
+		UE_LOG(LogMetal, Display, TEXT("Detected GPUFamilyApple2 device."));
+		EnumRemoveFlags(GPixelFormats[PF_A2B10G10R10].Capabilities, EPixelFormatCapabilities::UAV);
+		EnumRemoveFlags(GPixelFormats[PF_FloatRGB].Capabilities, EPixelFormatCapabilities::UAV);
+		EnumRemoveFlags(GPixelFormats[PF_FloatR11G11B10].Capabilities, EPixelFormatCapabilities::UAV);
+		EnumRemoveFlags(GPixelFormats[PF_R9G9B9EXP5].Capabilities, EPixelFormatCapabilities::UAV);
+	}
+
 	auto AddTypedUAVSupport = [](EPixelFormat InPixelFormat)
 	{
 		EnumAddFlags(GPixelFormats[InPixelFormat].Capabilities, EPixelFormatCapabilities::TypedUAVLoad | EPixelFormatCapabilities::TypedUAVStore);
diff --git a/Engine/Source/Runtime/Renderer/Private/PostProcess/PostProcessCombineLUTs.cpp b/Engine/Source/Runtime/Renderer/Private/PostProcess/PostProcessCombineLUTs.cpp
--- a/Engine/Source/Runtime/Renderer/Private/PostProcess/PostProcessCombineLUTs.cpp
+++ b/Engine/Source/Runtime/Renderer/Private/PostProcess/PostProcessCombineLUTs.cpp
@@ -73,15 +73,6 @@ struct FColorTransform
 };
 } //! namespace
 
-// false:use 1024x32 texture / true:use volume texture (faster, requires geometry shader)
-// USE_VOLUME_LUT: needs to be the same for C++ and HLSL.
-// Safe to use at pipeline and run time.
-bool PipelineVolumeTextureLUTSupportGuaranteedAtRuntime(EShaderPlatform Platform)
-{
-	// This is used to know if the target shader platform does not support required volume texture features we need for sure (read, render to).
-	return RHIVolumeTextureRenderingSupportGuaranteed(Platform) && (RHISupportsGeometryShaders(Platform) || RHISupportsVertexShaderLayer(Platform));
-}
-
 FVector3f GetMappingPolynomial()
 {
 	FColorTransform ColorTransform;
@@ -182,16 +173,14 @@ struct FCachedLUTSettings
 	EShaderPlatform ShaderPlatform = GMaxRHIShaderPlatform;
 	FCombineLUTParameters Parameters;
 	FWorkingColorSpaceShaderParameters WorkingColorSpaceShaderParameters;
-	bool bUseCompute = false;
 	int32 AcesVersion = 0;
 
-	bool UpdateCachedValues(FRDGBuilder& GraphBuilder, const FViewInfo& View, const FTexture* const* Textures, const float* Weights, uint32 BlendCount, bool bInUseCompute)
+	bool UpdateCachedValues(FRDGBuilder& GraphBuilder, const FViewInfo& View, const FTexture* const* Textures, const float* Weights, uint32 BlendCount)
 	{
 		bool bHasChanged = false;
 		GetCombineLUTParameters(GraphBuilder, View, Textures, Weights, BlendCount, bHasChanged);
 		UPDATE_CACHE_SETTINGS(UniqueID, View.State ? View.State->GetViewKey() : 0, bHasChanged);
 		UPDATE_CACHE_SETTINGS(ShaderPlatform, View.GetShaderPlatform(), bHasChanged);
-		UPDATE_CACHE_SETTINGS(bUseCompute, bInUseCompute, bHasChanged);
 		static const auto CVarAcesVersion = IConsoleManager::Get().FindTConsoleVariableDataInt(TEXT("r.HDR.Aces.Version"));
 		UPDATE_CACHE_SETTINGS(AcesVersion, CVarAcesVersion->GetValueOnRenderThread(), bHasChanged);
 
@@ -324,77 +313,29 @@ struct FCachedLUTSettings
 	}
 };
 
-class FLUTBlenderShader : public FGlobalShader
+class FLUTBlenderCS : public FGlobalShader
 {
 public:
+	static constexpr int32 GroupSize = 4;
+
 	class FBlendCount : SHADER_PERMUTATION_RANGE_INT("BLENDCOUNT", 1, 5);
 	class FOutputDeviceSRGB : SHADER_PERMUTATION_BOOL("OUTPUT_DEVICE_SRGB");
 	class FUseACES2 : SHADER_PERMUTATION_BOOL("USE_ACES_2");
 	using FPermutationDomain = TShaderPermutationDomain<FBlendCount, FOutputDeviceSRGB, FUseACES2>;
 
-	static bool ShouldCompilePermutation(const FGlobalShaderPermutationParameters& Parameters)
-	{
-		return true;
-	}
-
-	static void ModifyCompilationEnvironment(const FGlobalShaderPermutationParameters& Parameters, FShaderCompilerEnvironment& OutEnvironment)
-	{
-		FGlobalShader::ModifyCompilationEnvironment(Parameters, OutEnvironment);
-
-		const int32 UseVolumeLUT = PipelineVolumeTextureLUTSupportGuaranteedAtRuntime(Parameters.Platform) ? 1 : 0;
-		OutEnvironment.SetDefine(TEXT("USE_VOLUME_LUT"), UseVolumeLUT);
-	}
-
-	FLUTBlenderShader() = default;
-	FLUTBlenderShader(const CompiledShaderInitializerType& Initializer)
-		: FGlobalShader(Initializer)
-	{}
-};
-
-class FLUTBlenderPS : public FLUTBlenderShader
-{
-public:
-	DECLARE_GLOBAL_SHADER(FLUTBlenderPS);
-	SHADER_USE_PARAMETER_STRUCT(FLUTBlenderPS, FLUTBlenderShader);
-
-	BEGIN_SHADER_PARAMETER_STRUCT(FParameters, )
-		SHADER_PARAMETER_STRUCT_INCLUDE(FCombineLUTParameters, CombineLUT)
-		RENDER_TARGET_BINDING_SLOTS()
-	END_SHADER_PARAMETER_STRUCT()
-
-	static bool ShouldCompilePermutation(const FGlobalShaderPermutationParameters& Parameters)
-	{
-		return true;
-	}
-};
-
-IMPLEMENT_GLOBAL_SHADER(FLUTBlenderPS, "/Engine/Private/PostProcessCombineLUTs.usf", "MainPS", SF_Pixel);
-
-class FLUTBlenderCS : public FLUTBlenderShader
-{
-public:
-	static constexpr int32 GroupSize2D = 8;
-	static constexpr int32 GroupSize3D = 4;
-
 	DECLARE_GLOBAL_SHADER(FLUTBlenderCS);
-	SHADER_USE_PARAMETER_STRUCT(FLUTBlenderCS, FLUTBlenderShader);
+	SHADER_USE_PARAMETER_STRUCT(FLUTBlenderCS, FGlobalShader);
 
 	BEGIN_SHADER_PARAMETER_STRUCT(FParameters, )
 		SHADER_PARAMETER_STRUCT_INCLUDE(FCombineLUTParameters, CombineLUT)
-		SHADER_PARAMETER(FVector2f, OutputExtentInverse)
-		SHADER_PARAMETER_RDG_TEXTURE_UAV(RWTexture2D<float4>, RWOutputTexture)
+		SHADER_PARAMETER(float, InvLUTSizeMinusOne)
+		SHADER_PARAMETER_RDG_TEXTURE_UAV(RWTexture3D<float4>, RWOutputTexture)
 	END_SHADER_PARAMETER_STRUCT()
 
-	static bool ShouldCompilePermutation(const FGlobalShaderPermutationParameters& Parameters)
-	{
-		return IsFeatureLevelSupported(Parameters.Platform, ERHIFeatureLevel::SM5);
-	}
-
 	static void ModifyCompilationEnvironment(const FGlobalShaderPermutationParameters& Parameters, FShaderCompilerEnvironment& OutEnvironment)
 	{
-		FLUTBlenderShader::ModifyCompilationEnvironment(Parameters, OutEnvironment);
+		FGlobalShader::ModifyCompilationEnvironment(Parameters, OutEnvironment);
 
-		const int32 GroupSize = PipelineVolumeTextureLUTSupportGuaranteedAtRuntime(Parameters.Platform) ? GroupSize3D : GroupSize2D;
 		OutEnvironment.SetDefine(TEXT("THREADGROUP_SIZE"), GroupSize);
 	}
 };
@@ -518,9 +459,6 @@ FRDGTextureRef AddCombineLUTPass(FRDGBuilder& GraphBuilder, const FViewInfo& Vie
 		LocalCount = GenerateFinalTable(View.FinalPostProcessSettings, LocalTextures, LocalWeights, GMaxLUTBlendCount);
 	}
 
-	const bool bUseVolumeTextureLUT = PipelineVolumeTextureLUTSupportGuaranteedAtRuntime(View.GetShaderPlatform());
-	const bool bUseComputePass = bUseVolumeTextureLUT || View.bUseComputePasses;
-
 	const EDisplayOutputFormat OutputDevice = ViewFamily.RenderTarget->GetDisplayOutputFormat();
 	// Note: We also enforce floats with HDR ScRGB (editor only) to preserve negative values.
 	const bool bUseFloatOutput = ViewFamily.SceneCaptureSource == SCS_FinalColorHDR ||
@@ -531,11 +469,11 @@ FRDGTextureRef AddCombineLUTPass(FRDGBuilder& GraphBuilder, const FViewInfo& Vie
 	// Attempt to register the persistent view LUT texture.
 	const int32 LUTSize = CVarLUTSize->GetInt();
 	FRDGTextureRef OutputTexture = TryRegisterExternalTexture(GraphBuilder,
-		View.GetTonemappingLUT(GraphBuilder.RHICmdList, LUTSize, bUseVolumeTextureLUT, bUseComputePass, bUseFloatOutput));
+		View.GetTonemappingLUT(GraphBuilder.RHICmdList, LUTSize, bUseFloatOutput));
 
 	View.SetValidTonemappingLUT();
 
-	const bool bHasChanged = CachedLUTSettings.UpdateCachedValues(GraphBuilder, View, LocalTextures, LocalWeights, LocalCount, bUseComputePass);
+	const bool bHasChanged = CachedLUTSettings.UpdateCachedValues(GraphBuilder, View, LocalTextures, LocalWeights, LocalCount);
 	if (!bHasChanged && OutputTexture && CVarLUTUpdateEveryFrame.GetValueOnRenderThread() == 0)
 	{
 		return OutputTexture;
@@ -545,109 +483,33 @@ FRDGTextureRef AddCombineLUTPass(FRDGBuilder& GraphBuilder, const FViewInfo& Vie
 	if (!OutputTexture)
 	{
 		OutputTexture = GraphBuilder.CreateTexture(
-			Translate(FSceneViewState::CreateLUTRenderTarget(LUTSize, bUseVolumeTextureLUT, bUseComputePass, bUseFloatOutput)),
+			Translate(FSceneViewState::CreateLUTRenderTarget(LUTSize, bUseFloatOutput)),
 			TEXT("CombineLUT"));
 	}
 
 	check(OutputTexture);
 
-	// For a 3D texture, the viewport defaults to 32x32 (per slice); for a 2D texture, it's unwrapped to 1024x32.
-	const FIntPoint OutputViewSize(bUseVolumeTextureLUT ? LUTSize : LUTSize * LUTSize, LUTSize);
-
-	FLUTBlenderShader::FPermutationDomain PermutationVector;
-	PermutationVector.Set<FLUTBlenderShader::FBlendCount>(LocalCount);
-	PermutationVector.Set<FLUTBlenderShader::FUseACES2>(CachedLUTSettings.AcesVersion > 1);
-
-	if (bUseComputePass)
-	{
-		FLUTBlenderCS::FParameters* PassParameters = GraphBuilder.AllocParameters<FLUTBlenderCS::FParameters>();
-		PassParameters->CombineLUT = CachedLUTSettings.Parameters;
-		PassParameters->OutputExtentInverse = FVector2f(1.0f, 1.0f) / FVector2f(OutputViewSize);
-		PassParameters->RWOutputTexture = GraphBuilder.CreateUAV(OutputTexture);
-
-		const bool bOutputDeviceSRGB = (PassParameters->CombineLUT.OutputDevice.OutputDevice == (uint32)EDisplayOutputFormat::SDR_sRGB);
-
-		PermutationVector.Set<FLUTBlenderShader::FOutputDeviceSRGB>(bOutputDeviceSRGB);
-		TShaderMapRef<FLUTBlenderCS> ComputeShader(View.ShaderMap, PermutationVector);
-
-		const int32 GroupSize = bUseVolumeTextureLUT ? FLUTBlenderCS::GroupSize3D : FLUTBlenderCS::GroupSize2D;
-		const uint32 GroupSizeXY = FMath::DivideAndRoundUp(OutputViewSize.X, GroupSize);
-		const uint32 GroupSizeZ = bUseVolumeTextureLUT ? GroupSizeXY : 1;
-
-		FComputeShaderUtils::AddPass(
-			GraphBuilder,
-			RDG_EVENT_NAME("CombineLUTs %d (CS)", LUTSize),
-			ComputeShader,
-			PassParameters,
-			FIntVector(GroupSizeXY, GroupSizeXY, GroupSizeZ));
-	}
-	else
-	{
-		FLUTBlenderPS::FParameters* PassParameters = GraphBuilder.AllocParameters<FLUTBlenderPS::FParameters>();
-		PassParameters->CombineLUT = CachedLUTSettings.Parameters;
-		PassParameters->RenderTargets[0] = FRenderTargetBinding(OutputTexture, ERenderTargetLoadAction::ENoAction);
-
-		const bool bOutputDeviceSRGB = (PassParameters->CombineLUT.OutputDevice.OutputDevice == (uint32)EDisplayOutputFormat::SDR_sRGB);
-
-		PermutationVector.Set<FLUTBlenderShader::FOutputDeviceSRGB>(bOutputDeviceSRGB);
-		TShaderMapRef<FLUTBlenderPS> PixelShader(View.ShaderMap, PermutationVector);
-
-		GraphBuilder.AddPass(
-			RDG_EVENT_NAME("CombineLUTS %d (PS)", LUTSize),
-			PassParameters,
-			ERDGPassFlags::Raster,
-			[&View, PixelShader, PassParameters, bUseVolumeTextureLUT, LUTSize] (FRDGAsyncTask, FRHICommandList& RHICmdList)
-		{
-			FGraphicsPipelineStateInitializer GraphicsPSOInit;
-			RHICmdList.ApplyCachedRenderTargets(GraphicsPSOInit);
-			GraphicsPSOInit.BlendState = TStaticBlendState<>::GetRHI();
-			GraphicsPSOInit.RasterizerState = TStaticRasterizerState<>::GetRHI();
-			GraphicsPSOInit.DepthStencilState = TStaticDepthStencilState<false, CF_Always>::GetRHI();
-
-			if (bUseVolumeTextureLUT)
-			{
-				const FVolumeBounds VolumeBounds(LUTSize);
+	FLUTBlenderCS::FPermutationDomain PermutationVector;
+	PermutationVector.Set<FLUTBlenderCS::FBlendCount>(LocalCount);
+	PermutationVector.Set<FLUTBlenderCS::FUseACES2>(CachedLUTSettings.AcesVersion > 1);
 
-				TShaderMapRef<FWriteToSliceVS> VertexShader(View.ShaderMap);
-				TOptionalShaderMapRef<FWriteToSliceGS> GeometryShader(View.ShaderMap);
+	FLUTBlenderCS::FParameters* PassParameters = GraphBuilder.AllocParameters<FLUTBlenderCS::FParameters>();
+	PassParameters->CombineLUT = CachedLUTSettings.Parameters;
+	PassParameters->InvLUTSizeMinusOne = 1.f / (LUTSize - 1);
+	PassParameters->RWOutputTexture = GraphBuilder.CreateUAV(OutputTexture);
 
-				GraphicsPSOInit.PrimitiveType = PT_TriangleStrip;
-				GraphicsPSOInit.BoundShaderState.VertexDeclarationRHI = GScreenVertexDeclaration.VertexDeclarationRHI;
-				GraphicsPSOInit.BoundShaderState.VertexShaderRHI = VertexShader.GetVertexShader();
-				GraphicsPSOInit.BoundShaderState.SetGeometryShader(GeometryShader.GetGeometryShader());
-				GraphicsPSOInit.BoundShaderState.PixelShaderRHI = PixelShader.GetPixelShader();
-				SetGraphicsPipelineState(RHICmdList, GraphicsPSOInit, 0);
+	const bool bOutputDeviceSRGB = (PassParameters->CombineLUT.OutputDevice.OutputDevice == (uint32)EDisplayOutputFormat::SDR_sRGB);
 
-				SetShaderParametersLegacyVS(RHICmdList, VertexShader, VolumeBounds, FIntVector(VolumeBounds.MaxX - VolumeBounds.MinX));
-				SetShaderParameters(RHICmdList, PixelShader, PixelShader.GetPixelShader(), *PassParameters);
+	PermutationVector.Set<FLUTBlenderCS::FOutputDeviceSRGB>(bOutputDeviceSRGB);
+	TShaderMapRef<FLUTBlenderCS> ComputeShader(View.ShaderMap, PermutationVector);
 
-				RasterizeToVolumeTexture(RHICmdList, VolumeBounds);
-			}
-			else
-			{
-				TShaderMapRef<FScreenPassVS> VertexShader(View.ShaderMap);
-
-				GraphicsPSOInit.PrimitiveType = PT_TriangleList;
-				GraphicsPSOInit.BoundShaderState.VertexDeclarationRHI = GFilterVertexDeclaration.VertexDeclarationRHI;
-				GraphicsPSOInit.BoundShaderState.VertexShaderRHI = VertexShader.GetVertexShader();
-				GraphicsPSOInit.BoundShaderState.PixelShaderRHI = PixelShader.GetPixelShader();
-				SetGraphicsPipelineState(RHICmdList, GraphicsPSOInit, 0);
-
-				SetShaderParameters(RHICmdList, PixelShader, PixelShader.GetPixelShader(), *PassParameters);
-
-				DrawRectangle(
-					RHICmdList,
-					0, 0,										// XY
-					LUTSize * LUTSize, LUTSize,				// SizeXY
-					0, 0,										// UV
-					LUTSize * LUTSize, LUTSize,				// SizeUV
-					FIntPoint(LUTSize * LUTSize, LUTSize),	// TargetSize
-					FIntPoint(LUTSize * LUTSize, LUTSize),	// TextureSize
-					VertexShader,
-					EDRF_UseTriangleOptimization);
-			}
-		});
-	}
+	const uint32 GroupSize = FMath::DivideAndRoundUp(LUTSize, FLUTBlenderCS::GroupSize);
+	FComputeShaderUtils::AddPass(
+		GraphBuilder,
+		RDG_EVENT_NAME("CombineLUTs %d (CS)", LUTSize),
+		ComputeShader,
+		PassParameters,
+		FIntVector(GroupSize, GroupSize, GroupSize));
 
 	return OutputTexture;
 }
diff --git a/Engine/Source/Runtime/Renderer/Private/PostProcess/PostProcessCombineLUTs.h b/Engine/Source/Runtime/Renderer/Private/PostProcess/PostProcessCombineLUTs.h
--- a/Engine/Source/Runtime/Renderer/Private/PostProcess/PostProcessCombineLUTs.h
+++ b/Engine/Source/Runtime/Renderer/Private/PostProcess/PostProcessCombineLUTs.h
@@ -10,6 +10,4 @@ END_SHADER_PARAMETER_STRUCT()
 
 FColorRemapParameters GetColorRemapParameters();
 
-bool PipelineVolumeTextureLUTSupportGuaranteedAtRuntime(EShaderPlatform Platform);
-
-FRDGTextureRef AddCombineLUTPass(FRDGBuilder& GraphBuilder, const FViewInfo& View);
\ No newline at end of file
+FRDGTextureRef AddCombineLUTPass(FRDGBuilder& GraphBuilder, const FViewInfo& View);
diff --git a/Engine/Source/Runtime/Renderer/Private/PostProcess/PostProcessTonemap.cpp b/Engine/Source/Runtime/Renderer/Private/PostProcess/PostProcessTonemap.cpp
--- a/Engine/Source/Runtime/Renderer/Private/PostProcess/PostProcessTonemap.cpp
+++ b/Engine/Source/Runtime/Renderer/Private/PostProcess/PostProcessTonemap.cpp
@@ -329,7 +329,7 @@ BEGIN_SHADER_PARAMETER_STRUCT(FTonemapParameters, )
 	SHADER_PARAMETER(float, FilmWhiteClip)
 
 	SHADER_PARAMETER_RDG_BUFFER_SRV(StructuredBuffer<float4>, EyeAdaptationBuffer)
-	SHADER_PARAMETER_RDG_TEXTURE(Texture2D, ColorGradingLUT)
+	SHADER_PARAMETER_RDG_TEXTURE(Texture3D, ColorGradingLUT)
 	SHADER_PARAMETER_TEXTURE(Texture2D, BloomDirtMaskTexture)
 	SHADER_PARAMETER_SAMPLER(SamplerState, ColorSampler)
 	SHADER_PARAMETER_SAMPLER(SamplerState, ColorGradingLUTSampler)
@@ -424,9 +424,6 @@ public:
 
 	static void ModifyCompilationEnvironment(const FGlobalShaderPermutationParameters& Parameters, FShaderCompilerEnvironment& OutEnvironment)
 	{
-		const int UseVolumeLut = PipelineVolumeTextureLUTSupportGuaranteedAtRuntime(Parameters.Platform) ? 1 : 0;
-		OutEnvironment.SetDefine(TEXT("USE_VOLUME_LUT"), UseVolumeLut);
-
 		OutEnvironment.SetDefine(TEXT("SUPPORTS_SCENE_COLOR_APPLY_PARAMETERS"), FTonemapInputs::SupportsSceneColorApplyParametersBuffer(Parameters.Platform) ? 1 : 0);
 		OutEnvironment.SetDefine(TEXT("SUPPORTS_FILM_GRAIN"), SupportsFilmGrain(Parameters.Platform) ? 1 : 0);
 	}
@@ -463,9 +460,6 @@ public:
 		OutEnvironment.SetDefine(TEXT("THREADGROUP_SIZEX"), GTonemapComputeTileSizeX);
 		OutEnvironment.SetDefine(TEXT("THREADGROUP_SIZEY"), GTonemapComputeTileSizeY);
 
-		const int UseVolumeLut = PipelineVolumeTextureLUTSupportGuaranteedAtRuntime(Parameters.Platform) ? 1 : 0;
-		OutEnvironment.SetDefine(TEXT("USE_VOLUME_LUT"), UseVolumeLut);
-
 		OutEnvironment.SetDefine(TEXT("SUPPORTS_SCENE_COLOR_APPLY_PARAMETERS"), FTonemapInputs::SupportsSceneColorApplyParametersBuffer(Parameters.Platform) ? 1 : 0);
 		OutEnvironment.SetDefine(TEXT("SUPPORTS_FILM_GRAIN"), SupportsFilmGrain(Parameters.Platform) ? 1 : 0);
 	}
@@ -1078,7 +1072,7 @@ class FMobileCustomResolvePS : public FGlobalShader
 		SHADER_PARAMETER_RDG_TEXTURE(Texture2D, ColorTexture)
 		SHADER_PARAMETER_SAMPLER(SamplerState, ColorSampler)
 		SHADER_PARAMETER(FVector4f, ColorScale0)
-		SHADER_PARAMETER_TEXTURE(Texture2D, ColorGradingLUT)
+		SHADER_PARAMETER_TEXTURE(Texture3D, ColorGradingLUT)
 		SHADER_PARAMETER_SAMPLER(SamplerState, ColorGradingLUTSampler)
 		SHADER_PARAMETER(float, LUTSize)
 		SHADER_PARAMETER(float, InvLUTSize)
@@ -1096,8 +1090,6 @@ class FMobileCustomResolvePS : public FGlobalShader
 
 	static void ModifyCompilationEnvironment(const FGlobalShaderPermutationParameters& Parameters, FShaderCompilerEnvironment& OutEnvironment)
 	{
-		const int UseVolumeLut = PipelineVolumeTextureLUTSupportGuaranteedAtRuntime(Parameters.Platform) ? 1 : 0;
-		OutEnvironment.SetDefine(TEXT("USE_VOLUME_LUT"), UseVolumeLut);
 		OutEnvironment.SetDefine(TEXT("OUTPUT_DEVICE_SRGB"), 1);
 	}
 };
@@ -1124,7 +1116,7 @@ void RenderMobileCustomResolve(FRHICommandList& RHICmdList, const FViewInfo& Vie
 	FMobileCustomResolvePS::FParameters PSShaderParameters;
 	PSShaderParameters.View = View.GetShaderParameters();
 	PSShaderParameters.ColorScale0 = FVector4f(Settings.SceneColorTint.R, Settings.SceneColorTint.G, Settings.SceneColorTint.B, 0);
-	PSShaderParameters.ColorGradingLUT = ColorGradingLUT ? ColorGradingLUT->GetRHI() : GBlackTexture->TextureRHI.GetReference();
+	PSShaderParameters.ColorGradingLUT = ColorGradingLUT ? ColorGradingLUT->GetRHI() : GBlackVolumeTexture->TextureRHI.GetReference();
 	PSShaderParameters.ColorGradingLUTSampler = TStaticSamplerState<SF_Bilinear, AM_Clamp, AM_Clamp, AM_Clamp>::GetRHI();
 	PSShaderParameters.LUTSize = LUTSize;
 	PSShaderParameters.InvLUTSize = 1.0f / LUTSize;
diff --git a/Engine/Source/Runtime/Renderer/Private/ScenePrivate.h b/Engine/Source/Runtime/Renderer/Private/ScenePrivate.h
--- a/Engine/Source/Runtime/Renderer/Private/ScenePrivate.h
+++ b/Engine/Source/Runtime/Renderer/Private/ScenePrivate.h
@@ -13,6 +13,7 @@
 #include "Math/RandomStream.h"
 #include "Templates/PimplPtr.h"
 #include "Templates/RefCounting.h"
+#include "RendererModule.h"
 #include "RHI.h"
 #include "RenderResource.h"
 #include "UniformBuffer.h"
@@ -1346,11 +1347,11 @@ public:
 		bValidTonemappingLUT = bValid;
 	}
 
-	static FPooledRenderTargetDesc CreateLUTRenderTarget(const int32 LUTSize, const bool bUseVolumeLUT, const bool bNeedUAV, const bool bNeedFloatOutput)
+	static FPooledRenderTargetDesc CreateLUTRenderTarget(const int32 LUTSize, const bool bNeedFloatOutput)
 	{
 		// Create the texture needed for the tonemapping LUT in one place
 		EPixelFormat LUTPixelFormat = PF_A2B10G10R10;
-		if (!GPixelFormats[LUTPixelFormat].Supported)
+		if (!UE::PixelFormat::HasCapabilities(LUTPixelFormat, EPixelFormatCapabilities::UAV))
 		{
 			LUTPixelFormat = PF_R8G8B8A8;
 		}
@@ -1358,33 +1359,31 @@ public:
 		{
 			LUTPixelFormat = PF_FloatRGBA;
 		}
-
-		FPooledRenderTargetDesc Desc = FPooledRenderTargetDesc::Create2DDesc(FIntPoint(LUTSize * LUTSize, LUTSize), LUTPixelFormat, FClearValueBinding::Transparent, TexCreate_None, TexCreate_ShaderResource, false);
-		Desc.Flags |= bNeedUAV ? TexCreate_UAV : TexCreate_RenderTargetable;
-
-		if (bUseVolumeLUT)
-		{
-			Desc.Extent = FIntPoint(LUTSize, LUTSize);
-			Desc.Depth = LUTSize;
-		}
-
+		UE_LOG(LogRenderer, Log, TEXT("Pixel format for LUT is %s"), GetPixelFormatString(LUTPixelFormat));
+
+		FPooledRenderTargetDesc Desc = FPooledRenderTargetDesc::CreateVolumeDesc(
+			LUTSize,
+			LUTSize,
+			LUTSize,
+			LUTPixelFormat,
+			FClearValueBinding::Transparent,
+			GFastVRamConfig.CombineLUTs,
+			TexCreate_ShaderResource | TexCreate_UAV,
+			false);
 		Desc.DebugName = TEXT("CombineLUTs");
-		Desc.Flags |= GFastVRamConfig.CombineLUTs;
 
 		return Desc;
 	}
 
 	// Returns a reference to the render target used for the LUT.  Allocated on the first request.
-	IPooledRenderTarget* GetTonemappingLUT(FRHICommandList& RHICmdList, const int32 LUTSize, const bool bUseVolumeLUT, const bool bNeedUAV, const bool bNeedFloatOutput)
+	IPooledRenderTarget* GetTonemappingLUT(FRHICommandList& RHICmdList, const int32 LUTSize, const bool bNeedFloatOutput)
 	{
 		if (CombinedLUTRenderTarget.IsValid() == false || 
 			CombinedLUTRenderTarget->GetDesc().Extent.Y != LUTSize ||
-			((CombinedLUTRenderTarget->GetDesc().Depth != 0) != bUseVolumeLUT) ||
-			!!(CombinedLUTRenderTarget->GetDesc().Flags & TexCreate_UAV) != bNeedUAV ||
 			(CombinedLUTRenderTarget->GetDesc().Format == PF_FloatRGBA) != bNeedFloatOutput)
 		{
 			// Create the texture needed for the tonemapping LUT
-			FPooledRenderTargetDesc Desc = CreateLUTRenderTarget(LUTSize, bUseVolumeLUT, bNeedUAV, bNeedFloatOutput);
+			FPooledRenderTargetDesc Desc = CreateLUTRenderTarget(LUTSize, bNeedFloatOutput);
 
 			GRenderTargetPool.FindFreeElement(RHICmdList, Desc, CombinedLUTRenderTarget, Desc.DebugName);
 		}
diff --git a/Engine/Source/Runtime/Renderer/Private/SceneRendering.cpp b/Engine/Source/Runtime/Renderer/Private/SceneRendering.cpp
--- a/Engine/Source/Runtime/Renderer/Private/SceneRendering.cpp
+++ b/Engine/Source/Runtime/Renderer/Private/SceneRendering.cpp
@@ -2465,12 +2465,12 @@ IPooledRenderTarget* FViewInfo::GetTonemappingLUT() const
 	return nullptr;
 };
 
-IPooledRenderTarget* FViewInfo::GetTonemappingLUT(FRHICommandList& RHICmdList, const int32 LUTSize, const bool bUseVolumeLUT, const bool bNeedUAV, const bool bNeedFloatOutput) const 
+IPooledRenderTarget* FViewInfo::GetTonemappingLUT(FRHICommandList& RHICmdList, const int32 LUTSize, const bool bNeedFloatOutput) const 
 {
 	FSceneViewState* EffectiveViewState = GetEyeAdaptationViewState();
 	if (EffectiveViewState)
 	{
-		return EffectiveViewState->GetTonemappingLUT(RHICmdList, LUTSize, bUseVolumeLUT, bNeedUAV, bNeedFloatOutput);
+		return EffectiveViewState->GetTonemappingLUT(RHICmdList, LUTSize, bNeedFloatOutput);
 	}
 	return nullptr;
 }
diff --git a/Engine/Source/Runtime/Renderer/Private/SceneRendering.h b/Engine/Source/Runtime/Renderer/Private/SceneRendering.h
--- a/Engine/Source/Runtime/Renderer/Private/SceneRendering.h
+++ b/Engine/Source/Runtime/Renderer/Private/SceneRendering.h
@@ -1731,7 +1731,7 @@ public:
 
 	/** Gets the rendertarget that will be populated by CombineLUTS post process 
 	* for stereo rendering, this will force the post-processing to use the same render target for both eyes*/
-	IPooledRenderTarget* GetTonemappingLUT(FRHICommandList& RHICmdList, const int32 LUTSize, const bool bUseVolumeLUT, const bool bNeedUAV, const bool bNeedFloatOutput) const;
+	IPooledRenderTarget* GetTonemappingLUT(FRHICommandList& RHICmdList, const int32 LUTSize, const bool bNeedFloatOutput) const;
 
 	bool IsFirstInFamily() const
 	{
