Deprecate RHICmdList.SubmitCommandsHint() and remove remaining uses of it
- Only the immediate command list can submit commands. Most existing uses were already on the immediate cmdlist, so could be trivially swapped for .ImmediateFlush(...::DispatchToRHIThread), or entirely removed if an existing flush was already there.
- One instance in PathTracing.cpp could be replaced with an RDG dispatch hint.

diff --git a/Engine/Plugins/Media/BinkMedia/Source/BinkMediaPlayer/Private/binktexturesRHI.cpp b/Engine/Plugins/Media/BinkMedia/Source/BinkMediaPlayer/Private/binktexturesRHI.cpp
--- a/Engine/Plugins/Media/BinkMedia/Source/BinkMediaPlayer/Private/binktexturesRHI.cpp
+++ b/Engine/Plugins/Media/BinkMedia/Source/BinkMediaPlayer/Private/binktexturesRHI.cpp
@@ -238,13 +238,6 @@ static void Free_textures(BINKTEXTURES* ptextures)
 	BINKTEXTURESRHI* textures = (BINKTEXTURESRHI*)ptextures;
 	BINKFRAMEBUFFERS *bb = &textures->bink_buffers;
 
-	/*
-	RHICmdList.ImmediateFlush(EImmediateFlushType::FlushRHIThreadFlushResources);
-	RHIFlushResources();
-	RHICmdList.SubmitCommandsHint();
-	FPlatformMisc::MemoryBarrier();
-	*/
-
 	for (int i = 0; i < bb->TotalFrames; ++i)
 	{
 		BINKFRAMEPLANESET *bp_src = &bb->Frames[i];
diff --git a/Engine/Plugins/MovieScene/MovieRenderPipeline/Source/MovieRenderPipelineCore/Private/MoviePipelineRendering.cpp b/Engine/Plugins/MovieScene/MovieRenderPipeline/Source/MovieRenderPipelineCore/Private/MoviePipelineRendering.cpp
--- a/Engine/Plugins/MovieScene/MovieRenderPipeline/Source/MovieRenderPipelineCore/Private/MoviePipelineRendering.cpp
+++ b/Engine/Plugins/MovieScene/MovieRenderPipeline/Source/MovieRenderPipelineCore/Private/MoviePipelineRendering.cpp
@@ -490,7 +490,7 @@ void UMoviePipeline::RenderFrame()
 						(FRHICommandListImmediate& RHICmdList)
 						{
 							RHICmdList.WriteGPUFence(Fence);
-							RHICmdList.SubmitCommandsHint();
+							RHICmdList.ImmediateFlush(EImmediateFlushType::DispatchToRHIThread);
 						});
 
 					// Switch fences for the next frame (this makes us wait on a different frame than what we just made the signal for)
diff --git a/Engine/Plugins/Runtime/nDisplay/Source/DisplayCluster/Private/Game/EngineClasses/Basics/DisplayClusterViewportClient.cpp b/Engine/Plugins/Runtime/nDisplay/Source/DisplayCluster/Private/Game/EngineClasses/Basics/DisplayClusterViewportClient.cpp
--- a/Engine/Plugins/Runtime/nDisplay/Source/DisplayCluster/Private/Game/EngineClasses/Basics/DisplayClusterViewportClient.cpp
+++ b/Engine/Plugins/Runtime/nDisplay/Source/DisplayCluster/Private/Game/EngineClasses/Basics/DisplayClusterViewportClient.cpp
@@ -798,7 +798,7 @@ void UDisplayClusterViewportClient::Draw(FViewport* InViewport, FCanvas* SceneCa
 						[SubmitGPUMask](FRHICommandListImmediate& RHICmdList)
 					{
 						SCOPED_GPU_MASK(RHICmdList, SubmitGPUMask);
-						RHICmdList.SubmitCommandsHint();
+						RHICmdList.ImmediateFlush(EImmediateFlushType::DispatchToRHIThread);
 					});
 				}
 			}
diff --git a/Engine/Plugins/Runtime/nDisplay/Source/DisplayCluster/Private/Render/Viewport/Preview/DisplayClusterViewportManagerPreview_Render.cpp b/Engine/Plugins/Runtime/nDisplay/Source/DisplayCluster/Private/Render/Viewport/Preview/DisplayClusterViewportManagerPreview_Render.cpp
--- a/Engine/Plugins/Runtime/nDisplay/Source/DisplayCluster/Private/Render/Viewport/Preview/DisplayClusterViewportManagerPreview_Render.cpp
+++ b/Engine/Plugins/Runtime/nDisplay/Source/DisplayCluster/Private/Render/Viewport/Preview/DisplayClusterViewportManagerPreview_Render.cpp
@@ -263,7 +263,7 @@ bool FDisplayClusterViewportManagerPreview::RenderViewportPreview(FViewport* InV
 						[SubmitGPUMask](FRHICommandListImmediate& RHICmdList)
 						{
 							SCOPED_GPU_MASK(RHICmdList, SubmitGPUMask);
-							RHICmdList.SubmitCommandsHint();
+							RHICmdList.ImmediateFlush(EImmediateFlushType::DispatchToRHIThread);
 						});
 				}
 			}
diff --git a/Engine/Source/Runtime/RHI/Public/RHICommandList.h b/Engine/Source/Runtime/RHI/Public/RHICommandList.h
--- a/Engine/Source/Runtime/RHI/Public/RHICommandList.h
+++ b/Engine/Source/Runtime/RHI/Public/RHICommandList.h
@@ -3380,7 +3380,7 @@ public:
 	}
 #endif // WITH_RHI_BREADCRUMBS
 
-	//UE_DEPRECATED(5.1, "SubmitCommandsHint is deprecated, and has no effect if called on a non-immediate RHI command list. Consider calling ImmediateFlush(EImmediateFlushType::DispatchToRHIThread) on the immediate command list instead.")
+	UE_DEPRECATED(5.8, "SubmitCommandsHint is deprecated, and has no effect if called on a non-immediate RHI command list. Consider calling ImmediateFlush(EImmediateFlushType::DispatchToRHIThread) on the immediate command list instead.")
 	inline void SubmitCommandsHint();
 
 	inline void CopyToStagingBuffer(FRHIBuffer* SourceBuffer, FRHIStagingBuffer* DestinationStagingBuffer, uint32 Offset, uint32 NumBytes)
@@ -5041,7 +5041,7 @@ public:
 		return nullptr;
 	}
 
-	//UE_DEPRECATED(5.1, "SubmitCommandsHint is deprecated. Consider calling ImmediateFlush(EImmediateFlushType::DispatchToRHIThread) instead.")
+	UE_DEPRECATED(5.8, "SubmitCommandsHint is deprecated. Consider calling ImmediateFlush(EImmediateFlushType::DispatchToRHIThread) instead.")
 	inline void SubmitCommandsHint()
 	{
 		ImmediateFlush(EImmediateFlushType::DispatchToRHIThread);
diff --git a/Engine/Source/Runtime/RHI/Public/RHICommandList.inl b/Engine/Source/Runtime/RHI/Public/RHICommandList.inl
--- a/Engine/Source/Runtime/RHI/Public/RHICommandList.inl
+++ b/Engine/Source/Runtime/RHI/Public/RHICommandList.inl
@@ -64,7 +64,9 @@ inline void FRHIComputeCommandList::SubmitCommandsHint()
 {
 	if (IsImmediate())
 	{
+	PRAGMA_DISABLE_DEPRECATION_WARNINGS
 		static_cast<FRHICommandListImmediate&>(*this).SubmitCommandsHint();
+	PRAGMA_ENABLE_DEPRECATION_WARNINGS
 	}
 }
 
diff --git a/Engine/Source/Runtime/RenderCore/Private/LongGPUTask.cpp b/Engine/Source/Runtime/RenderCore/Private/LongGPUTask.cpp
--- a/Engine/Source/Runtime/RenderCore/Private/LongGPUTask.cpp
+++ b/Engine/Source/Runtime/RenderCore/Private/LongGPUTask.cpp
@@ -93,7 +93,6 @@ void MeasureLongGPUTaskExecutionTime(FRHICommandListImmediate& RHICmdList)
 	RHICmdList.EndRenderQuery(TimeQueryEnd.GetQuery());
 
 	// Required by DX12 to resolve the query
-	RHICmdList.SubmitCommandsHint();
 	RHICmdList.ImmediateFlush(EImmediateFlushType::FlushRHIThread);
 
 	if (RHIGetRenderQueryResult(TimeQueryStart.GetQuery(), StartTime, true) && RHIGetRenderQueryResult(TimeQueryEnd.GetQuery(), EndTime, true))
diff --git a/Engine/Source/Runtime/Renderer/Private/GPUBenchmark.cpp b/Engine/Source/Runtime/Renderer/Private/GPUBenchmark.cpp
--- a/Engine/Source/Runtime/Renderer/Private/GPUBenchmark.cpp
+++ b/Engine/Source/Runtime/Renderer/Private/GPUBenchmark.cpp
@@ -578,7 +578,6 @@ void RendererGPUBenchmark(FRHICommandListImmediate& RHICmdList, FSynthBenchmarkR
 		{
 			uint64 OldAbsTime = 0;
 			// flushes the RHI thread to make sure all RHICmdList.EndRenderQuery() commands got executed.
-			RHICmdList.SubmitCommandsHint();
 			RHICmdList.ImmediateFlush(EImmediateFlushType::FlushRHIThread);
 			RHIGetRenderQueryResult(TimerQueries[0].GetQuery(), OldAbsTime, true);
 
diff --git a/Engine/Source/Runtime/Renderer/Private/MobileShadingRenderer.cpp b/Engine/Source/Runtime/Renderer/Private/MobileShadingRenderer.cpp
--- a/Engine/Source/Runtime/Renderer/Private/MobileShadingRenderer.cpp
+++ b/Engine/Source/Runtime/Renderer/Private/MobileShadingRenderer.cpp
@@ -1947,12 +1947,6 @@ void FMobileSceneRenderer::RenderForwardSinglePass(FRDGBuilder& GraphBuilder, FM
 
 		if (bDoOcclusionInMainPass)
 		{
-			// Issue occlusion queries
-			if (bAdrenoOcclusionMode && IsOpenGLPlatform(ShaderPlatform))
-			{
-				// flush
-				RHICmdList.SubmitCommandsHint();
-			}
 			RenderOcclusion(RHICmdList, View, ViewQueries);
 		}
 
diff --git a/Engine/Source/Runtime/Renderer/Private/PathTracing.cpp b/Engine/Source/Runtime/Renderer/Private/PathTracing.cpp
--- a/Engine/Source/Runtime/Renderer/Private/PathTracing.cpp
+++ b/Engine/Source/Runtime/Renderer/Private/PathTracing.cpp
@@ -4305,7 +4305,7 @@ void FDeferredShadingSceneRenderer::RenderPathTracing(
 								RDG_EVENT_NAME("Path Tracer Sample=%d/%d NumLights=%d (Bounce=%d%s%s)", PathTracingState->SampleIndex, MaxSPP, PassParameters->SceneLightCount, PassParameters->FirstBounce, bUseCompaction ? "" : "+", bUseIndirectDispatch && Bounce > 0 ? TEXT(" indirect") : TEXT("")),
 								PassParameters,
 								ERDGPassFlags::Compute,
-								[PassParameters, RayGenShader, DispatchSizeX, DispatchSizeYLocal, bUseIndirectDispatch, bUse1DDispatch, bFlushRenderingCommands, GPUIndex, &View](FRDGAsyncTask, FRHICommandList& RHICmdList)
+								[PassParameters, RayGenShader, DispatchSizeX, DispatchSizeYLocal, bUseIndirectDispatch, bUse1DDispatch, GPUIndex, &View](FRDGAsyncTask, FRHICommandList& RHICmdList)
 								{
 									FRHIBatchedShaderParameters& GlobalResources = RHICmdList.GetScratchShaderParameters();
 									SetShaderParameters(GlobalResources, RayGenShader, *PassParameters);
@@ -4341,11 +4341,13 @@ void FDeferredShadingSceneRenderer::RenderPathTracing(
 											DispatchSizeX, DispatchSizeYLocal
 										);
 									}
-									if (bFlushRenderingCommands)
-									{
-										RHICmdList.SubmitCommandsHint();
-									}
 								});
+
+							if (bFlushRenderingCommands)
+							{
+								GraphBuilder.AddDispatchHint();
+							}
+
 							if (PreviousPassParameters == nullptr)
 							{
 								PreviousPassParameters = PassParameters;
