PR #13859: Use Blue Noise for SkyAtmosphere dithering

#tests PC editor dx12 + engine test
[RN]

diff --git a/Engine/Shaders/Private/SkyAtmosphere.usf b/Engine/Shaders/Private/SkyAtmosphere.usf
--- a/Engine/Shaders/Private/SkyAtmosphere.usf
+++ b/Engine/Shaders/Private/SkyAtmosphere.usf
@@ -102,6 +102,8 @@
 
 #endif // SAMPLE_OPAQUE_SHADOW
 
+#define BlueNoise SkyAtmosphere.BlueNoise
+#include "BlueNoise.ush"
 
 #if (SAMPLE_CLOUD_SHADOW || SAMPLE_CLOUD_SKYAO)
 #include "VolumetricCloudCommon.ush" 
@@ -406,7 +408,11 @@ float SkyAtmosphereNoise(float2 UV)
 	//	return DEFAULT_SAMPLE_OFFSET;
 	//	return float(Rand3DPCG32(int3(UV.x, UV.y, S)).x) / 4294967296.0f;
 #if VIEWDATA_AVAILABLE && PER_PIXEL_NOISE
+#if 0 // Simple noise
 	return View.RealTimeReflectionCapture ? DEFAULT_SAMPLE_OFFSET : InterleavedGradientNoise(UV.xy, float(View.StateFrameIndexMod8));
+#else // Blue noise works best
+	return View.RealTimeReflectionCapture ? DEFAULT_SAMPLE_OFFSET : BlueNoiseScalar(uint2(UV), View.StateFrameIndex);
+#endif
 #else
 	return DEFAULT_SAMPLE_OFFSET;
 #endif
diff --git a/Engine/Source/Runtime/Renderer/Private/SkyAtmosphereRendering.cpp b/Engine/Source/Runtime/Renderer/Private/SkyAtmosphereRendering.cpp
--- a/Engine/Source/Runtime/Renderer/Private/SkyAtmosphereRendering.cpp
+++ b/Engine/Source/Runtime/Renderer/Private/SkyAtmosphereRendering.cpp
@@ -5,6 +5,7 @@
 =============================================================================*/
 
 #include "SkyAtmosphereRendering.h"
+#include "BlueNoise.h"
 #include "CanvasTypes.h"
 #include "Components/SkyAtmosphereComponent.h"
 #include "DeferredShadingRenderer.h"
@@ -262,6 +263,8 @@ BEGIN_GLOBAL_SHADER_PARAMETER_STRUCT(FSkyAtmosphereInternalCommonParameters, )
 	SHADER_PARAMETER(float, AerialPespectiveViewDistanceScale)
 	SHADER_PARAMETER(float, FogShowFlagFactor)
 
+	SHADER_PARAMETER_STRUCT(FBlueNoise, BlueNoise)
+
 	SHADER_PARAMETER(FVector3f, SkyLuminanceFactor)
 	SHADER_PARAMETER(FVector3f, SkyAndAerialPerspectiveLuminanceFactor)
 END_GLOBAL_SHADER_PARAMETER_STRUCT()
@@ -1350,6 +1353,8 @@ static void SetupSkyAtmosphereInternalCommonParameters(
 	InternalCommonParameters.AerialPespectiveViewDistanceScale = SkyAtmosphereSceneProxy.GetAerialPespectiveViewDistanceScale();
 	InternalCommonParameters.FogShowFlagFactor = ViewFamily.EngineShowFlags.Fog ? 1.0f : 0.0f;
 
+	InternalCommonParameters.BlueNoise = GetBlueNoiseGlobalParameters();
+
 	auto ValidateDistanceValue = [](float& Value)
 	{
 		Value = Value < KINDA_SMALL_NUMBER ? KINDA_SMALL_NUMBER : Value;
