[HWRT] Remove code deprecated in 5.6.

diff --git a/Engine/Plugins/Experimental/GPULightmass/Source/GPULightmass/Private/LightmapRenderer.cpp b/Engine/Plugins/Experimental/GPULightmass/Source/GPULightmass/Private/LightmapRenderer.cpp
--- a/Engine/Plugins/Experimental/GPULightmass/Source/GPULightmass/Private/LightmapRenderer.cpp
+++ b/Engine/Plugins/Experimental/GPULightmass/Source/GPULightmass/Private/LightmapRenderer.cpp
@@ -1100,24 +1100,14 @@ bool FSceneRenderState::SetupRayTracingScene(FRDGBuilder& GraphBuilder, FSceneUn
 								StaticMeshInstanceRenderStates.Elements.Num() + InstanceGroupRenderStates.Elements.Num() + LandscapeIndex
 							);
 
-							PRAGMA_DISABLE_DEPRECATION_WARNINGS
-
-							DynamicGeometryUpdateManager.Update(*ReferenceView.Get());
-
-							const uint32 BLASScratchSize = DynamicGeometryUpdateManager.ComputeScratchBufferSize();
-
-							const FRHIBufferCreateDesc CreateDesc =
-								FRHIBufferCreateDesc::CreateStructured(TEXT("RHILandscapeScratchBuffer"), BLASScratchSize, 0)
-								.AddUsage(EBufferUsageFlags::RayTracingScratch)
-								.SetInitialState(ERHIAccess::UAVCompute);
-
-							FBufferRHIRef ScratchBuffer = RHICmdList.CreateBuffer(CreateDesc);
-
-							RHICmdList.SetStaticUniformBuffers({ReferenceView->ViewUniformBuffer, SceneUniforms.GetBufferRHI(GraphBuilder)});
-							DynamicGeometryUpdateManager.DispatchUpdates(RHICmdList, ScratchBuffer);
-							DynamicGeometryUpdateManager.EndUpdate();
-
-							PRAGMA_ENABLE_DEPRECATION_WARNINGS
+							FRDGBufferRef DynamicGeometryScratchBuffer;
+							DynamicGeometryUpdateManager.AddDynamicGeometryUpdatePass(
+								GraphBuilder,
+								ERDGPassFlags::Compute,
+								SceneUniforms.GetBuffer(GraphBuilder),
+								/*bUseTracingFeedback*/ false,
+								ERHIPipeline::Graphics,
+								DynamicGeometryScratchBuffer);
 
 							// Landscape VF doesn't really use the vertex buffer in HitGroupSystemParameters
 							// We can release after all related RHI cmds get dispatched onto the cmd list
diff --git a/Engine/Source/Editor/UnrealEd/Private/RayTracingDebugVisualizationMenuCommands.cpp b/Engine/Source/Editor/UnrealEd/Private/RayTracingDebugVisualizationMenuCommands.cpp
--- a/Engine/Source/Editor/UnrealEd/Private/RayTracingDebugVisualizationMenuCommands.cpp
+++ b/Engine/Source/Editor/UnrealEd/Private/RayTracingDebugVisualizationMenuCommands.cpp
@@ -90,50 +90,6 @@ void FRayTracingDebugVisualizationMenuCommands::BuildCommandMap()
 	}
 }
 
-// Deprecated. When removed, also remove AddCommandTypeToMenu(FMenuBuilder& Menu, const FVisualizationType Type)
-void FRayTracingDebugVisualizationMenuCommands::BuildVisualisationSubMenu(FMenuBuilder& Menu)
-{
-	const bool bShowTiming = GRayTracingVisualizeTiming != 0;
-	const bool bShowOther = GRayTracingVisualizeOther != 0;
-
-	const FRayTracingDebugVisualizationMenuCommands& Commands = FRayTracingDebugVisualizationMenuCommands::Get();
-	if (Commands.IsPopulated())
-	{
-
-		{
-			Commands.AddCommandTypeToMenu(Menu, FVisualizationType::Overview);
-		}
-
-		{
-			Menu.BeginSection("LevelViewportRayTracingVisualizationGeneral", LOCTEXT("RayTracingVisualizationGeneral", "General"));
-			Commands.AddCommandTypeToMenu(Menu, FVisualizationType::Standard);
-			Menu.EndSection();
-		}
-
-		{
-			Menu.BeginSection("LevelViewportRayTracingVisualizationPerformance", LOCTEXT("RayTracingVisualizationPerformance", "Performance"));
-			Commands.AddCommandTypeToMenu(Menu, FVisualizationType::Performance);
-			Menu.EndSection();
-		}
-
-		if (bShowTiming)
-		{
-			Menu.BeginSection("LevelViewportRayTracingVisualizationTiming", LOCTEXT("RayTracingVisualizationTiming", "Timing"));
-			Commands.AddCommandTypeToMenu(Menu, FVisualizationType::Timing);
-			Menu.EndSection();
-		}
-
-		if (bShowOther)
-		{
-			Menu.BeginSection("LevelViewportRayTracingVisualizationOther", LOCTEXT("RayTracingVisualizationOther", "Other"));
-			Commands.AddCommandTypeToMenu(Menu, FVisualizationType::Other);
-			Menu.EndSection();
-		}
-
-		Menu.EndSection();
-	}
-}
-
 void FRayTracingDebugVisualizationMenuCommands::BuildVisualisationSubMenu(UToolMenu* InMenu)
 {
 	const FRayTracingDebugVisualizationMenuCommands& Commands = FRayTracingDebugVisualizationMenuCommands::Get();
@@ -177,24 +133,6 @@ void FRayTracingDebugVisualizationMenuCommands::BuildVisualisationSubMenu(UToolM
 	}
 }
 
-bool FRayTracingDebugVisualizationMenuCommands::AddCommandTypeToMenu(FMenuBuilder& Menu, const FVisualizationType Type) const
-{
-	bool bAddedCommands = false;
-
-	const TVisualizationModeCommandMap& Commands = CommandMap;
-	for (TCommandConstIterator It = CreateCommandConstIterator(); It; ++It)
-	{
-		const FVisualizationRecord& Record = It.Value();
-		if (Record.Type == Type)
-		{
-			Menu.AddMenuEntry(Record.Command, NAME_None, Record.Command->GetLabel());
-			bAddedCommands = true;
-		}
-	}
-
-	return bAddedCommands;
-}
-
 bool FRayTracingDebugVisualizationMenuCommands::AddCommandTypeToSection(FToolMenuSection& InSection, const FVisualizationType Type) const
 {
 	bool bAddedCommands = false;
diff --git a/Engine/Source/Editor/UnrealEd/Public/RayTracingDebugVisualizationMenuCommands.h b/Engine/Source/Editor/UnrealEd/Public/RayTracingDebugVisualizationMenuCommands.h
--- a/Engine/Source/Editor/UnrealEd/Public/RayTracingDebugVisualizationMenuCommands.h
+++ b/Engine/Source/Editor/UnrealEd/Public/RayTracingDebugVisualizationMenuCommands.h
@@ -46,9 +46,6 @@ public:
 
 	TCommandConstIterator CreateCommandConstIterator() const;
 
-	UE_DEPRECATED(5.6, "Please use the version taking UToolMenu* as argument")
-	static void BuildVisualisationSubMenu(FMenuBuilder& Menu);
-
 	static void BuildVisualisationSubMenu(UToolMenu* InMenu);
 
 	virtual void RegisterCommands() override;
@@ -63,9 +60,6 @@ public:
 private:
 	void BuildCommandMap();
 
-	/** Remove this function when removing deprecated BuildVisualisationSubMenu(FMenuBuilder& Menu) */
-	bool AddCommandTypeToMenu(FMenuBuilder& Menu, const FVisualizationType Type) const;
-
 	bool AddCommandTypeToSection(FToolMenuSection& InSection, const FVisualizationType Type) const;
 
 	static void ChangeRayTracingDebugVisualizationMode(TWeakPtr<FEditorViewportClient> WeakClient, FName InName);
diff --git a/Engine/Source/Runtime/Engine/Public/RayTracingSkinnedGeometry.h b/Engine/Source/Runtime/Engine/Public/RayTracingSkinnedGeometry.h
deleted file mode 100644
diff --git a/Engine/Source/Runtime/Engine/Public/SceneInterface.h b/Engine/Source/Runtime/Engine/Public/SceneInterface.h
--- a/Engine/Source/Runtime/Engine/Public/SceneInterface.h
+++ b/Engine/Source/Runtime/Engine/Public/SceneInterface.h
@@ -745,11 +745,6 @@ public:
 	virtual void IncrementFrameNumber() {}
 
 	virtual void UpdateCachedRayTracingState(class FPrimitiveSceneProxy* SceneProxy) {}
-	
-	UE_DEPRECATED(5.6, "FRayTracingDynamicGeometryCollection has been removed - all raytracing dynamic geometries are collected via the FRayTracingInstanceCollector during GetDynamicRayTracingInstances on the proxy.")
-	virtual class FRayTracingDynamicGeometryCollection* GetRayTracingDynamicGeometryCollection() { return nullptr; }
-	UE_DEPRECATED(5.6, "FRayTracingSkinnedGeometryUpdateQueue has been removed - all raytracing dynamic geometries are collected via the FRayTracingInstanceCollector during GetDynamicRayTracingInstances on the proxy.")
-	virtual class FRayTracingSkinnedGeometryUpdateQueue* GetRayTracingSkinnedGeometryUpdateQueue() { return nullptr; }
 
 	virtual bool RequestGPUSceneUpdate(FPrimitiveSceneInfo& PrimitiveSceneInfo, EPrimitiveDirtyState PrimitiveDirtyState) { return false; }
 	virtual bool RequestUniformBufferUpdate(FPrimitiveSceneInfo& PrimitiveSceneInfo) { return false; }
diff --git a/Engine/Source/Runtime/RenderCore/Public/RayTracingGeometry.h b/Engine/Source/Runtime/RenderCore/Public/RayTracingGeometry.h
--- a/Engine/Source/Runtime/RenderCore/Public/RayTracingGeometry.h
+++ b/Engine/Source/Runtime/RenderCore/Public/RayTracingGeometry.h
@@ -12,7 +12,6 @@ class FRHICommandListBase;
 namespace RayTracing
 {
 	using FGeometryGroupHandle = int32;
-	using GeometryGroupHandle UE_DEPRECATED(5.6, "Use FGeometryGroupHandle instead.") = FGeometryGroupHandle;
 
 	RENDERCORE_API bool ShouldForceRuntimeBLAS();
 }
diff --git a/Engine/Source/Runtime/RenderCore/Public/RayTracingGeometryManagerInterface.h b/Engine/Source/Runtime/RenderCore/Public/RayTracingGeometryManagerInterface.h
--- a/Engine/Source/Runtime/RenderCore/Public/RayTracingGeometryManagerInterface.h
+++ b/Engine/Source/Runtime/RenderCore/Public/RayTracingGeometryManagerInterface.h
@@ -17,7 +17,6 @@ enum class ERTAccelerationStructureBuildPriority : uint8;
 namespace RayTracing
 {
 	using FGeometryGroupHandle = int32;
-	using GeometryGroupHandle UE_DEPRECATED(5.6, "Use FGeometryGroupHandle instead.") = FGeometryGroupHandle;
 }
 
 class IRayTracingGeometryManager
@@ -25,10 +24,7 @@ class IRayTracingGeometryManager
 public:
 
 	using FBuildRequestIndex = int32;
-	using BuildRequestIndex UE_DEPRECATED(5.6, "Use FBuildRequestIndex instead.") = FBuildRequestIndex;
-
 	using FGeometryHandle = int32;
-	using RayTracingGeometryHandle UE_DEPRECATED(5.6, "Use FGeometryHandle instead.") = FGeometryHandle;
 
 	virtual ~IRayTracingGeometryManager() = default;
 
diff --git a/Engine/Source/Runtime/Renderer/Private/RayTracing/RayTracingDynamicGeometryUpdateManager.cpp b/Engine/Source/Runtime/Renderer/Private/RayTracing/RayTracingDynamicGeometryUpdateManager.cpp
--- a/Engine/Source/Runtime/Renderer/Private/RayTracing/RayTracingDynamicGeometryUpdateManager.cpp
+++ b/Engine/Source/Runtime/Renderer/Private/RayTracing/RayTracingDynamicGeometryUpdateManager.cpp
@@ -787,11 +787,6 @@ void FRayTracingDynamicGeometryUpdateManager::ScheduleUpdates(bool bUseTracingFe
 	ScratchBufferSize = BLASScratchSize;
 }
 
-void FRayTracingDynamicGeometryUpdateManager::Update(const FViewInfo& View)
-{
-	ScheduleUpdates(View.bRayTracingFeedbackEnabled);
-}
-
 void FRayTracingDynamicGeometryUpdateManager::AddDynamicGeometryUpdatePass(FRDGBuilder& GraphBuilder, ERDGPassFlags ComputePassFlags, const TRDGUniformBufferRef<FSceneUniformParameters>& SceneUB, bool bUseTracingFeedback, ERHIPipeline ResourceAccessPipelines, FRDGBufferRef& OutDynamicGeometryScratchBuffer)
 {
 	RDG_GPU_MASK_SCOPE(GraphBuilder, FRHIGPUMask::All());
@@ -865,17 +860,10 @@ void FRayTracingDynamicGeometryUpdateManager::AddDynamicGeometryUpdatePass(FRDGB
 	// which would allow the passes above to be tagged with FRDGAsyncTask
 	GraphBuilder.AddPostExecuteCallback([this]
 		{
-			PRAGMA_DISABLE_DEPRECATION_WARNINGS
 			EndUpdate();
-			PRAGMA_ENABLE_DEPRECATION_WARNINGS
 		});
 }
 
-void FRayTracingDynamicGeometryUpdateManager::AddDynamicGeometryUpdatePass(const FViewInfo& View, FRDGBuilder& GraphBuilder, ERDGPassFlags ComputePassFlags, ERHIPipeline ResourceAccessPipelines, FRDGBufferRef& OutDynamicGeometryScratchBuffer)
-{
-	AddDynamicGeometryUpdatePass(GraphBuilder, ComputePassFlags, View.GetSceneUniforms().GetBuffer(GraphBuilder), View.bRayTracingFeedbackEnabled, ResourceAccessPipelines, OutDynamicGeometryScratchBuffer);
-}
-
 void FRayTracingDynamicGeometryUpdateManager::DispatchUpdates(FRHICommandList& RHICmdList, TConstArrayView<FMeshComputeDispatchCommand> DispatchCommands, ERHIPipeline SrcResourceAccessPipelines, ERHIPipeline DstResourceAccessPipelines)
 {
 	if (DispatchCommands.Num() > 0)
@@ -1025,9 +1013,7 @@ void FRayTracingDynamicGeometryUpdateManager::DispatchUpdates(FRHICommandList& R
 		RHICmdList.BuildAccelerationStructures(BuildParams, ScratchBufferRange);
 	}
 
-	PRAGMA_DISABLE_DEPRECATION_WARNINGS
 	EndUpdate();
-	PRAGMA_ENABLE_DEPRECATION_WARNINGS
 }
 
 void FRayTracingDynamicGeometryUpdateManager::EndUpdate()
@@ -1037,11 +1023,6 @@ void FRayTracingDynamicGeometryUpdateManager::EndUpdate()
 	Clear();
 }
 
-uint32 FRayTracingDynamicGeometryUpdateManager::ComputeScratchBufferSize()
-{	
-	return ScratchBufferSize;
-}
-
 #undef USE_RAY_TRACING_DYNAMIC_GEOMETRY_PARALLEL_COMMAND_LISTS
 
 #endif // RHI_RAYTRACING
diff --git a/Engine/Source/Runtime/Renderer/Private/RayTracing/RayTracingInstanceBufferUtil.cpp b/Engine/Source/Runtime/Renderer/Private/RayTracing/RayTracingInstanceBufferUtil.cpp
--- a/Engine/Source/Runtime/Renderer/Private/RayTracing/RayTracingInstanceBufferUtil.cpp
+++ b/Engine/Source/Runtime/Renderer/Private/RayTracing/RayTracingInstanceBufferUtil.cpp
@@ -54,6 +54,21 @@ static TAutoConsoleVariable<bool> CVarRayTracingInstanceBufferParallelFill(
 	TEXT("Whether to fill the ray tracing instance upload buffer using parallel loops. 0=disabled, 1=enabled (default)"),
 	ECVF_RenderThreadSafe);
 
+struct FRayTracingInstanceDescriptor
+{
+	uint32 GPUSceneInstanceOrTransformIndex;
+	uint32 OutputDescriptorIndex;
+	uint32 AccelerationStructureIndex;
+	uint32 InstanceId;
+	uint32 InstanceMaskAndFlags;
+	uint32 InstanceContributionToHitGroupIndex;
+	union
+	{
+		uint32 SceneInstanceIndexAndApplyLocalBoundsTransform;
+		uint32 bApplyLocalBoundsTransform;
+	};
+};
+
 struct FRayTracingInstanceGroup
 {
 	uint32 BaseInstanceIndex : 30;
@@ -128,19 +143,18 @@ struct FGroupHelper
 	}
 };
 
-FRayTracingSceneInitializationData BuildRayTracingSceneInitializationData(TConstArrayView<FRayTracingGeometryInstance> Instances, const TBitArray<>& VisibleInstances)
+void FRayTracingInstanceBufferBuilder::BuildRayTracingSceneInitializationData()
 {
 	const bool bRLEAllowed = CVarRayTracingInstanceBufferRLE.GetValueOnRenderThread();
 
 	const uint32 NumSceneInstances = Instances.Num();
 	
-	FRayTracingSceneInitializationData Output;
-	Output.NumNativeGPUSceneInstances = 0;
-	Output.NumNativeCPUInstances = 0;
-	Output.InstanceGeometryIndices.SetNumUninitialized(NumSceneInstances);
-	Output.BaseUploadBufferOffsets.SetNumUninitialized(NumSceneInstances);
-	Output.BaseInstancePrefixSum.SetNumUninitialized(NumSceneInstances);
-	Output.InstanceGroupEntryRefs.SetNumUninitialized(NumSceneInstances);
+	Data.NumNativeGPUSceneInstances = 0;
+	Data.NumNativeCPUInstances = 0;
+	Data.InstanceGeometryIndices.SetNumUninitialized(NumSceneInstances);
+	Data.BaseUploadBufferOffsets.SetNumUninitialized(NumSceneInstances);
+	Data.BaseInstancePrefixSum.SetNumUninitialized(NumSceneInstances);
+	Data.InstanceGroupEntryRefs.SetNumUninitialized(NumSceneInstances);
 
 	TArray<uint32> InstanceGroups;
 
@@ -168,30 +182,30 @@ FRayTracingSceneInitializationData BuildRayTracingSceneInitializationData(TConst
 
 		checkf(InstanceDesc.GeometryRHI, TEXT("Ray tracing instance must have a valid geometry."));
 
-		uint32 GeometryIndex = UniqueGeometries.FindOrAdd(InstanceDesc.GeometryRHI, Output.ReferencedGeometries.Num());
-		Output.InstanceGeometryIndices[InstanceIndex] = GeometryIndex;
-		if (GeometryIndex == Output.ReferencedGeometries.Num())
+		uint32 GeometryIndex = UniqueGeometries.FindOrAdd(InstanceDesc.GeometryRHI, Data.ReferencedGeometries.Num());
+		Data.InstanceGeometryIndices[InstanceIndex] = GeometryIndex;
+		if (GeometryIndex == Data.ReferencedGeometries.Num())
 		{
-			Output.ReferencedGeometries.Add(InstanceDesc.GeometryRHI);
+			Data.ReferencedGeometries.Add(InstanceDesc.GeometryRHI);
 		}
 
 		if (bGpuSceneInstance)
 		{
 			check(InstanceDesc.Transforms.IsEmpty());
-			Output.BaseUploadBufferOffsets[InstanceIndex] = GPUGroupHelper.NumInstanceDescriptors;
-			Output.NumNativeGPUSceneInstances += InstanceDesc.NumTransforms;
+			Data.BaseUploadBufferOffsets[InstanceIndex] = GPUGroupHelper.NumInstanceDescriptors;
+			Data.NumNativeGPUSceneInstances += InstanceDesc.NumTransforms;
 		}
 		else if (bCpuInstance)
 		{
-			Output.BaseUploadBufferOffsets[InstanceIndex] = CPUGroupHelper.NumInstanceDescriptors;
-			Output.NumNativeCPUInstances += InstanceDesc.NumTransforms;
+			Data.BaseUploadBufferOffsets[InstanceIndex] = CPUGroupHelper.NumInstanceDescriptors;
+			Data.NumNativeCPUInstances += InstanceDesc.NumTransforms;
 		}
 		else
 		{
 			checkNoEntry();
 		}
 		
-		Output.BaseInstancePrefixSum[InstanceIndex] = NumNativeInstances;
+		Data.BaseInstancePrefixSum[InstanceIndex] = NumNativeInstances;
 
 		NumNativeInstances += InstanceDesc.NumTransforms;
 
@@ -199,39 +213,27 @@ FRayTracingSceneInitializationData BuildRayTracingSceneInitializationData(TConst
 
 		if (bGpuSceneInstance)
 		{
-			Output.InstanceGroupEntryRefs[InstanceIndex].GroupIndex = GPUGroupHelper.CurrentGroupIndex;
-			Output.InstanceGroupEntryRefs[InstanceIndex].BaseIndexInGroup = GPUGroupHelper.CurrentIndexInGroup;
+			Data.InstanceGroupEntryRefs[InstanceIndex].GroupIndex = GPUGroupHelper.CurrentGroupIndex;
+			Data.InstanceGroupEntryRefs[InstanceIndex].BaseIndexInGroup = GPUGroupHelper.CurrentIndexInGroup;
 
 			const bool bRLECompatible = bRLEAllowed && (InstanceDesc.BaseInstanceSceneDataOffset != -1) && !bUseUniqueUserData;
 			GPUGroupHelper.AddInstances(InstanceDesc.NumTransforms, bRLECompatible);
 		}
 		else
 		{
-			Output.InstanceGroupEntryRefs[InstanceIndex].GroupIndex = CPUGroupHelper.CurrentGroupIndex;
-			Output.InstanceGroupEntryRefs[InstanceIndex].BaseIndexInGroup = CPUGroupHelper.CurrentIndexInGroup;
+			Data.InstanceGroupEntryRefs[InstanceIndex].GroupIndex = CPUGroupHelper.CurrentGroupIndex;
+			Data.InstanceGroupEntryRefs[InstanceIndex].BaseIndexInGroup = CPUGroupHelper.CurrentIndexInGroup;
 
 			const bool bRLECompatible = bRLEAllowed && !bUseUniqueUserData;
 			CPUGroupHelper.AddInstances(InstanceDesc.NumTransforms, bRLECompatible);
 		}
 	}
 
-	Output.NumGPUInstanceGroups = GPUGroupHelper.CurrentGroupIndex + (GPUGroupHelper.CurrentIndexInGroup > 0 ? 1 : 0);
-	Output.NumCPUInstanceGroups = CPUGroupHelper.CurrentGroupIndex + (CPUGroupHelper.CurrentIndexInGroup > 0 ? 1 : 0);
-
-	Output.NumGPUInstanceDescriptors = GPUGroupHelper.NumInstanceDescriptors;
-	Output.NumCPUInstanceDescriptors = CPUGroupHelper.NumInstanceDescriptors;
-
-	return MoveTemp(Output);
-}
-
+	Data.NumGPUInstanceGroups = GPUGroupHelper.CurrentGroupIndex + (GPUGroupHelper.CurrentIndexInGroup > 0 ? 1 : 0);
+	Data.NumCPUInstanceGroups = CPUGroupHelper.CurrentGroupIndex + (CPUGroupHelper.CurrentIndexInGroup > 0 ? 1 : 0);
 
-FRayTracingSceneInitializationData BuildRayTracingSceneInitializationData(TConstArrayView<FRayTracingGeometryInstance> Instances)
-{
-	// Deprecated code path assumes all provided instances are visible.
-	TBitArray<> VisibleInstances;
-	VisibleInstances.SetNum(Instances.Num(), true);
-
-	return BuildRayTracingSceneInitializationData(Instances, VisibleInstances);
+	Data.NumGPUInstanceDescriptors = GPUGroupHelper.NumInstanceDescriptors;
+	Data.NumCPUInstanceDescriptors = CPUGroupHelper.NumInstanceDescriptors;
 }
 
 void WriteInstanceDescriptor(
@@ -317,7 +319,7 @@ void WriteInstanceDescriptor(
 
 // Helper function to fill upload buffers required by BuildRayTracingInstanceBuffer with instance descriptors
 // Transforms of CPU instances are copied to OutTransformData
-void FillRayTracingInstanceUploadBuffer(
+void FRayTracingInstanceBufferBuilder::FillRayTracingInstanceUploadBuffer(
 	FVector PreViewTranslation,
 	TConstArrayView<FRayTracingGeometryInstance> Instances,
 	const TBitArray<>& VisibleInstances,
@@ -327,7 +329,7 @@ void FillRayTracingInstanceUploadBuffer(
 	TConstArrayView<uint32> InstanceGeometryIndices,
 	TConstArrayView<uint32> BaseUploadBufferOffsets,
 	TConstArrayView<uint32> BaseInstancePrefixSum,
-	TConstArrayView<FRayTracingInstanceGroupEntryRef> InstanceGroupEntryRefs,
+	TConstArrayView<FInstanceGroupEntryRef> InstanceGroupEntryRefs,
 	uint32 NumGPUInstanceGroups,
 	uint32 NumCPUInstanceGroups,
 	uint32 NumGPUInstanceDescriptors,
@@ -401,7 +403,7 @@ void FillRayTracingInstanceUploadBuffer(
 
 			const bool bRLECompatible = bRLEAllowed && (!bGpuSceneInstance || SceneInstance.BaseInstanceSceneDataOffset != -1) && !bUseUniqueUserData;
 
-			const FRayTracingInstanceGroupEntryRef& GroupEntryRef = InstanceGroupEntryRefs[SceneInstanceIndex];
+			const FInstanceGroupEntryRef& GroupEntryRef = InstanceGroupEntryRefs[SceneInstanceIndex];
 
 			uint32 GroupIndex = GroupEntryRef.GroupIndex;
 			uint32 BaseIndexInGroup = GroupEntryRef.BaseIndexInGroup;
@@ -532,42 +534,6 @@ void FillRayTracingInstanceUploadBuffer(
 		}, CVarRayTracingInstanceBufferParallelFill.GetValueOnRenderThread() ? EParallelForFlags::None : EParallelForFlags::ForceSingleThread);
 }
 
-void FillRayTracingInstanceUploadBuffer(
-	FRayTracingSceneRHIRef RayTracingSceneRHI,
-	FVector PreViewTranslation,
-	TConstArrayView<FRayTracingGeometryInstance> Instances,
-	TConstArrayView<uint32> InstanceGeometryIndices,
-	TConstArrayView<uint32> BaseUploadBufferOffsets,
-	TConstArrayView<uint32> BaseInstancePrefixSum,
-	uint32 NumNativeGPUSceneInstances,
-	uint32 NumNativeCPUInstances,
-	TArrayView<FRayTracingInstanceDescriptor> OutInstanceUploadData,
-	TArrayView<FVector4f> OutTransformData)
-{
-	// Deprecated code path assumes all provided instances are visible.
-	TBitArray<> VisibleInstances;
-	VisibleInstances.SetNum(Instances.Num(), true);
-
-	FillRayTracingInstanceUploadBuffer(
-		PreViewTranslation,
-		Instances,
-		VisibleInstances,
-		false,
-		false,
-		false,
-		InstanceGeometryIndices,
-		BaseUploadBufferOffsets,
-		BaseInstancePrefixSum,
-		{},
-		0,
-		0,
-		NumNativeGPUSceneInstances,
-		NumNativeCPUInstances,
-		{},
-		OutInstanceUploadData,
-		OutTransformData);
-}
-
 struct FRayTracingBuildInstanceBufferCS : public FGlobalShader
 {
 	DECLARE_GLOBAL_SHADER(FRayTracingBuildInstanceBufferCS);
@@ -817,43 +783,6 @@ void BuildRayTracingInstanceBuffer(
 	}
 }
 
-void BuildRayTracingInstanceBuffer(
-	FRHICommandList& RHICmdList,
-	const FGPUScene* GPUScene,
-	const FDFVector3& PreViewTranslation,
-	FRHIUnorderedAccessView* InstancesUAV,
-	FRHIUnorderedAccessView* HitGroupContributionsUAV,
-	FRHIShaderResourceView* InstanceUploadSRV,
-	FRHIShaderResourceView* AccelerationStructureAddressesSRV,
-	FRHIShaderResourceView* CPUInstanceTransformSRV,
-	uint32 NumNativeGPUSceneInstances,
-	uint32 NumNativeCPUInstances,
-	const FRayTracingCullingParameters* CullingParameters,
-	FRHIUnorderedAccessView* OutputStatsUAV,
-	FRHIUnorderedAccessView* InstanceExtraDataUAV)
-{
-	BuildRayTracingInstanceBuffer(
-		RHICmdList,
-		GPUScene,
-		PreViewTranslation,
-		InstancesUAV,
-		HitGroupContributionsUAV,
-		nullptr,
-		InstanceUploadSRV,
-		AccelerationStructureAddressesSRV,
-		CPUInstanceTransformSRV,
-		NumNativeGPUSceneInstances + NumNativeCPUInstances,
-		0,
-		0,
-		NumNativeGPUSceneInstances,
-		NumNativeCPUInstances,
-		CullingParameters,
-		/*bCompactOutput*/ false,
-		OutputStatsUAV,
-		0,
-		InstanceExtraDataUAV);
-}
-
 void FRayTracingInstanceBufferBuilder::Init(TConstArrayView<FRayTracingGeometryInstance> InInstances, FVector InPreViewTranslation)
 {
 	checkf(Instances.IsEmpty(), TEXT("RayTracingInstanceBufferBuilder was already initialized."));
@@ -864,7 +793,7 @@ void FRayTracingInstanceBufferBuilder::Init(TConstArrayView<FRayTracingGeometryI
 	// Deprecated code path assumes all provided instances are visible.
 	VisibleInstances.SetNum(Instances.Num(), true);
 
-	Data = BuildRayTracingSceneInitializationData(Instances, VisibleInstances);
+	BuildRayTracingSceneInitializationData();
 }
 
 void FRayTracingInstanceBufferBuilder::Init(FRayTracingInstanceBufferBuilderInitializer Initializer)
@@ -877,7 +806,7 @@ void FRayTracingInstanceBufferBuilder::Init(FRayTracingInstanceBufferBuilderInit
 	bForceOpaque = Initializer.bForceOpaque;
 	bDisableTriangleCull = Initializer.bDisableTriangleCull;
 
-	Data = BuildRayTracingSceneInitializationData(Instances, VisibleInstances);
+	BuildRayTracingSceneInitializationData();
 }
 
 void FRayTracingInstanceBufferBuilder::FillRayTracingInstanceUploadBuffer(FRHICommandList& RHICmdList)
@@ -956,7 +885,7 @@ void FRayTracingInstanceBufferBuilder::FillRayTracingInstanceUploadBuffer(FRHICo
 	FRayTracingInstanceDescriptor* InstanceUploadData = (FRayTracingInstanceDescriptor*)RHICmdList.LockBuffer(InstanceUploadBuffer, 0, InstanceUploadBytes, RLM_WriteOnly);
 	FVector4f* TransformUploadData = (Data.NumNativeCPUInstances > 0) ? (FVector4f*)RHICmdList.LockBuffer(TransformUploadBuffer, 0, TransformUploadBytes, RLM_WriteOnly) : nullptr;
 
-	::FillRayTracingInstanceUploadBuffer(
+	FillRayTracingInstanceUploadBuffer(
 		PreViewTranslation,
 		Instances,
 		VisibleInstances,
diff --git a/Engine/Source/Runtime/Renderer/Public/RayTracingDynamicGeometryUpdateManager.h b/Engine/Source/Runtime/Renderer/Public/RayTracingDynamicGeometryUpdateManager.h
--- a/Engine/Source/Runtime/Renderer/Public/RayTracingDynamicGeometryUpdateManager.h
+++ b/Engine/Source/Runtime/Renderer/Public/RayTracingDynamicGeometryUpdateManager.h
@@ -67,26 +67,15 @@ public:
 		ERHIPipeline ResourceAccessPipelines,
 		FRDGBufferRef& OutDynamicGeometryScratchBuffer);
 
-	UE_DEPRECATED(5.6, "Provide SceneUB and bUseTracingFeedback instead of View.")
-	RENDERER_API void AddDynamicGeometryUpdatePass(const FViewInfo& View, FRDGBuilder& GraphBuilder, ERDGPassFlags ComputePassFlags, ERHIPipeline ResourceAccessPipelines, FRDGBufferRef& OutDynamicGeometryScratchBuffer);
-
 	/** Clears the working arrays to not hold any references. */
 	RENDERER_API void Clear();
 
-	UE_DEPRECATED(5.5, "Use AddDynamicGeometryUpdatePass instead.")
-	RENDERER_API void DispatchUpdates(FRHICommandList& RHICmdList, FRHIBuffer* ScratchBuffer, ERHIPipeline SrcResourceAccessPipelines = ERHIPipeline::Graphics, ERHIPipeline DstResourceAccessPipelines = ERHIPipeline::Graphics);
-
-	UE_DEPRECATED(5.5, "Use AddDynamicGeometryUpdatePass instead.")
-	RENDERER_API void EndUpdate();
-
-	UE_DEPRECATED(5.5, "Use AddDynamicGeometryUpdatePass instead which allocates scratch buffer internally.")
-	RENDERER_API uint32 ComputeScratchBufferSize();
-
-	UE_DEPRECATED(5.6, "Use AddDynamicGeometryUpdatePass(...) which internally handles the full update.")
-	RENDERER_API void Update(const FViewInfo& View);
 private:
 
+	void EndUpdate();
 	void ScheduleUpdates(bool bUseTracingFeedback);
+	void DispatchUpdates(FRHICommandList& RHICmdList, FRHIBuffer* ScratchBuffer, ERHIPipeline SrcResourceAccessPipelines = ERHIPipeline::Graphics, ERHIPipeline DstResourceAccessPipelines = ERHIPipeline::Graphics);
+
 	static void DispatchUpdates(FRHICommandList& RHICmdList, TConstArrayView<FMeshComputeDispatchCommand> DispatchCommands, ERHIPipeline SrcResourceAccessPipelines = ERHIPipeline::Graphics, ERHIPipeline DstResourceAccessPipelines = ERHIPipeline::Graphics);
 	
 	struct FRayTracingDynamicGeometryBuildParams
@@ -138,9 +127,4 @@ private:
 	uint32 ScratchBufferSize = 0;
 };
 
-UE_DEPRECATED(5.6, "Use FRayTracingDynamicGeometryUpdateManager instead of FRayTracingDynamicGeometryCollection");
-class FRayTracingDynamicGeometryCollection : public FRayTracingDynamicGeometryUpdateManager
-{
-};
-
 #endif // RHI_RAYTRACING
\ No newline at end of file
diff --git a/Engine/Source/Runtime/Renderer/Public/RayTracingInstance.h b/Engine/Source/Runtime/Renderer/Public/RayTracingInstance.h
--- a/Engine/Source/Runtime/Renderer/Public/RayTracingInstance.h
+++ b/Engine/Source/Runtime/Renderer/Public/RayTracingInstance.h
@@ -32,12 +32,12 @@ struct FRayTracingInstance
 		return Materials.Num() != 0;
 	}
 
-	TArrayView<const FMeshBatch> GetMaterials() const
+	TConstArrayView<FMeshBatch> GetMaterials() const
 	{
 		if (OwnsMaterials())
 		{
 			check(MaterialsView.Num() == 0);
-			return TArrayView<const FMeshBatch>(Materials);
+			return TConstArrayView<FMeshBatch>(Materials);
 		}
 		else
 		{
@@ -52,9 +52,6 @@ struct FRayTracingInstance
 	/** Whether the instance is thin geometry (e.g., Hair strands)*/
 	bool bThinGeometry = false;
 
-	UE_DEPRECATED(5.6, "Near/far field assignment is done based on ERayTracingPrimitiveFlags::FarField.")
-	ERayTracingInstanceLayer InstanceLayer = ERayTracingInstanceLayer::NearField;
-
 	/** Mark InstanceMaskAndFlags dirty to be automatically updated in the renderer module (dirty by default).
 	* If caching is used, clean the dirty state by setting it to false so no duplicate update will be performed in the renderer module.
 	*/
diff --git a/Engine/Source/Runtime/Renderer/Public/RayTracingInstanceBufferUtil.h b/Engine/Source/Runtime/Renderer/Public/RayTracingInstanceBufferUtil.h
--- a/Engine/Source/Runtime/Renderer/Public/RayTracingInstanceBufferUtil.h
+++ b/Engine/Source/Runtime/Renderer/Public/RayTracingInstanceBufferUtil.h
@@ -12,94 +12,8 @@ class FGPUScene;
 struct FRayTracingCullingParameters;
 struct FDFVector3;
 
-struct UE_DEPRECATED(5.6, "Use FRayTracingInstanceBufferBuilder instead.") FRayTracingInstanceDescriptor
-{
-	uint32 GPUSceneInstanceOrTransformIndex;
-	uint32 OutputDescriptorIndex;
-	uint32 AccelerationStructureIndex;
-	uint32 InstanceId;
-	uint32 InstanceMaskAndFlags;
-	uint32 InstanceContributionToHitGroupIndex;
-	union
-	{
-		uint32 SceneInstanceIndexAndApplyLocalBoundsTransform;
-		uint32 bApplyLocalBoundsTransform;
-	};
-};
-
-struct UE_DEPRECATED(5.6, "Use FRayTracingInstanceBufferBuilder instead.") FRayTracingInstanceGroupEntryRef
-{
-	uint32 GroupIndex;
-	uint32 BaseIndexInGroup;
-};
-
-PRAGMA_DISABLE_DEPRECATION_WARNINGS
-struct UE_DEPRECATED(5.6, "Use FRayTracingInstanceBufferBuilder instead.") FRayTracingSceneInitializationData
-{
-	uint32 NumGPUInstanceGroups;
-	uint32 NumCPUInstanceGroups;
-	uint32 NumGPUInstanceDescriptors;
-	uint32 NumCPUInstanceDescriptors;
-	uint32 NumNativeGPUSceneInstances;
-	uint32 NumNativeCPUInstances;
-
-	UE_DEPRECATED(5.6, "No longer used. Use FShaderBindingTable instead.")
-	uint32 TotalNumSegments;
-
-	// index of each instance geometry in ReferencedGeometries
-	TArray<uint32> InstanceGeometryIndices;
-	// base offset of each instance entries in the instance upload buffer
-	TArray<uint32> BaseUploadBufferOffsets;
-	// prefix sum of `Instance.NumTransforms` for all instances in this scene
-	TArray<uint32> BaseInstancePrefixSum;
-	// reference to corresponding instance group entry
-	TArray<FRayTracingInstanceGroupEntryRef> InstanceGroupEntryRefs;
-
-	// Unique list of geometries referenced by all instances in this scene.
-	// Any referenced geometry is kept alive while the scene is alive.
-	TArray<FRHIRayTracingGeometry*> ReferencedGeometries;
-
-	FRayTracingSceneInitializationData() = default;
-	FRayTracingSceneInitializationData(const FRayTracingSceneInitializationData&) = default;
-	FRayTracingSceneInitializationData& operator=(const FRayTracingSceneInitializationData&) = default;
-	FRayTracingSceneInitializationData(FRayTracingSceneInitializationData&&) = default;
-	FRayTracingSceneInitializationData& operator=(FRayTracingSceneInitializationData&&) = default;
-	~FRayTracingSceneInitializationData() = default;
-};
-
-UE_DEPRECATED(5.6, "Use FRayTracingInstanceBufferBuilder instead.")
-RENDERER_API FRayTracingSceneInitializationData BuildRayTracingSceneInitializationData(TConstArrayView<FRayTracingGeometryInstance> Instances);
-
-UE_DEPRECATED(5.6, "Use FRayTracingInstanceBufferBuilder instead.")
-RENDERER_API void FillRayTracingInstanceUploadBuffer(
-	FRayTracingSceneRHIRef RayTracingSceneRHI,
-	FVector PreViewTranslation,
-	TConstArrayView<FRayTracingGeometryInstance> Instances,
-	TConstArrayView<uint32> InstanceGeometryIndices,
-	TConstArrayView<uint32> BaseUploadBufferOffsets,
-	TConstArrayView<uint32> BaseInstancePrefixSum,
-	uint32 NumNativeGPUSceneInstances,
-	uint32 NumNativeCPUInstances,
-	TArrayView<FRayTracingInstanceDescriptor> OutInstanceUploadData,
-	TArrayView<FVector4f> OutTransformData);
-
-UE_DEPRECATED(5.6, "Use FRayTracingInstanceBufferBuilder instead.")
-RENDERER_API void BuildRayTracingInstanceBuffer(
-	FRHICommandList& RHICmdList,
-	const FGPUScene* GPUScene,
-	const FDFVector3& PreViewTranslation,
-	FRHIUnorderedAccessView* InstancesUAV,
-	FRHIUnorderedAccessView* HitGroupContributionsUAV,
-	FRHIShaderResourceView* InstanceUploadSRV,
-	FRHIShaderResourceView* AccelerationStructureAddressesSRV,
-	FRHIShaderResourceView* CPUInstanceTransformSRV,
-	uint32 NumNativeGPUSceneInstances,
-	uint32 NumNativeCPUInstances,
-	const FRayTracingCullingParameters* CullingParameters,
-	FRHIUnorderedAccessView* OutputStatsUAV,
-	FRHIUnorderedAccessView* InstanceExtraDataUAV);
-
-PRAGMA_ENABLE_DEPRECATION_WARNINGS
+struct FRayTracingInstanceGroup;
+struct FRayTracingInstanceDescriptor;
 
 struct FRayTracingInstanceBufferBuilderInitializer
 {
@@ -133,18 +47,64 @@ public:
 		uint32 OutputStatsOffset,
 		FRHIUnorderedAccessView* InstanceExtraDataUAV);
 
-	PRAGMA_DISABLE_DEPRECATION_WARNINGS
-
 	uint32 GetMaxNumInstances() const { return Data.NumNativeGPUSceneInstances + Data.NumNativeCPUInstances; }
 
 	TConstArrayView<FRHIRayTracingGeometry*> GetReferencedGeometries() const { return Data.ReferencedGeometries; }
 	TConstArrayView<uint32> GetInstanceGeometryIndices() const { return Data.InstanceGeometryIndices; }
 	TConstArrayView<uint32> GetBaseInstancePrefixSum() const { return Data.BaseInstancePrefixSum; }
 
-	PRAGMA_ENABLE_DEPRECATION_WARNINGS
-
 private:
 
+	struct FInstanceGroupEntryRef
+	{
+		uint32 GroupIndex;
+		uint32 BaseIndexInGroup;
+	};
+
+	struct FInitializationData
+	{
+		uint32 NumGPUInstanceGroups;
+		uint32 NumCPUInstanceGroups;
+		uint32 NumGPUInstanceDescriptors;
+		uint32 NumCPUInstanceDescriptors;
+		uint32 NumNativeGPUSceneInstances;
+		uint32 NumNativeCPUInstances;
+
+		// index of each instance geometry in ReferencedGeometries
+		TArray<uint32> InstanceGeometryIndices;
+		// base offset of each instance entries in the instance upload buffer
+		TArray<uint32> BaseUploadBufferOffsets;
+		// prefix sum of `Instance.NumTransforms` for all instances in this scene
+		TArray<uint32> BaseInstancePrefixSum;
+		// reference to corresponding instance group entry
+		TArray<FInstanceGroupEntryRef> InstanceGroupEntryRefs;
+
+		// Unique list of geometries referenced by all instances in this scene.
+		// Any referenced geometry is kept alive while the scene is alive.
+		TArray<FRHIRayTracingGeometry*> ReferencedGeometries;
+	};
+
+	void BuildRayTracingSceneInitializationData();
+
+	static void FillRayTracingInstanceUploadBuffer(
+		FVector PreViewTranslation,
+		TConstArrayView<FRayTracingGeometryInstance> Instances,
+		const TBitArray<>& VisibleInstances,
+		bool bUseLightingChannels,
+		bool bForceOpaque,
+		bool bDisableTriangleCull,
+		TConstArrayView<uint32> InstanceGeometryIndices,
+		TConstArrayView<uint32> BaseUploadBufferOffsets,
+		TConstArrayView<uint32> BaseInstancePrefixSum,
+		TConstArrayView<FInstanceGroupEntryRef> InstanceGroupEntryRefs,
+		uint32 NumGPUInstanceGroups,
+		uint32 NumCPUInstanceGroups,
+		uint32 NumGPUInstanceDescriptors,
+		uint32 NumCPUInstanceDescriptors,
+		TArrayView<FRayTracingInstanceGroup> OutInstanceGroupUploadData,
+		TArrayView<FRayTracingInstanceDescriptor> OutInstanceUploadData,
+		TArrayView<FVector4f> OutTransformData);
+
 	TConstArrayView<FRayTracingGeometryInstance> Instances;
 	TBitArray<> VisibleInstances;
 	FVector PreViewTranslation;
@@ -152,9 +112,7 @@ private:
 	bool bForceOpaque = false;
 	bool bDisableTriangleCull = false;
 
-	PRAGMA_DISABLE_DEPRECATION_WARNINGS
-	FRayTracingSceneInitializationData Data;
-	PRAGMA_ENABLE_DEPRECATION_WARNINGS
+	FInitializationData Data;
 
 	FBufferRHIRef InstanceGroupUploadBuffer;
 	FShaderResourceViewRHIRef InstanceGroupUploadSRV;
diff --git a/Engine/Source/Runtime/Renderer/Public/RayTracingShaderBindingLayout.h b/Engine/Source/Runtime/Renderer/Public/RayTracingShaderBindingLayout.h
--- a/Engine/Source/Runtime/Renderer/Public/RayTracingShaderBindingLayout.h
+++ b/Engine/Source/Runtime/Renderer/Public/RayTracingShaderBindingLayout.h
@@ -19,12 +19,6 @@ namespace RayTracing
 
 	// Setup the runtime static uniform buffer bindings on the command list if enabled
 	RENDERER_API TOptional<FScopedUniformBufferStaticBindings> BindStaticUniformBufferBindings(const FViewInfo& View, FRHIUniformBuffer* SceneUniformBuffer, FRHIUniformBuffer* NaniteRayTracingUniformBuffer, FRHICommandList& RHICmdList);
-
-	UE_DEPRECATED(5.6, "Must provide NaniteRayTracingUniformBuffer (accessible using Nanite::GetPublicGlobalRayTracingUniformBuffer()).")
-	inline TOptional<FScopedUniformBufferStaticBindings> BindStaticUniformBufferBindings(const FViewInfo& View, FRHIUniformBuffer* SceneUniformBuffer, FRHICommandList& RHICmdList)
-	{
-		return BindStaticUniformBufferBindings(View, SceneUniformBuffer, nullptr, RHICmdList);
-	}
 }
 
 #endif // RHI_RAYTRACING
\ No newline at end of file
