Substrate - simple tile of 16x16 pixel working with substrate only (not lumen nor mega light) for some platforms.
disabled validation due to weird error I could not reproduce locally.

#tests PC editor and console
#rnx

diff --git a/Engine/Shaders/Private/Substrate/SubstrateMaterialClassification.usf b/Engine/Shaders/Private/Substrate/SubstrateMaterialClassification.usf
--- a/Engine/Shaders/Private/Substrate/SubstrateMaterialClassification.usf
+++ b/Engine/Shaders/Private/Substrate/SubstrateMaterialClassification.usf
@@ -8,6 +8,12 @@
 #define SUBSTRATE_COMPLEXSPECIALPATH (SUBSTRATE_GBUFFER_FORMAT==1)
 
 #include "/Engine/Private/Substrate/Substrate.ush"
+#if SUBSTRATE_CLASSIFICATION_TILE_SIZE_16
+#undef SUBSTRATE_TILE_SIZE
+#undef SUBSTRATE_TILE_SIZE_DIV_AS_SHIFT
+#define SUBSTRATE_TILE_SIZE							16
+#define SUBSTRATE_TILE_SIZE_DIV_AS_SHIFT			4
+#endif
 #include "SubstrateTile.ush"
 
 #if SUBSTRATE_GBUFFER_FORMAT==0
diff --git a/Engine/Shaders/Private/Substrate/SubstrateTile.usf b/Engine/Shaders/Private/Substrate/SubstrateTile.usf
--- a/Engine/Shaders/Private/Substrate/SubstrateTile.usf
+++ b/Engine/Shaders/Private/Substrate/SubstrateTile.usf
@@ -2,6 +2,12 @@
 
 #include "/Engine/Private/Common.ush"
 #include "/Engine/Shared/SubstrateDefinitions.h"   
+#if SUBSTRATE_CLASSIFICATION_TILE_SIZE_16
+#undef SUBSTRATE_TILE_SIZE
+#undef SUBSTRATE_TILE_SIZE_DIV_AS_SHIFT
+#define SUBSTRATE_TILE_SIZE							16
+#define SUBSTRATE_TILE_SIZE_DIV_AS_SHIFT			4
+#endif
 #include "SubstrateTile.ush"
 
 
diff --git a/Engine/Source/Runtime/RenderCore/Private/RenderUtils.cpp b/Engine/Source/Runtime/RenderCore/Private/RenderUtils.cpp
--- a/Engine/Source/Runtime/RenderCore/Private/RenderUtils.cpp
+++ b/Engine/Source/Runtime/RenderCore/Private/RenderUtils.cpp
@@ -2089,10 +2089,10 @@ static TAutoConsoleVariable<int32> CVarSubstrateStochasticLighting(
 	TEXT("Enable stochastic lighting to get better performance."),
 	ECVF_ReadOnly | ECVF_RenderThreadSafe);
 
-static TAutoConsoleVariable<int32> CVarSubstrateTileLightingPasses(
-	TEXT("r.Substrate.TileLightingPasses"),
-	1,
-	TEXT("Enable tiling of full screen directional light or environment lighting passes. Disabling this can help performance on low end platform not supporting Rect primitives."),
+static TAutoConsoleVariable<int32> CVarSubstrateSimpleTiledLightingPasses(
+	TEXT("r.Substrate.SimpleTiledLightingPasses"),
+	0,
+	TEXT("Enabling simple tiled lighting passes when greater than 0, can help performance on low end platforms not supporting Rect primitives. Only stencil marking will be tiled. And tiled passes will use larger tiles. Not compatible with Lumen nor MegaLight."),
 	ECVF_RenderThreadSafe);
 
 namespace Substrate
@@ -2357,7 +2357,12 @@ namespace Substrate
 
 	bool AreLightingPassesTiled()
 	{
-		return CVarSubstrateTileLightingPasses.GetValueOnAnyThread() > 0;
+		return CVarSubstrateSimpleTiledLightingPasses.GetValueOnAnyThread() <= 0;
+	}
+
+	bool UseTileSize16()
+	{
+		return CVarSubstrateSimpleTiledLightingPasses.GetValueOnAnyThread() > 0;
 	}
 }
 
diff --git a/Engine/Source/Runtime/RenderCore/Public/RenderUtils.h b/Engine/Source/Runtime/RenderCore/Public/RenderUtils.h
--- a/Engine/Source/Runtime/RenderCore/Public/RenderUtils.h
+++ b/Engine/Source/Runtime/RenderCore/Public/RenderUtils.h
@@ -644,6 +644,7 @@ namespace Substrate
 	RENDERCORE_API bool IsHiddenMaterialAssetConversionEnabled();
 
 	RENDERCORE_API bool AreLightingPassesTiled();
+	RENDERCORE_API bool UseTileSize16();
 
 	RENDERCORE_API uint32 GetRayTracingMaterialPayloadSizeInBytes(bool bFullySimplifiedMaterial);
 
diff --git a/Engine/Source/Runtime/Renderer/Private/IndirectLightRendering.cpp b/Engine/Source/Runtime/Renderer/Private/IndirectLightRendering.cpp
--- a/Engine/Source/Runtime/Renderer/Private/IndirectLightRendering.cpp
+++ b/Engine/Source/Runtime/Renderer/Private/IndirectLightRendering.cpp
@@ -1597,6 +1597,7 @@ void FDeferredShadingSceneRenderer::RenderDiffuseIndirectAndAmbientOcclusion(
 				Substrate::FSubstrateTilePassVS::FPermutationDomain VSPermutationVector;
 				VSPermutationVector.Set< Substrate::FSubstrateTilePassVS::FEnableDebug >(false);
 				VSPermutationVector.Set< Substrate::FSubstrateTilePassVS::FEnableTexCoordScreenVector >(false);
+				VSPermutationVector.Set< Substrate::FSubstrateTilePassVS::FClassificationTileSize16 >(Substrate::UseTileSize16());
 				TShaderMapRef<Substrate::FSubstrateTilePassVS> TileVertexShader(View.ShaderMap, VSPermutationVector);
 
 				ClearUnusedGraphResources(PixelShader, PSParameters);
@@ -1689,14 +1690,16 @@ void FDeferredShadingSceneRenderer::RenderDiffuseIndirectAndAmbientOcclusion(
 					PSParameters->RenderTargets[2] = FRenderTargetBinding(Scene->SubstrateSceneData.SeparatedSubSurfaceSceneColor, ERenderTargetLoadAction::ELoad);
 				}
 
+				const bool bSubstrateEnabled = Substrate::IsSubstrateEnabled();
+
 				TShaderMapRef<FPostProcessVS> FullScreenVertexShader(View.ShaderMap);
 
 				Substrate::FSubstrateTilePassVS::FPermutationDomain VSPermutationVector;
 				VSPermutationVector.Set< Substrate::FSubstrateTilePassVS::FEnableDebug >(false);
 				VSPermutationVector.Set< Substrate::FSubstrateTilePassVS::FEnableTexCoordScreenVector >(false);
+				VSPermutationVector.Set< Substrate::FSubstrateTilePassVS::FClassificationTileSize16 >(bSubstrateEnabled && Substrate::UseTileSize16());
 				TShaderMapRef<Substrate::FSubstrateTilePassVS> TileVertexShader(View.ShaderMap, VSPermutationVector);
 
-				const bool bSubstrateEnabled = Substrate::IsSubstrateEnabled();
 				EPrimitiveType PrimitiveType = PT_TriangleList;
 				if (Substrate::IsSubstrateEnabled())
 				{
@@ -1916,6 +1919,9 @@ static void AddSkyReflectionPass(
 		bLumenStandaloneReflections, bIsTiledSSR, bHasHairCardsOrMeshes,
 		SubstrateTileMaterialType);
 
+	const bool bSubstrateEnabled = Substrate::IsSubstrateEnabled();
+	const bool bSubstrateLightingPassesTiled = Substrate::AreLightingPassesTiled();
+
 	TShaderMapRef<FReflectionEnvironmentSkyLightingPS> PixelShader(View.ShaderMap, PermutationVector);
 	ClearUnusedGraphResources(PixelShader, &PassParameters->PS);
 
@@ -1923,11 +1929,10 @@ static void AddSkyReflectionPass(
 	Substrate::FSubstrateTilePassVS::FPermutationDomain VSPermutationVector;
 	VSPermutationVector.Set< Substrate::FSubstrateTilePassVS::FEnableDebug >(false);
 	VSPermutationVector.Set< Substrate::FSubstrateTilePassVS::FEnableTexCoordScreenVector >(false);
+	VSPermutationVector.Set< Substrate::FSubstrateTilePassVS::FClassificationTileSize16 >(bSubstrateEnabled && bSubstrateLightingPassesTiled && Substrate::UseTileSize16());
 	TShaderMapRef<Substrate::FSubstrateTilePassVS> SubstrateTilePassVertexShader(View.ShaderMap, VSPermutationVector);
 	PassParameters->VS.TileIndirectBuffer = nullptr;
 	PassParameters->VS.TileListBuffer = nullptr;
-	const bool bSubstrateEnabled = Substrate::IsSubstrateEnabled();
-	const bool bSubstrateLightingPassesTiled = Substrate::AreLightingPassesTiled();
 	if (bSubstrateEnabled && bSubstrateLightingPassesTiled)
 	{
 		check(SubstrateTileMaterialType <= ESubstrateTileType::EComplexSpecial);
diff --git a/Engine/Source/Runtime/Renderer/Private/LightRendering.cpp b/Engine/Source/Runtime/Renderer/Private/LightRendering.cpp
--- a/Engine/Source/Runtime/Renderer/Private/LightRendering.cpp
+++ b/Engine/Source/Runtime/Renderer/Private/LightRendering.cpp
@@ -2702,6 +2702,7 @@ static void InternalRenderLight(
 			Substrate::FSubstrateTilePassVS::FPermutationDomain VSPermutationVector;
 			VSPermutationVector.Set< Substrate::FSubstrateTilePassVS::FEnableDebug >(false);
 			VSPermutationVector.Set< Substrate::FSubstrateTilePassVS::FEnableTexCoordScreenVector >(true);
+			VSPermutationVector.Set< Substrate::FSubstrateTilePassVS::FClassificationTileSize16 >(Substrate::UseTileSize16());
 			TShaderMapRef<Substrate::FSubstrateTilePassVS> TileVertexShader(View.ShaderMap, VSPermutationVector);
 
 			GraphicsPSOInit.BoundShaderState.VertexShaderRHI = bEnableSubstrateTiledPass ? TileVertexShader.GetVertexShader() : VertexShader.GetVertexShader();
@@ -3695,6 +3696,7 @@ void DeferredLightGlobalPSOCollector(const FSceneTexturesConfig& SceneTexturesCo
 				Substrate::FSubstrateTilePassVS::FPermutationDomain VSPermutationVector;
 				VSPermutationVector.Set< Substrate::FSubstrateTilePassVS::FEnableDebug >(false);
 				VSPermutationVector.Set< Substrate::FSubstrateTilePassVS::FEnableTexCoordScreenVector >(true);
+				VSPermutationVector.Set< Substrate::FSubstrateTilePassVS::FClassificationTileSize16 >(Substrate::UseTileSize16());
 				TShaderMapRef<Substrate::FSubstrateTilePassVS> TileVertexShader(GlobalShaderMap, VSPermutationVector);
 
 				const bool bEnableSubstrateTiledPass = SubstrateTileMaterialType != ESubstrateTileType::ECount && Substrate::IsSubstrateEnabled() && Substrate::AreLightingPassesTiled();
diff --git a/Engine/Source/Runtime/Renderer/Private/Substrate/Substrate.cpp b/Engine/Source/Runtime/Renderer/Private/Substrate/Substrate.cpp
--- a/Engine/Source/Runtime/Renderer/Private/Substrate/Substrate.cpp
+++ b/Engine/Source/Runtime/Renderer/Private/Substrate/Substrate.cpp
@@ -260,7 +260,9 @@ bool UsesStochasticLightingClassification(EShaderPlatform InPlatform)
 		Substrate::IsSubstrateBlendableGBufferEnabled(InPlatform) && 
 		Substrate::UsesSubstrateMaterialBuffer(InPlatform) && // Only needed for deferred lighting
 		(DoesPlatformSupportLumenGI(InPlatform) || MegaLights::ShouldCompileShaders(InPlatform)) &&
-		CVarSubstrateUseStochasticLightingClassification.GetValueOnRenderThread();
+		CVarSubstrateUseStochasticLightingClassification.GetValueOnRenderThread() && 
+		!Substrate::UseTileSize16();	// Only supported when tile size is 8 otherwise tiles are incomplete (stochastic pass would have to gather 2x2 tiles of 8 pixel ot create a new tile of 16x16 pixels).
+	;
 }
 
 static EPixelFormat GetClassificationTileFormat(const FIntPoint& InResolution, uint32 InTileEncoding)
@@ -984,8 +986,9 @@ class FSubstrateMaterialTileClassificationPassCS : public FGlobalShader
 
 	class FCmask : SHADER_PERMUTATION_BOOL("PERMUTATION_CMASK");
 	class FWaveOps : SHADER_PERMUTATION_BOOL("PERMUTATION_WAVE_OPS");
-	class FDecal : SHADER_PERMUTATION_BOOL("PERMUTATION_DECAL"); 
-	using FPermutationDomain = TShaderPermutationDomain<FCmask, FWaveOps, FDecal>;
+	class FDecal : SHADER_PERMUTATION_BOOL("PERMUTATION_DECAL");
+	class FClassificationTileSize16 : SHADER_PERMUTATION_BOOL("SUBSTRATE_CLASSIFICATION_TILE_SIZE_16");
+	using FPermutationDomain = TShaderPermutationDomain<FCmask, FWaveOps, FDecal, FClassificationTileSize16>;
 
 	BEGIN_SHADER_PARAMETER_STRUCT(FParameters, )
 		SHADER_PARAMETER_STRUCT_REF(FViewUniformShaderParameters, ViewUniformBuffer)
@@ -1317,6 +1320,7 @@ void AddSubstrateInternalClassificationTilePass(
 	FSubstrateTilePassVS::FPermutationDomain VSPermutationVector;
 	VSPermutationVector.Set< FSubstrateTilePassVS::FEnableDebug >(bDebug);
 	VSPermutationVector.Set< FSubstrateTilePassVS::FEnableTexCoordScreenVector >(false);
+	VSPermutationVector.Set< FSubstrateTilePassVS::FClassificationTileSize16 >(Substrate::UseTileSize16());
 	TShaderMapRef<FSubstrateTilePassVS> VertexShader(View.ShaderMap, VSPermutationVector);
 	TShaderMapRef<FSubstrateMaterialStencilTaggingPassPS> PixelShader(View.ShaderMap);
 
@@ -1689,10 +1693,12 @@ void AddSubstrateMaterialClassificationPass(FRDGBuilder& GraphBuilder, const FMi
 			// If Dbuffer pass (i.e. apply DBuffer data after the base-pass) is enabled, run special classification for outputing tile with/without tiles
 			const bool bDBufferTiles = IsDBufferPassEnabled(Platform) && CVarSubstrateDBufferPassDedicatedTiles.GetValueOnRenderThread() > 0 && DBufferTextures.IsValid() && IsConsolePlatform(Platform);
 
+			const bool bUseTileSize16 = Substrate::UseTileSize16();
 			FSubstrateMaterialTileClassificationPassCS::FPermutationDomain PermutationVector;
 			PermutationVector.Set< FSubstrateMaterialTileClassificationPassCS::FCmask >(bSupportCMask);
 			PermutationVector.Set< FSubstrateMaterialTileClassificationPassCS::FWaveOps >(bWaveOps);
 			PermutationVector.Set< FSubstrateMaterialTileClassificationPassCS::FDecal>(bDBufferTiles);
+			PermutationVector.Set< FSubstrateMaterialTileClassificationPassCS::FClassificationTileSize16 >(bUseTileSize16);
 			TShaderMapRef<FSubstrateMaterialTileClassificationPassCS> ComputeShader(View.ShaderMap, PermutationVector);
 			FSubstrateMaterialTileClassificationPassCS::FParameters* PassParameters = GraphBuilder.AllocParameters<FSubstrateMaterialTileClassificationPassCS::FParameters>();
 			PassParameters->ViewUniformBuffer = View.ViewUniformBuffer;
@@ -1715,7 +1721,7 @@ void AddSubstrateMaterialClassificationPass(FRDGBuilder& GraphBuilder, const FMi
 				PassParameters->TileListBufferOffsets[TileType] = FUintVector4(SubstrateViewData->ClassificationTileListBufferOffset[TileType], 0, 0, 0);
 			}
 
-			const uint32 GroupSize = 8;
+			const uint32 GroupSize = bUseTileSize16 ? 16 : 8;
 			FComputeShaderUtils::AddPass(
 				GraphBuilder,
 				RDG_EVENT_NAME("Substrate::MaterialTileClassification(%s%s)", bWaveOps ? TEXT("Wave") : TEXT("SharedMemory"), bSupportCMask ? TEXT(", CMask") : TEXT("")),
diff --git a/Engine/Source/Runtime/Renderer/Private/Substrate/Substrate.h b/Engine/Source/Runtime/Renderer/Private/Substrate/Substrate.h
--- a/Engine/Source/Runtime/Renderer/Private/Substrate/Substrate.h
+++ b/Engine/Source/Runtime/Renderer/Private/Substrate/Substrate.h
@@ -265,7 +265,8 @@ class FSubstrateTilePassVS : public FGlobalShader
 
 	class FEnableDebug : SHADER_PERMUTATION_BOOL("PERMUTATION_ENABLE_DEBUG");
 	class FEnableTexCoordScreenVector : SHADER_PERMUTATION_BOOL("PERMUTATION_ENABLE_TEXCOORD_SCREENVECTOR");
-	using FPermutationDomain = TShaderPermutationDomain<FEnableDebug, FEnableTexCoordScreenVector>;
+	class FClassificationTileSize16 : SHADER_PERMUTATION_BOOL("SUBSTRATE_CLASSIFICATION_TILE_SIZE_16");
+	using FPermutationDomain = TShaderPermutationDomain<FEnableDebug, FEnableTexCoordScreenVector, FClassificationTileSize16>;
 
 	BEGIN_SHADER_PARAMETER_STRUCT(FParameters, )
 		SHADER_PARAMETER_STRUCT_REF(FViewUniformShaderParameters, View)
