[HWRT] Add priority based update to dynamic ray tracing geometry update:
* Either driven by distance or frames-since-last-update or blend of both.
* Out of frustum geometries have their priority reduced.
* Controlled with r.RayTracing.DynamicGeometry.PriorityBasedUpdate (disabled by default).

#rnx

diff --git a/Engine/Source/Runtime/Renderer/Private/RayTracing/RayTracingDynamicGeometryUpdateManager.cpp b/Engine/Source/Runtime/Renderer/Private/RayTracing/RayTracingDynamicGeometryUpdateManager.cpp
--- a/Engine/Source/Runtime/Renderer/Private/RayTracing/RayTracingDynamicGeometryUpdateManager.cpp
+++ b/Engine/Source/Runtime/Renderer/Private/RayTracing/RayTracingDynamicGeometryUpdateManager.cpp
@@ -52,6 +52,27 @@ static TAutoConsoleVariable<int32> CVarRTDynGeomForceBuildMaxUpdatePrimitivesPer
 	ECVF_RenderThreadSafe
 );
 
+static TAutoConsoleVariable<int32> CVarRTDynGeomPriorityBasedUpdate(
+	TEXT("r.RayTracing.DynamicGeometry.PriorityBasedUpdate"),
+	0,
+	TEXT("Whether dynamic geometry should be update based on priority."),
+	ECVF_RenderThreadSafe
+);
+
+static TAutoConsoleVariable<int32> CVarRTDynGeomPriorityBasedUpdateMaxFramesSinceLastUpdate(
+	TEXT("r.RayTracing.DynamicGeometry.PriorityBasedUpdate.MaxFramesSinceLastUpdate"),
+	60,
+	TEXT("Sets the maximum number of frames since the last update after which it is assigned the highest update priority."),
+	ECVF_RenderThreadSafe
+);
+
+static TAutoConsoleVariable<float> CVarRTDynGeomPriorityBasedUpdateLastUpdateWeight(
+	TEXT("r.RayTracing.DynamicGeometry.PriorityBasedUpdate.LastUpdateWeight"),
+	0.01f,
+	TEXT("Sets how time (frames since last update) affects the overall geometry priority weight (value between 0.0 - 1.0)."),
+	ECVF_RenderThreadSafe
+);
+
 static TAutoConsoleVariable<int32> CVarRTDynGeomForceBuildMinUpdatesSinceLastBuild(
 	TEXT("r.RayTracing.DynamicGeometry.ForceBuild.MinUpdatesSinceLastBuild"),
 	-1,
@@ -59,6 +80,32 @@ static TAutoConsoleVariable<int32> CVarRTDynGeomForceBuildMinUpdatesSinceLastBui
 	ECVF_RenderThreadSafe
 );
 
+// Lower value means higher priority - max priority is 0.0.
+float CalculatePriority(const FSceneView* View,	const FPrimitiveSceneProxy* PrimitiveSceneProxy, const FRayTracingDynamicGeometryUpdateParams& UpdateParams)
+{
+	const float MaxDistance = GetRayTracingCullingRadius();
+	const float InvMaxDistanceSq = 1.0f / (MaxDistance * MaxDistance);
+	const FBoxSphereBounds& ProxyBounds = PrimitiveSceneProxy->GetBounds();
+
+	const float InvMaxAge = 1.0f / float(CVarRTDynGeomPriorityBasedUpdateMaxFramesSinceLastUpdate.GetValueOnRenderThread());
+
+	const float Distance = ProxyBounds.ComputeSquaredDistanceFromBoxToPoint(View->ViewMatrices.GetViewOrigin());
+	const float Age = GFrameCounterRenderThread - UpdateParams.Geometry->LastUpdatedFrame;
+
+	const float DistancePriority = FMath::Clamp(Distance * InvMaxDistanceSq, 0.0f, 1.0f);
+	const float AgePriority = 1.0f - FMath::Clamp(Age * InvMaxAge, 0.0f, 1.0f);
+
+	const float AgeWeight = FMath::Clamp(CVarRTDynGeomPriorityBasedUpdateLastUpdateWeight.GetValueOnRenderThread(), 0.0f, 1.0f);
+	float Priority = AgeWeight * AgePriority + (1.0f - AgeWeight) * DistancePriority;
+
+	const bool bInFrustum = View->GetCullingFrustum().IntersectBox(ProxyBounds.Origin, ProxyBounds.BoxExtent);
+	if (bInFrustum)
+	{
+		Priority *= 0.5;
+	}
+	return FMath::Clamp(Priority, 0.0f, 1.0f);
+}
+
 class FRayTracingDynamicGeometryConverterCS : public FMeshMaterialShader
 {
 	DECLARE_SHADER_TYPE(FRayTracingDynamicGeometryConverterCS, MeshMaterial);
@@ -328,6 +375,11 @@ void FRayTracingDynamicGeometryUpdateManager::AddDynamicGeometryToUpdate(
 	GeometryBuildParams.DispatchCommands.Reserve(UpdateParams.MeshBatches.Num());
 	InstancedViewUniformBuffers.FindOrAdd(View->ViewUniformBuffer, View->GetInstancedViewUniformBuffer());
 
+	if (CVarRTDynGeomPriorityBasedUpdate.GetValueOnRenderThread())
+	{
+		GeometryBuildParams.Priority = CalculatePriority(View, PrimitiveSceneProxy, UpdateParams);
+	}	
+
 	// Only update when we have mesh batches
 	if (!UpdateParams.MeshBatches.IsEmpty())
 	{
@@ -715,15 +767,25 @@ void FRayTracingDynamicGeometryUpdateManager::ScheduleUpdates(bool bUseTracingFe
 	const bool bNeedsSorting = (int32(MaxUpdatePrimitivesPerFrame) != -1) || (MaxForceBuildPrimitivesPerFrame != 0);
 	if (bNeedsSorting)
 	{
-		DynamicGeometryUpdates.Sort([](const FRayTracingDynamicGeometryBuildParams& InLHS, const FRayTracingDynamicGeometryBuildParams& InRHS)
-			{
-				if (InLHS.Geometry->LastUpdatedFrame == InRHS.Geometry->LastUpdatedFrame)
+		if (CVarRTDynGeomPriorityBasedUpdate.GetValueOnRenderThread())
+		{
+			DynamicGeometryUpdates.Sort([](const FRayTracingDynamicGeometryBuildParams& InLHS, const FRayTracingDynamicGeometryBuildParams& InRHS)
 				{
-					return InLHS.Geometry->NumUpdatesSinceLastBuild > InRHS.Geometry->NumUpdatesSinceLastBuild;
-				}
+					return InLHS.Priority < InRHS.Priority;
+				});
+		}
+		else
+		{
+			DynamicGeometryUpdates.Sort([](const FRayTracingDynamicGeometryBuildParams& InLHS, const FRayTracingDynamicGeometryBuildParams& InRHS)
+				{
+					if (InLHS.Geometry->LastUpdatedFrame == InRHS.Geometry->LastUpdatedFrame)
+					{
+						return InLHS.Geometry->NumUpdatesSinceLastBuild > InRHS.Geometry->NumUpdatesSinceLastBuild;
+					}
 
-				return InLHS.Geometry->LastUpdatedFrame < InRHS.Geometry->LastUpdatedFrame;
-			});
+					return InLHS.Geometry->LastUpdatedFrame < InRHS.Geometry->LastUpdatedFrame;
+				});
+		}
 	}	
 
 	for (const FRayTracingDynamicGeometryBuildParams& Update : DynamicGeometryUpdates)
diff --git a/Engine/Source/Runtime/Renderer/Public/RayTracingDynamicGeometryUpdateManager.h b/Engine/Source/Runtime/Renderer/Public/RayTracingDynamicGeometryUpdateManager.h
--- a/Engine/Source/Runtime/Renderer/Public/RayTracingDynamicGeometryUpdateManager.h
+++ b/Engine/Source/Runtime/Renderer/Public/RayTracingDynamicGeometryUpdateManager.h
@@ -84,6 +84,7 @@ private:
 		FRHIUniformBuffer* ViewUniformBuffer = nullptr;
 		FRayTracingGeometry* Geometry = nullptr;
 		int32 SegmentOffset = -1;
+		float Priority = 0.0;
 	};
 
 	void AddDispatchCommands(FRHICommandListBase& RHICmdList,
