Replaced full update with allocation update and always update the header.

diff --git a/Engine/Source/Runtime/Renderer/Private/Skinning/SkinningSceneExtension.cpp b/Engine/Source/Runtime/Renderer/Private/Skinning/SkinningSceneExtension.cpp
--- a/Engine/Source/Runtime/Renderer/Private/Skinning/SkinningSceneExtension.cpp
+++ b/Engine/Source/Runtime/Renderer/Private/Skinning/SkinningSceneExtension.cpp
@@ -754,15 +754,6 @@ FSkinningSceneExtension::FUpdater::FUpdater(FSkinningSceneExtension& InSceneData
 	SceneData->UpdateCounter++;
 }
 
-void FSkinningSceneExtension::FUpdater::AddHeaderUpdate(FHeaderData& HeaderData, int32 PrimitiveIndex)
-{
-	if (HeaderData.UpdateCounters.Header != SceneData->UpdateCounter)
-	{
-		HeaderData.UpdateCounters.Header = SceneData->UpdateCounter;
-		HeaderDataUpdateList.Emplace(PrimitiveIndex);
-	}
-}
-
 void FSkinningSceneExtension::FUpdater::End()
 {
 	// Ensure these tasks finish before we fall out of scope.
@@ -924,10 +915,10 @@ void FSkinningSceneExtension::FUpdater::PostMeshUpdate(FRDGBuilder& GraphBuilder
 			int32 NumUpdates = 0;
 			int32 NumLODChanges = 0;
 
-			const auto UpdateLOD = [this, &NumLODChanges] (FHeaderData& HeaderData, int32 PrimitiveIndex, int32 LODIndex)
+			const auto UpdateLambda = [this, &NumLODChanges] (FHeaderData& HeaderData, int32 PrimitiveIndex, int32 LODIndex)
 			{
 				HeaderData.DirtyBoneTransforms |= EDirtyBoneTransforms::Current;
-				HeaderData.UpdateCounters.UpdateLOD = SceneData->UpdateCounter;
+				HeaderData.UpdateCounter = SceneData->UpdateCounter;
 
 				if (!HeaderData.bAlwaysUpdateTransforms)
 				{
@@ -942,14 +933,14 @@ void FSkinningSceneExtension::FUpdater::PostMeshUpdate(FRDGBuilder& GraphBuilder
 					HeaderData.MaxHierarchyCount = HeaderData.Proxy->GetMaxBoneHierarchyCount();
 					HeaderData.MaxObjectSpaceCount = HeaderData.Proxy->GetMaxBoneObjectSpaceCount();
 
-					FullUpdateList.Emplace(PrimitiveIndex);
+					AllocationUpdateList.Emplace(PrimitiveIndex);
 				}
 
-				AddHeaderUpdate(HeaderData, PrimitiveIndex);
+				HeaderDataUpdateList.Emplace(PrimitiveIndex);
 			};
 
 			TConstArrayView<FSkeletalMeshUpdateEvent> UpdateEvents = SkeletalMeshSubmitData->GetUpdateEvents();
-			FullUpdateList.Reserve(UpdatedSceneInfoList.Num() + UpdateEvents.Num());
+			AllocationUpdateList.Reserve(UpdatedSceneInfoList.Num() + UpdateEvents.Num());
 			HeaderDataUpdateList.Reserve(UpdatedSceneInfoList.Num() + UpdateEvents.Num());
 
 			for (FSkeletalMeshUpdateEvent UpdateEvent : UpdateEvents)
@@ -976,7 +967,8 @@ void FSkinningSceneExtension::FUpdater::PostMeshUpdate(FRDGBuilder& GraphBuilder
 					continue;
 				}
 
-				UpdateLOD(HeaderData, PrimitiveIndex, UpdateEvent.LODIndex);
+				UpdateLambda(HeaderData, PrimitiveIndex, UpdateEvent.LODIndex);
+				HeaderData.UpdateCounter = SceneData->UpdateCounter;
 				NumUpdates++;
 			}
 			
@@ -986,9 +978,10 @@ void FSkinningSceneExtension::FUpdater::PostMeshUpdate(FRDGBuilder& GraphBuilder
 				if (SceneData->HeaderDatas.IsValidIndex(Index))
 				{
 					FHeaderData& HeaderData = SceneData->HeaderDatas[Index];
-					if (HeaderData.UpdateCounters.UpdateLOD != SceneData->UpdateCounter)
+
+					if (HeaderData.UpdateCounter != SceneData->UpdateCounter)
 					{
-						UpdateLOD(HeaderData, Index, HeaderData.MeshObject->GetLOD());
+						UpdateLambda(HeaderData, Index, HeaderData.MeshObject->GetLOD());
 					}
 				}
 			}
@@ -1120,11 +1113,6 @@ void FSkinningSceneExtension::FUpdater::PostMeshUpdate(FRDGBuilder& GraphBuilder
 
 		check(Data.UniqueAnimationCount * Data.MaxTransformCount * 2u == Data.TransformBufferCount);
 
-		if (bRequireUpload)
-		{
-			AddHeaderUpdate(Data, Data.PrimitiveSceneInfo->GetPersistentIndex().Index);
-		}
-
 		if (Data.Proxy->UseSkeletonBatching())
 		{
 			Data.bIsBatched = true;
@@ -1167,7 +1155,7 @@ void FSkinningSceneExtension::FUpdater::PostMeshUpdate(FRDGBuilder& GraphBuilder
 			else
 			{
 				// Only check to reallocate space for primitives that have requested an update
-				for (int32 PersistentIndex : FullUpdateList)
+				for (int32 PersistentIndex : AllocationUpdateList)
 				{
 					AllocSpaceForPrimitive(PersistentIndex);
 					bHasUpdates = true;
@@ -1328,7 +1316,7 @@ void FSkinningSceneExtension::FUpdater::PostMeshUpdate(FRDGBuilder& GraphBuilder
 			}
 			else
 			{
-				for (int32 PersistentIndex : FullUpdateList)
+				for (int32 PersistentIndex : AllocationUpdateList)
 				{
 					UploadHierarchyData(SceneData->HeaderDatas[PersistentIndex]);
 				}
@@ -1356,7 +1344,7 @@ void FSkinningSceneExtension::FUpdater::PostMeshUpdate(FRDGBuilder& GraphBuilder
 			}
 			else
 			{
-				for (int32 PersistentIndex : FullUpdateList)
+				for (int32 PersistentIndex : AllocationUpdateList)
 				{
 					UploadTransformData(SceneData->HeaderDatas[PersistentIndex], bProvidersEnabled);
 				}
diff --git a/Engine/Source/Runtime/Renderer/Private/Skinning/SkinningSceneExtension.h b/Engine/Source/Runtime/Renderer/Private/Skinning/SkinningSceneExtension.h
--- a/Engine/Source/Runtime/Renderer/Private/Skinning/SkinningSceneExtension.h
+++ b/Engine/Source/Runtime/Renderer/Private/Skinning/SkinningSceneExtension.h
@@ -25,12 +25,6 @@ class FSkinningSceneExtension : public ISceneExtension
 	DECLARE_SCENE_EXTENSION(RENDERER_API, FSkinningSceneExtension);
 	struct FHeaderData;
 
-	struct FUpdateCounters
-	{
-		uint32 Header = 0;
-		uint32 UpdateLOD = 0;
-	};
-
 public:
 	class FUpdater : public ISceneExtensionUpdater
 	{
@@ -50,7 +44,7 @@ public:
 
 		FSkinningSceneExtension* SceneData = nullptr;
 		TConstArrayView<FPrimitiveSceneInfo*> AddedList;
-		TArray<int32, FSceneRenderingArrayAllocator> FullUpdateList;
+		TArray<int32, FSceneRenderingArrayAllocator> AllocationUpdateList;
 		TArray<int32, FSceneRenderingArrayAllocator> HeaderDataUpdateList;
 		const bool bEnableAsync = true;
 		bool bForceFullUpload = false;
@@ -114,7 +108,7 @@ private:
 		FSkinningSceneExtensionProxy* Proxy     = nullptr;
 		const FSkeletalMeshObject* MeshObject   = nullptr;
 		FGuid ProviderId;
-		FUpdateCounters UpdateCounters;
+		uint32 UpdateCounter                     = 0;
 		uint32 InstanceSceneDataOffset           = 0;
 		uint32 NumInstanceSceneDataEntries       = 0;
 		uint32 ObjectSpaceBufferOffset           = INDEX_NONE;
