Fixed AnimBank animation delta time computation to not double count frames.

diff --git a/Engine/Source/Runtime/Renderer/Private/Skinning/AnimBankTransformProvider.cpp b/Engine/Source/Runtime/Renderer/Private/Skinning/AnimBankTransformProvider.cpp
--- a/Engine/Source/Runtime/Renderer/Private/Skinning/AnimBankTransformProvider.cpp
+++ b/Engine/Source/Runtime/Renderer/Private/Skinning/AnimBankTransformProvider.cpp
@@ -558,6 +558,11 @@ static FAnimBankCPUData BuildAnimBankCPUData(const FAnimBankRecordMap& BankRecor
 
 void FAnimBankTransformProvider::AdvanceAnimation(FSkinningTransformProvider::FProviderContext& Context)
 {
+	if (LastAnimationFrame == GFrameNumberRenderThread)
+	{
+		return;
+	}
+
 	const float GlobalTimeScale = CVarAnimBankTimeScale.GetValueOnRenderThread();
 	const float DeltaTime = Context.ViewFamily.Time.GetDeltaWorldTimeSeconds();
 
@@ -572,6 +577,8 @@ void FAnimBankTransformProvider::AdvanceAnimation(FSkinningTransformProvider::FP
 		BankRecord.PreviousTime = BankRecord.CurrentTime;
 		BankRecord.CurrentTime += (DeltaTime * BankRecord.Desc.PlayRate) * GlobalTimeScale;
 	}
+
+	LastAnimationFrame = GFrameNumberRenderThread;
 }
 
 void FAnimBankTransformProvider::ScatterAnimation(
diff --git a/Engine/Source/Runtime/Renderer/Private/Skinning/AnimBankTransformProvider.h b/Engine/Source/Runtime/Renderer/Private/Skinning/AnimBankTransformProvider.h
--- a/Engine/Source/Runtime/Renderer/Private/Skinning/AnimBankTransformProvider.h
+++ b/Engine/Source/Runtime/Renderer/Private/Skinning/AnimBankTransformProvider.h
@@ -40,6 +40,7 @@ private:
 private:
 	FAnimBankRecordMap BankRecordMap;
 	FSpanAllocator BankAllocator;
+	uint64 LastAnimationFrame = 0;
 };
 
 RENDERER_API const FSkinningTransformProvider::FProviderId& GetAnimBankProviderId();
\ No newline at end of file
