Heterogeneous Volumes: Removing SetParameters from voxel grid builder shaders

diff --git a/Engine/Source/Runtime/Renderer/Private/HeterogeneousVolumes/HeterogeneousVolumesVoxelGridPipeline.cpp b/Engine/Source/Runtime/Renderer/Private/HeterogeneousVolumes/HeterogeneousVolumesVoxelGridPipeline.cpp
--- a/Engine/Source/Runtime/Renderer/Private/HeterogeneousVolumes/HeterogeneousVolumesVoxelGridPipeline.cpp
+++ b/Engine/Source/Runtime/Renderer/Private/HeterogeneousVolumes/HeterogeneousVolumesVoxelGridPipeline.cpp
@@ -725,7 +725,7 @@ class FRasterizeBottomLevelFrustumGridCS : public FMeshMaterialShader
 
 		OutEnvironment.SetDefine(TEXT("GET_PRIMITIVE_DATA_OVERRIDE"), 1);
 	}
-
+#if 0
 	void SetParameters(
 		FRHIComputeCommandList& RHICmdList,
 		FRHIComputeShader* ShaderRHI,
@@ -737,7 +737,7 @@ class FRasterizeBottomLevelFrustumGridCS : public FMeshMaterialShader
 		FMaterialShader::SetViewParameters(RHICmdList, ShaderRHI, View, View.ViewUniformBuffer);
 		FMaterialShader::SetParameters(RHICmdList, ShaderRHI, MaterialProxy, Material, View);
 	}
-
+#endif
 	static int32 GetThreadGroupSize1D() { return GetThreadGroupSize3D() * GetThreadGroupSize3D() * GetThreadGroupSize3D(); }
 	static int32 GetThreadGroupSize3D() { return 4; }
 };
@@ -1041,7 +1041,7 @@ class FRasterizeBottomLevelOrthoGridCS : public FMeshMaterialShader
 
 		OutEnvironment.SetDefine(TEXT("GET_PRIMITIVE_DATA_OVERRIDE"), 1);
 	}
-
+#if 0
 	void SetParameters(
 		FRHIComputeCommandList& RHICmdList,
 		FRHIComputeShader* ShaderRHI,
@@ -1053,7 +1053,7 @@ class FRasterizeBottomLevelOrthoGridCS : public FMeshMaterialShader
 		FMaterialShader::SetViewParameters(RHICmdList, ShaderRHI, View, View.ViewUniformBuffer);
 		FMaterialShader::SetParameters(RHICmdList, ShaderRHI, MaterialProxy, Material, View);
 	}
-
+#endif
 	static int32 GetThreadGroupSize1D() { return GetThreadGroupSize3D() * GetThreadGroupSize3D() * GetThreadGroupSize3D(); }
 	static int32 GetThreadGroupSize3D() { return 4; }
 };
