Remove log line about LUT pixel format than can spam from editor, but keep everywhere else where it won't spam

diff --git a/Engine/Source/Runtime/Renderer/Private/ScenePrivate.h b/Engine/Source/Runtime/Renderer/Private/ScenePrivate.h
--- a/Engine/Source/Runtime/Renderer/Private/ScenePrivate.h
+++ b/Engine/Source/Runtime/Renderer/Private/ScenePrivate.h
@@ -1359,7 +1359,10 @@ public:
 		{
 			LUTPixelFormat = PF_FloatRGBA;
 		}
+#if !WITH_EDITOR
+		// Editor can enter this function frequently. Log for other build targets.
 		UE_LOG(LogRenderer, Log, TEXT("Pixel format for LUT is %s"), GetPixelFormatString(LUTPixelFormat));
+#endif
 
 		FPooledRenderTargetDesc Desc = FPooledRenderTargetDesc::CreateVolumeDesc(
 			LUTSize,
