[MegaFilters] Add the basic MegaFilters interface structure.

diff --git a/Engine/Source/Runtime/Engine/Private/SceneView.cpp b/Engine/Source/Runtime/Engine/Private/SceneView.cpp
--- a/Engine/Source/Runtime/Engine/Private/SceneView.cpp
+++ b/Engine/Source/Runtime/Engine/Private/SceneView.cpp
@@ -2988,6 +2988,7 @@ FSceneViewFamily::FSceneViewFamily(const ConstructionValues& CVS)
 	SceneRenderer(nullptr),
 	ScreenPercentageInterface(nullptr),
 	TemporalUpscalerInterface(nullptr),
+	MegaFiltersInterface(nullptr),
 	PrimarySpatialUpscalerInterface(nullptr),
 	SecondarySpatialUpscalerInterface(nullptr)
 {
@@ -3065,6 +3066,12 @@ FSceneViewFamily::~FSceneViewFamily()
 		delete reinterpret_cast<ISceneViewFamilyExtention*>(TemporalUpscalerInterface);
 	}
 
+	if (MegaFiltersInterface)
+	{
+		// IMegaFilters* is only defined in renderer's private header because no backward compatibility is provided between major version.
+		delete reinterpret_cast<ISceneViewFamilyExtention*>(MegaFiltersInterface);
+	}
+
 	if (PrimarySpatialUpscalerInterface)
 	{
 		// ISpatialUpscaler* is only defined in renderer's private header because no backward compatibility is provided between major version.
@@ -3083,6 +3090,7 @@ FSceneViewFamily::FSceneViewFamily(FSceneViewFamily&& InViewFamily)
 {
 	check(ScreenPercentageInterface == nullptr);
 	check(TemporalUpscalerInterface == nullptr);
+	check(MegaFiltersInterface == nullptr);
 	check(PrimarySpatialUpscalerInterface == nullptr);
 	check(SecondarySpatialUpscalerInterface == nullptr);
 }
diff --git a/Engine/Source/Runtime/Engine/Public/SceneView.h b/Engine/Source/Runtime/Engine/Public/SceneView.h
--- a/Engine/Source/Runtime/Engine/Public/SceneView.h
+++ b/Engine/Source/Runtime/Engine/Public/SceneView.h
@@ -41,7 +41,7 @@ namespace UE::Renderer::Private
 {
 
 class ITemporalUpscaler;
-
+class IMegaFilters;
 }
 
 class FRenderTarget;
@@ -2581,6 +2581,18 @@ PRAGMA_ENABLE_DEPRECATION_WARNINGS
 		return TemporalUpscalerInterface;
 	}
 
+	inline void SetMegaFiltersInterface(UE::Renderer::Private::IMegaFilters* InMegaFiltersInterface)
+	{
+		check(InMegaFiltersInterface);
+		checkf(MegaFiltersInterface == nullptr, TEXT("View family already had mega filters assigned."));
+		MegaFiltersInterface = InMegaFiltersInterface;
+	}
+
+	inline const UE::Renderer::Private::IMegaFilters* GetMegaFiltersInterface() const
+	{
+		return MegaFiltersInterface;
+	}
+
 	inline void SetPrimarySpatialUpscalerInterface(ISpatialUpscaler* InSpatialUpscalerInterface)
 	{
 		check(InSpatialUpscalerInterface);
@@ -2620,6 +2632,7 @@ private:
 
 	/** Renderer private interfaces, automatically have same lifetime as FSceneViewFamily. */
 	UE::Renderer::Private::ITemporalUpscaler* TemporalUpscalerInterface;
+	UE::Renderer::Private::IMegaFilters* MegaFiltersInterface;
 	ISpatialUpscaler* PrimarySpatialUpscalerInterface;
 	ISpatialUpscaler* SecondarySpatialUpscalerInterface;
 
diff --git a/Engine/Source/Runtime/Renderer/Internal/MegaFiltersInterface.h b/Engine/Source/Runtime/Renderer/Internal/MegaFiltersInterface.h
new file mode 100644
--- /dev/null
+++ b/Engine/Source/Runtime/Renderer/Internal/MegaFiltersInterface.h
@@ -0,0 +1,92 @@
+// Copyright Epic Games, Inc. All Rights Reserved.
+
+#pragma once
+
+#include "ScreenPass.h"
+
+namespace UE::Renderer::Private
+{
+
+enum class EMegaFiltersPassType : uint32
+{
+	Default,
+	MegaLights,
+	MAX
+};
+
+/** Interface for implementing different filters */
+class IMegaFilters : public ISceneViewFamilyExtention
+{
+public:
+
+	/** Ref counted history to be saved in the history*/
+	class IHistory : public FRefCountBase
+	{
+	public:
+		virtual ~IHistory() {}
+
+		/** Debug name of the history. Must exactly point to the same const TCHAR* as IMegaFilters::GetDebugName().
+		* This is used for debugging GPU memory uses of a viewport, but also to ensure IHistory is fed to a compatible IMegaFilters.
+		*/
+		virtual const TCHAR* GetDebugName() const = 0;
+
+		/** Size of the history on the GPU in bytes. */
+		virtual uint64 GetGPUSizeBytes() const = 0;
+	};
+
+	/** Inputs of the filter*/
+	struct FInputs
+	{
+		/** Which filter pass this input is used for*/
+		EMegaFiltersPassType MegaFiltersPassType;
+
+		/** Outputs view rect that must be on FOutputs::FullRes::ViewRect. */
+		TArray<FIntRect> OutputViewRects;
+
+		/** Configuration variables per pass (e.g., TemporalJitter and PreExposure)*/
+		TArray<TVariant<float, int, FVector2f, FVector4f>> Variables;
+
+		/** Input texture and buffer*/
+		TArray<FScreenPassTexture> ScreenPassTextures;
+
+		TArray<FRDGTextureRef> Textures;
+
+		TArray<FRDGBufferRef> Buffers;
+
+		/** The history of the previous frame set by FOutputs::NewHistory. PrevHistory->GetDebugName() is guarentee to match the IMegaFilters.*/
+		TRefCountPtr<IHistory> PrevHistory;
+	};
+
+	/** Outputs of the third party MegaFilters. */
+	struct FOutputs
+	{
+		/** Output of the MegaFilters. FullRes[i].ViewRect must match FInputs::OutputViewRects[i]. */
+		TArray<FScreenPassTexture> FullRes;
+
+		/** New history to be kept alive for next frame. NewHistory->GetDebugName() must exactly point to the same const TCHAR* as IMegaFilters::GetDebugName(). 
+		Meaningful when bPassThrough = true;
+		*/
+		TRefCountPtr<IHistory> NewHistory;
+
+		/** Indicate the pass does nothing*/
+		bool bPassThrough = false;
+	};
+
+	virtual ~IMegaFilters() {};
+
+	/** Debug name of the history. Must exactly point to the same const TCHAR* as IMegaFilters::IHistory::GetDebugName(). */
+	virtual const TCHAR* GetDebugName() const = 0;
+
+	/** Adds the necessary passes into RDG for filtering or temporal upscaling the rendering resolution to desired output res. */
+	virtual FOutputs AddPasses(
+		FRDGBuilder& GraphBuilder,
+		const FSceneView& View,
+		const FInputs& Inputs) const = 0;
+
+	virtual IMegaFilters* Fork_GameThread(const class FSceneViewFamily& ViewFamily) const = 0;
+};
+
+} // namespace UE::Renderer::Private
+
+
+
diff --git a/Engine/Source/Runtime/Renderer/Private/PostProcess/TemporalAA.cpp b/Engine/Source/Runtime/Renderer/Private/PostProcess/TemporalAA.cpp
--- a/Engine/Source/Runtime/Renderer/Private/PostProcess/TemporalAA.cpp
+++ b/Engine/Source/Runtime/Renderer/Private/PostProcess/TemporalAA.cpp
@@ -1087,6 +1087,54 @@ FDefaultTemporalUpscaler::FOutputs AddGen4MainTemporalAAPasses(
 	return Outputs;
 } // AddGen4MainTemporalAAPasses()
 
+UE::Renderer::Private::IMegaFilters::FOutputs AddMegaFiltersPass(FRDGBuilder& GraphBuilder, const FViewInfo& View, const TCHAR* Name,const UE::Renderer::Private::IMegaFilters::FInputs& Inputs)
+{
+	using UE::Renderer::Private::IMegaFilters;
+
+	const IMegaFilters* MegaFiltersToUse = View.Family->GetMegaFiltersInterface();
+	IMegaFilters::FOutputs MegaFilterOutputs;
+
+	if (MegaFiltersToUse)
+	{
+		const TCHAR* FilterName = MegaFiltersToUse->GetDebugName();
+
+		IMegaFilters::FInputs MegaFilterInputs = Inputs;
+		checkf(MegaFilterInputs.PrevHistory == nullptr,TEXT("History is set internally"));
+		{
+			// History can be accessed/updated in the same frame at different location for the same MegaFiltersPassType.
+			TRefCountPtr<IMegaFilters::IHistory> PrevHistory = 
+				(View.ViewState && View.ViewState->PrevFrameViewInfo.MegaFiltersHistory) ? 
+				View.ViewState->PrevFrameViewInfo.MegaFiltersHistory : View.PrevViewInfo.MegaFiltersHistory;
+
+			if (PrevHistory && PrevHistory->GetDebugName() == FilterName)
+			{
+				MegaFilterInputs.PrevHistory = PrevHistory;
+			}
+		}
+
+		checkf(MegaFilterInputs.OutputViewRects.Num() > 0, TEXT("At least define one output view rect in MegaFilterInputs.OutputViewRects"));
+
+		RDG_EVENT_SCOPE(
+			GraphBuilder,
+			"MegaFilters %s %dx%d -> %dx%d",
+			Name,
+			View.ViewRect.Width(), View.ViewRect.Height(),
+			MegaFilterInputs.OutputViewRects[0].Width(), MegaFilterInputs.OutputViewRects[0].Height());
+
+		MegaFilterOutputs = MegaFiltersToUse->AddPasses(
+			GraphBuilder,
+			View,
+			MegaFilterInputs);
+
+		if (View.ViewState && !View.bStatePrevViewInfoIsReadOnly)
+		{
+			View.ViewState->PrevFrameViewInfo.MegaFiltersHistory = MegaFilterOutputs.NewHistory;
+		}
+	}
+
+	return MegaFilterOutputs;
+}
+
 EMainTAAPassConfig GetMainTAAPassConfig(const FViewInfo& View)
 {
 	if (!IsPostProcessingEnabled(View) && !View.bSceneCaptureMainViewJitter)
diff --git a/Engine/Source/Runtime/Renderer/Private/PostProcess/TemporalAA.h b/Engine/Source/Runtime/Renderer/Private/PostProcess/TemporalAA.h
--- a/Engine/Source/Runtime/Renderer/Private/PostProcess/TemporalAA.h
+++ b/Engine/Source/Runtime/Renderer/Private/PostProcess/TemporalAA.h
@@ -6,6 +6,7 @@
 #include "PostProcess/PostProcessMotionBlur.h"
 #include "PostProcess/LensDistortion.h"
 #include "TemporalUpscaler.h"
+#include "MegaFiltersInterface.h"
 
 struct FTemporalAAHistory;
 struct FTranslucencyPassResources;
@@ -308,3 +309,10 @@ struct FVisualizeTemporalUpscalerInputs
 };
 
 FScreenPassTexture AddVisualizeTemporalUpscalerPass(FRDGBuilder& GraphBuilder, const FViewInfo& View, const FVisualizeTemporalUpscalerInputs& Inputs);
+
+UE::Renderer::Private::IMegaFilters::FOutputs AddMegaFiltersPass(
+	FRDGBuilder& GraphBuilder, 
+	const FViewInfo& View,
+	const TCHAR* Name,
+	const UE::Renderer::Private::IMegaFilters::FInputs& Inputs);
+
diff --git a/Engine/Source/Runtime/Renderer/Private/SceneRendering.cpp b/Engine/Source/Runtime/Renderer/Private/SceneRendering.cpp
--- a/Engine/Source/Runtime/Renderer/Private/SceneRendering.cpp
+++ b/Engine/Source/Runtime/Renderer/Private/SceneRendering.cpp
@@ -2855,6 +2855,12 @@ FSceneRenderer::FSceneRenderer(const FSceneViewFamily* InViewFamily, FHitProxyCo
 			}
 		}
 
+		if (ViewFamily.MegaFiltersInterface)
+		{
+			ViewFamily.MegaFiltersInterface = nullptr;
+			ViewFamily.SetMegaFiltersInterface(InViewFamily->MegaFiltersInterface->Fork_GameThread(ViewFamily));
+		}
+
 		if (ViewFamily.PrimarySpatialUpscalerInterface)
 		{
 			ViewFamily.PrimarySpatialUpscalerInterface = nullptr;
diff --git a/Engine/Source/Runtime/Renderer/Private/SceneRendering.h b/Engine/Source/Runtime/Renderer/Private/SceneRendering.h
--- a/Engine/Source/Runtime/Renderer/Private/SceneRendering.h
+++ b/Engine/Source/Runtime/Renderer/Private/SceneRendering.h
@@ -39,6 +39,7 @@
 #include "HairStrands/HairStrandsData.h"
 #include "Substrate/Substrate.h"
 #include "TemporalUpscaler.h"
+#include "MegaFiltersInterface.h"
 #include "GPUScene.h"
 #include "RenderCore.h"
 #include "SceneTextures.h"
@@ -1023,6 +1024,9 @@ struct FPreviousViewInfo
 	// Custom Temporal AA result of last frame, used by plugins
 	TRefCountPtr<UE::Renderer::Private::ITemporalUpscaler::IHistory> ThirdPartyTemporalUpscalerHistory;
 
+	// History of all mega filters.
+	TRefCountPtr<UE::Renderer::Private::IMegaFilters::IHistory> MegaFiltersHistory;
+
 	// Half resolution version temporal AA result of last frame
 	TRefCountPtr<IPooledRenderTarget> HalfResTemporalAAHistory;
 
diff --git a/Engine/Source/Runtime/Renderer/Private/SceneViewState.cpp b/Engine/Source/Runtime/Renderer/Private/SceneViewState.cpp
--- a/Engine/Source/Runtime/Renderer/Private/SceneViewState.cpp
+++ b/Engine/Source/Runtime/Renderer/Private/SceneViewState.cpp
@@ -221,6 +221,15 @@ uint64 FPreviousViewInfo::GetGPUSizeBytes(bool bLogSizes) const
 		}
 	}
 
+	if (MegaFiltersHistory)
+	{
+		TotalSize += MegaFiltersHistory->GetGPUSizeBytes();
+		if (bLogSizes)
+		{
+			UE_LOG(LogRenderer, Log, TEXT("LogSizes\tMegaFiltersHistory\t%s\t%llu"), MegaFiltersHistory->GetDebugName(), MegaFiltersHistory->GetGPUSizeBytes());
+		}
+	}
+
 	return TotalSize;
 }
 
