Add different stage option for Substrate's stencil marking pass.

The CVar r.Substrate.UseStandaloneClassificationStencilPass is now replaced by r.Substrate.StencilPassStage which takes 3 values:
0: During the indirect compose pass
1: Before lighting pass
2: During the first need (lighting pass or sky lighting pass)

diff --git a/Engine/Source/Runtime/Renderer/Private/IndirectLightRendering.cpp b/Engine/Source/Runtime/Renderer/Private/IndirectLightRendering.cpp
--- a/Engine/Source/Runtime/Renderer/Private/IndirectLightRendering.cpp
+++ b/Engine/Source/Runtime/Renderer/Private/IndirectLightRendering.cpp
@@ -102,6 +102,11 @@ static TAutoConsoleVariable<bool> CVarDiffuseIndirectForceCopyPass(
 	TEXT("Forces use of copy pass instead of dual source blend. (for debugging)"),
 	ECVF_RenderThreadSafe);
 
+static TAutoConsoleVariable<bool> CVarDiffuseIndirectAllowed(
+	TEXT("r.DiffuseIndirect.Allowed"), true,
+	TEXT("Allow indirect diffuse composition pass to run."),
+	ECVF_RenderThreadSafe);
+
 DECLARE_GPU_STAT_NAMED(ReflectionEnvironment, TEXT("Reflection Environment"));
 DECLARE_GPU_STAT_NAMED(RayTracingReflections, TEXT("Ray Tracing Reflections"));
 DECLARE_GPU_STAT(SkyLightDiffuse);
@@ -1016,7 +1021,8 @@ void FDeferredShadingSceneRenderer::RenderDiffuseIndirectAndAmbientOcclusion(
 
 	extern int32 GLumenVisualizeIndirectDiffuse;
 	if ((GLumenVisualizeIndirectDiffuse != 0) != bIsVisualizePass
-		|| ViewFamily.EngineShowFlags.VisualizeLightCulling)
+		|| ViewFamily.EngineShowFlags.VisualizeLightCulling 
+		|| !CVarDiffuseIndirectAllowed.GetValueOnRenderThread())
 	{
 		return;
 	}
@@ -1437,7 +1443,7 @@ void FDeferredShadingSceneRenderer::RenderDiffuseIndirectAndAmbientOcclusion(
 			PSParameters->RenderTargets[0] = FRenderTargetBinding(
 				SceneColorTexture, ERenderTargetLoadAction::ELoad);
 
-			const bool bWriteStencil = !bCompositeRegularLumenOnly && Substrate::IsSubstrateEnabled() && !Substrate::UsesStandaloneClassificationStencilPass();
+			const bool bWriteStencil = !bCompositeRegularLumenOnly && Substrate::IsSubstrateEnabled() && Substrate::GetStencilPassStage() == Substrate::EStencilPassStage::DuringIndirectLighting;
 			if (bWriteStencil)
 			{
 				PSParameters->RenderTargets.DepthStencil = FDepthStencilBinding(SceneDepthTexture, ERenderTargetLoadAction::ELoad, ERenderTargetLoadAction::ELoad, FExclusiveDepthStencil::DepthRead_StencilWrite);
@@ -2203,6 +2209,12 @@ void FDeferredShadingSceneRenderer::RenderDeferredReflectionsAndSkyLighting(
 
 		if (bRequiresApply)
 		{
+			// Update the stencil buffer, marking single/simple/complex Substrate material only if not already done.
+			if (Substrate::IsSubstrateEnabled() && Substrate::GetStencilPassStage() == Substrate::EStencilPassStage::StandaloneAdaptive)
+			{
+				Substrate::AddSubstrateStencilPass(GraphBuilder, Scene, Views, SceneTextures);
+			}
+
 			RDG_EVENT_SCOPE_STAT(GraphBuilder, ReflectionEnvironment, "ReflectionEnvironment");
 			RDG_GPU_STAT_SCOPE(GraphBuilder, ReflectionEnvironment);
 
diff --git a/Engine/Source/Runtime/Renderer/Private/LightRendering.cpp b/Engine/Source/Runtime/Renderer/Private/LightRendering.cpp
--- a/Engine/Source/Runtime/Renderer/Private/LightRendering.cpp
+++ b/Engine/Source/Runtime/Renderer/Private/LightRendering.cpp
@@ -1552,11 +1552,10 @@ void FDeferredShadingSceneRenderer::RenderLights(
 	{
 		RDG_EVENT_SCOPE(GraphBuilder, "DirectLighting");
 
-		// SUBSTRATE_TODO move right after stencil clear so that it is also common with EnvLight pass
-		if (ViewFamily.EngineShowFlags.DirectLighting && Substrate::IsSubstrateEnabled())
+		// Update the stencil buffer, marking single/simple/complex Substrate material once for all passes (lights, env., ...)
+		if (ViewFamily.EngineShowFlags.DirectLighting && Substrate::IsSubstrateEnabled() && Substrate::GetStencilPassStage() == Substrate::EStencilPassStage::Standalone)
 		{
-			// Update the stencil buffer, marking simple/complex Substrate material only once for all the following passes.
-			Substrate::AddSubstrateStencilPass(GraphBuilder, Views, SceneTextures);
+			Substrate::AddSubstrateStencilPass(GraphBuilder, Scene, Views, SceneTextures);
 		}
 
 		// The shadow scene renderer is not optional in the deferred renderer
@@ -1594,6 +1593,16 @@ void FDeferredShadingSceneRenderer::RenderLights(
 				RenderSimpleLightsStandardDeferred(GraphBuilder, SceneTextures, SortedLightSet);
 			}
 
+			const bool bHasBatchedLights = StandardDeferredStart < UnbatchedLightStart;
+			if (bHasBatchedLights)
+			{
+			
+			// Update the stencil buffer, marking single/simple/complex Substrate material only if not already done.
+			if (Substrate::IsSubstrateEnabled() && Substrate::GetStencilPassStage() == Substrate::EStencilPassStage::StandaloneAdaptive)
+			{
+				Substrate::AddSubstrateStencilPass(GraphBuilder, Scene, Views, SceneTextures);
+			}
+
 			// Draw non-shadowed non-light function lights without changing render targets between them
 			for (int32 ViewIndex = 0, ViewCount = Views.Num(); ViewIndex < ViewCount; ++ViewIndex)
 			{
@@ -1661,6 +1670,7 @@ void FDeferredShadingSceneRenderer::RenderLights(
 					}
 				}
 			}
+			} // bHasBatchedLights
 		}
 
 		{
@@ -2292,6 +2302,12 @@ void FDeferredShadingSceneRenderer::RenderLights(
 				// Render the light to the scene color buffer, conditionally using the attenuation buffer or a 1x1 white texture as input 
 				if (bDirectLighting && !bIsHairPass)
 				{
+					// Update the stencil buffer, marking single/simple/complex Substrate material only if not already done.
+					if (Substrate::IsSubstrateEnabled() && Substrate::GetStencilPassStage() == Substrate::EStencilPassStage::StandaloneAdaptive)
+					{
+						Substrate::AddSubstrateStencilPass(GraphBuilder, Scene, Views, SceneTextures);
+					}
+
 					for (int32 ViewIndex = 0, ViewCount = Views.Num(); ViewIndex < ViewCount; ++ViewIndex)
 					{
 						const FViewInfo& View = Views[ViewIndex];
@@ -3525,6 +3541,12 @@ void FDeferredShadingSceneRenderer::RenderSimpleLightsStandardDeferred(
 	SCOPE_CYCLE_COUNTER(STAT_DirectLightRenderingTime);
 	INC_DWORD_STAT_BY(STAT_NumLightsUsingStandardDeferred, NumSimpleLights);
 
+	// Update the stencil buffer, marking single/simple/complex Substrate material only if not already done.
+	if (Substrate::IsSubstrateEnabled() && Substrate::GetStencilPassStage() == Substrate::EStencilPassStage::StandaloneAdaptive)
+	{
+		Substrate::AddSubstrateStencilPass(GraphBuilder, Scene, Views, SceneTextures);
+	}
+
 	// GBufffer lighting
 	for (int32 ViewIndex = 0, NumViews = Views.Num(); ViewIndex < NumViews; ViewIndex++)
 	{
diff --git a/Engine/Source/Runtime/Renderer/Private/Substrate/Substrate.cpp b/Engine/Source/Runtime/Renderer/Private/Substrate/Substrate.cpp
--- a/Engine/Source/Runtime/Renderer/Private/Substrate/Substrate.cpp
+++ b/Engine/Source/Runtime/Renderer/Private/Substrate/Substrate.cpp
@@ -89,10 +89,10 @@ static TAutoConsoleVariable<int32> CVarSubstrateUseStochasticLightingClassificat
 	TEXT("Run Substrate tile classification within the Stochastic lighting classification. Only available when using Substrate Blendable GBuffer."),
 	ECVF_RenderThreadSafe);
 
-static TAutoConsoleVariable<int32> CVarSubstrateWriteStencilDuringIndirectLightingPass(
-	TEXT("r.Substrate.UseStandaloneClassificationStencilPass"),
+static TAutoConsoleVariable<int32> CVarSubstrateStencilPassStage(
+	TEXT("r.Substrate.StencilPassStage"),
 	1,
-	TEXT("Run Substrate tile stencil marking within a standalone pass prior to the lighting passes."),
+	TEXT("Run Substrate tile stencil writing: \n 0: During the indirect compose pass \n 1: Before lighting pass \n 2: During the first need (lighting pass or sky lighting pass)"),
 	ECVF_RenderThreadSafe);
 
 IMPLEMENT_GLOBAL_SHADER_PARAMETER_STRUCT(FSubstrateGlobalUniformParameters, "Substrate");
@@ -123,9 +123,16 @@ namespace MegaLights
 namespace Substrate
 {
 
-bool UsesStandaloneClassificationStencilPass()
+EStencilPassStage GetStencilPassStage()
 {
-	return CVarSubstrateWriteStencilDuringIndirectLightingPass.GetValueOnRenderThread() > 0;
+	const uint32 Stage = FMath::Clamp(CVarSubstrateStencilPassStage.GetValueOnRenderThread(), 0, 2);
+	switch (Stage)
+	{
+		case 0:  return EStencilPassStage::DuringIndirectLighting;
+		case 1:  return EStencilPassStage::Standalone;
+		case 2:  return EStencilPassStage::StandaloneAdaptive;
+		default: return EStencilPassStage::None;
+	}
 }
 
 uint32 GetMaxDownsampleFactor()
@@ -640,6 +647,10 @@ void InitialiseSubstrateFrameSceneData(FRDGBuilder& GraphBuilder, FSceneRenderer
 	Out.bRoughnessTracking = CVarSubstrateDebugRoughnessTracking.GetValueOnRenderThread() > 0 ? 1 : 0;
 	Out.bStochasticLighting = IsStochasticLightingActive(SceneRenderer.ShaderPlatform);
 
+	// Mark stencil with a standalone pass 
+	const EStencilPassStage StencilPassStage = GetStencilPassStage();
+	Out.bWriteStencil = StencilPassStage == EStencilPassStage::Standalone || StencilPassStage == EStencilPassStage::StandaloneAdaptive;
+
 	if (bNeedsMaterialBuffer)
 	{
 		// SUBSTRATE_TODO allocate a slice for StoringDebugSubstrate only if SUBSTRATE_ADVANCED_DEBUG_ENABLED is enabled 
@@ -1453,14 +1464,18 @@ void AddSubstrateInternalClassificationTilePass(
 
 void AddSubstrateStencilPass(
 	FRDGBuilder& GraphBuilder,
+	FScene* Scene,
 	const TArray<FViewInfo>& Views,
 	const FMinimalSceneTextures& SceneTextures)
 {
-	if (!UsesStandaloneClassificationStencilPass())
+	if (!Scene->SubstrateSceneData.bWriteStencil)
 	{
 		return;
 	}
 
+	// Marke Stencil as written
+	Scene->SubstrateSceneData.bWriteStencil = false;
+
 	for (int32 i = 0; i < Views.Num(); ++i)
 	{
 		RDG_EVENT_SCOPE_CONDITIONAL(GraphBuilder, Views.Num() > 1, "View%d", i);
diff --git a/Engine/Source/Runtime/Renderer/Private/Substrate/Substrate.h b/Engine/Source/Runtime/Renderer/Private/Substrate/Substrate.h
--- a/Engine/Source/Runtime/Renderer/Private/Substrate/Substrate.h
+++ b/Engine/Source/Runtime/Renderer/Private/Substrate/Substrate.h
@@ -100,6 +100,9 @@ struct FSubstrateSceneData
 	uint8 UsesTileTypeMask = 0;
 	bool bUsesAnisotropy = false;
 
+	// Track if stencil needs to be written
+	bool bWriteStencil = false;
+
 	// Current max BytesPerPixel / ClosurePerPixel
 	uint32 EffectiveMaxBytesPerPixel = 0;
 	uint32 EffectiveMaxClosurePerPixel = 0;
@@ -210,6 +213,14 @@ struct FSubstrateViewDebugData
 namespace Substrate
 {
 
+enum class EStencilPassStage : uint8
+{
+	None,
+	Standalone,
+	StandaloneAdaptive,
+	DuringIndirectLighting
+};
+
 // Bits for filling/clearing stencil during shadow/lighting passes. 
 // Those are designed to not conflicts with other stencil shadow bits, when generating shadow mask. 
 // Withouth these 'trailing' bits, the incr./decr. operation would change/corrupt the 'Substrate' bits.
@@ -225,8 +236,8 @@ bool GetSubstrateUsesComplexSpecialPath(const FViewInfo& View);
 bool GetSubstrateUsesAnisotropy(const FViewInfo& View);
 bool UsesSubstrateMaterialBuffer(EShaderPlatform In);
 bool UsesStochasticLightingClassification(EShaderPlatform InPlatform);
-bool UsesStandaloneClassificationStencilPass();
 bool GetSubstrateUsesTileType(const FViewInfo& View, ESubstrateTileType TileType);
+EStencilPassStage GetStencilPassStage();
 
 void InitialiseSubstrateFrameSceneData(FRDGBuilder& GraphBuilder, FSceneRenderer& SceneRenderer);
 
@@ -244,7 +255,7 @@ TRDGUniformBufferRef<FSubstratePublicGlobalUniformParameters> CreatePublicGlobal
 void AddSubstrateMaterialClassificationPass(FRDGBuilder& GraphBuilder, const FMinimalSceneTextures& SceneTextures, const FDBufferTextures& DBufferTextures, const TArray<FViewInfo>& Views);
 void AddSubstrateDBufferPass(FRDGBuilder& GraphBuilder, const FSceneTextures& SceneTextures, const FDBufferTextures& DBufferTextures, const TArray<FViewInfo>& Views);
 
-void AddSubstrateStencilPass(FRDGBuilder& GraphBuilder, const TArray<FViewInfo>& Views, const FMinimalSceneTextures& SceneTextures);
+void AddSubstrateStencilPass(FRDGBuilder& GraphBuilder, FScene* Scene, const TArray<FViewInfo>& Views, const FMinimalSceneTextures& SceneTextures);
 void AddSubstrateSampleMaterialPass(FRDGBuilder& GraphBuilder, const FScene* Scene, const FMinimalSceneTextures& SceneTextures, const TArray<FViewInfo>& Views);
 void AddSubstrateMaterialClassificationIndirectArgsPass(FRDGBuilder& GraphBuilder, const FViewInfo& View, ERDGPassFlags ComputePassFlags);
 
