Fixed crash due to GDME render thread task not being executed. Only the last task in the chain was being waited on and flushed. The fix is to pump the entire local queue until idle instead.

diff --git a/Engine/Source/Runtime/Renderer/Private/SceneVisibility.cpp b/Engine/Source/Runtime/Renderer/Private/SceneVisibility.cpp
--- a/Engine/Source/Runtime/Renderer/Private/SceneVisibility.cpp
+++ b/Engine/Source/Runtime/Renderer/Private/SceneVisibility.cpp
@@ -4594,7 +4594,7 @@ void FVisibilityTaskData::GatherDynamicMeshElements(FDynamicPrimitiveIndexList&&
 
 	if (!RenderThreadDynamicPrimitiveIndexList.IsEmpty())
 	{
-		Tasks.DynamicMeshElementsRenderThread = DynamicMeshElements.ContextContainer.LaunchRenderThreadTask(MoveTemp(RenderThreadDynamicPrimitiveIndexList));
+		DynamicMeshElements.ContextContainer.LaunchRenderThreadTask(MoveTemp(RenderThreadDynamicPrimitiveIndexList));
 	}
 
 	if (!DynamicPrimitiveIndexList.IsEmpty())
@@ -4858,10 +4858,7 @@ void FVisibilityTaskData::ProcessRenderThreadTasks()
 	}
 
 	// Now process all gather dynamic mesh element tasks that were queued up to run on the render thread.
-	if (Tasks.DynamicMeshElementsRenderThread)
-	{
-		Tasks.DynamicMeshElementsRenderThread->Wait(ENamedThreads::GetRenderThread_Local());
-	}
+	FTaskGraphInterface::Get().ProcessThreadUntilIdle(ENamedThreads::GetRenderThread_Local());
 
 	Tasks.LightVisibility.Wait();
 	Tasks.FinalizeRelevance.Wait();
diff --git a/Engine/Source/Runtime/Renderer/Private/SceneVisibilityPrivate.h b/Engine/Source/Runtime/Renderer/Private/SceneVisibilityPrivate.h
--- a/Engine/Source/Runtime/Renderer/Private/SceneVisibilityPrivate.h
+++ b/Engine/Source/Runtime/Renderer/Private/SceneVisibilityPrivate.h
@@ -591,7 +591,6 @@ private:
 		FGraphEventRef ComputeRelevanceLegacyTask;
 
 		FGraphEventRef DynamicMeshElementsPipe;
-		FGraphEventRef DynamicMeshElementsRenderThread;
 
 		UE::Tasks::FTaskEvent LightVisibility{ UE_SOURCE_LOCATION };
 		UE::Tasks::FTaskEvent BeginInitVisibility{ UE_SOURCE_LOCATION };
