[DOF] Added r.DOF.Gather.MaxMipMapCount as a workaround for artifacts in the reduce pass.
The new cvar allows you to limit the number of mipmap levels that are used in the gather stage of DoF. Mipmapping is used there because when the circle-of-confusion radius gets large, the blur kernel becomes large but also sparse. Prefiltering (mipmapping) counteracts some of the noise that this sparse kernel may introduce. Due to an active issue with this prefiltering, foreground pixels may leak into the background, causing artifacts on edges. By setting the cvar to 1, you disable prefiltering, which removes the artifacts but potentially increases noise in highlights in very out-of-focus areas.

diff --git a/Engine/Shaders/Private/DiaphragmDOF/DOFReduce.usf b/Engine/Shaders/Private/DiaphragmDOF/DOFReduce.usf
--- a/Engine/Shaders/Private/DiaphragmDOF/DOFReduce.usf
+++ b/Engine/Shaders/Private/DiaphragmDOF/DOFReduce.usf
@@ -655,13 +655,13 @@ void ReduceCS(
 	#if CONFIG_WAVE_BROADCAST_REDUCTION
 	{
 		UNROLL
-		for (uint i = 0; i < (DIM_REDUCE_MIP_COUNT - 1); i++)
+		for (int MipLevel = 1; MipLevel < DIM_REDUCE_MIP_COUNT; MipLevel++)
 		{
 			// Multiplier that needs to be applied on Coc computation to be resolution independent.
-			DownsampleParams.CocRadiusMultiplier = 0.5 / float(1u << (i));
+			DownsampleParams.CocRadiusMultiplier = 0.5 / float(1u << (MipLevel - 1));
 
 			// Scale of the reduction within the 8x8 tile.
-			const uint ReductionScale = 1u << i;
+			const uint ReductionScale = 1u << (MipLevel - 1);
 			
 			// Chooses output CoC across lanes.
 			float OutputCocRadius = DownsampleCoc_Wave8x8(CocRadius, ReductionScale);
@@ -685,9 +685,9 @@ void ReduceCS(
 			// Output mip map.
 			OutputPixelId = OutputPixelId >> 1;
 			BRANCH
-			if (((DispatchThreadId.x | DispatchThreadId.y) & ((2u << i) - 1)) == 0)
+			if (((DispatchThreadId.x | DispatchThreadId.y) & ((1u << MipLevel) - 1)) == 0)
 			{
-				OutputMipLevel(/* MipLevel = */ i + 1, OutputPixelId, GatheredColor, CocRadius);
+				OutputMipLevel(MipLevel, OutputPixelId, GatheredColor, CocRadius);
 			}
 		}
 	}
@@ -707,10 +707,10 @@ void ReduceCS(
 
 		// Output mip > 0 of the gathering pass.
 		UNROLL
-		for (uint i = 0; i < (DIM_REDUCE_MIP_COUNT - 1); i++)
+		for (int MipLevel = 1; MipLevel < DIM_REDUCE_MIP_COUNT; MipLevel++)
 		{
 			// Multiplier that needs to be applied on Coc computation to be resolution independent.
-			const uint TileSize = THREADGROUP_TILE_SIZE / (1u << (i + 1));
+			const uint TileSize = THREADGROUP_TILE_SIZE / (1u << MipLevel);
 			const uint ReduceBankSize = TileSize * TileSize;
 		
 			// GroupSharedArray has been written before.
@@ -718,7 +718,7 @@ void ReduceCS(
 		
 			if (GroupThreadIndex < ReduceBankSize)
 			{
-				DownsampleParams.CocRadiusMultiplier = 0.5 / float(1u << (i));
+				DownsampleParams.CocRadiusMultiplier = 0.5 / float(1u << (MipLevel - 1));
 
 				// Reduce.
 				ReduceOperator(DownsampleParams, GroupThreadIndex, ReduceBankSize);
@@ -734,7 +734,7 @@ void ReduceCS(
 			
 				// Output mip map.
 				OutputPixelId = OutputPixelId >> 1;
-				OutputMipLevel(/* MipLevel = */ i + 1, OutputPixelId, Color, OutCocRadius);
+				OutputMipLevel(MipLevel, OutputPixelId, Color, OutCocRadius);
 			}
 		}
 	}
diff --git a/Engine/Source/Runtime/Renderer/Private/PostProcess/DiaphragmDOF.cpp b/Engine/Source/Runtime/Renderer/Private/PostProcess/DiaphragmDOF.cpp
--- a/Engine/Source/Runtime/Renderer/Private/PostProcess/DiaphragmDOF.cpp
+++ b/Engine/Source/Runtime/Renderer/Private/PostProcess/DiaphragmDOF.cpp
@@ -78,6 +78,12 @@ TAutoConsoleVariable<int32> CVarRingCount(
 	TEXT("Number of rings for gathering kernels [[3; 5]]. Default to 5.\n"),
 	ECVF_Scalability | ECVF_RenderThreadSafe);
 
+TAutoConsoleVariable<int32> CVarGatherMaxMipMapCount(
+	TEXT("r.DOF.Gather.MaxMipMapCount"),
+	4,
+	TEXT("Set a maximum amount of mipmaps to calculate in the reduction phase.\n")
+	TEXT("These are then used in the gather step to reduce noise in areas with a large but sparse blurring kernel."),
+	ECVF_Scalability | ECVF_RenderThreadSafe);
 
 TAutoConsoleVariable<int32> CVarHybridScatterForegroundMode(
 	TEXT("r.DOF.Scatter.ForegroundCompositing"),
@@ -875,7 +881,7 @@ class FDiaphragmDOFReduceCS : public FDiaphragmDOFShader
 	DECLARE_GLOBAL_SHADER(FDiaphragmDOFReduceCS);
 	SHADER_USE_PARAMETER_STRUCT(FDiaphragmDOFReduceCS, FDiaphragmDOFShader);
 
-	class FReduceMipCount : SHADER_PERMUTATION_RANGE_INT("DIM_REDUCE_MIP_COUNT", 2, 3);
+	class FReduceMipCount : SHADER_PERMUTATION_RANGE_INT("DIM_REDUCE_MIP_COUNT", 0, kMaxMipLevelCount);
 	class FHybridScatterForeground : SHADER_PERMUTATION_BOOL("DIM_HYBRID_SCATTER_FGD");
 	class FHybridScatterBackground : SHADER_PERMUTATION_BOOL("DIM_HYBRID_SCATTER_BGD");
 	class FDisableOutputMip0 : SHADER_PERMUTATION_BOOL("DIM_DISABLE_OUTPUT_MIP0");
@@ -896,11 +902,21 @@ class FDiaphragmDOFReduceCS : public FDiaphragmDOFShader
 		return bSupportAlphaChannel ? 2 : kMaxMipLevelCount;
 	}
 
+	/** Returns the nearest number of mip levels the reduce pass is able to output */
+	static int32 GetNearestMipLevelReductionCount(int ReductionCount, bool bSupportAlphaChannel)
+	{
+		int32 MaxCount = GetMaxReductionMipLevelCount(bSupportAlphaChannel);
+		ReductionCount = FMath::Min(ReductionCount, MaxCount);
+
+		return ReductionCount;
+	}
+
 	static bool ShouldCompilePermutation(const FGlobalShaderPermutationParameters& Parameters)
 	{
 		FPermutationDomain PermutationVector(Parameters.PermutationId);
 
 		const int32 MaxMipLevelCount = GetMaxReductionMipLevelCount(PermutationVector.Get<FDiaphragmDOFShader::FAlphaChannelDim>());
+		bool bAlphaChannel = PermutationVector.Get<FDiaphragmDOFShader::FAlphaChannelDim>();
 
 		if (PermutationVector.Get<FDisableOutputMip0>())
 		{
@@ -909,20 +925,13 @@ class FDiaphragmDOFReduceCS : public FDiaphragmDOFShader
 			{
 				return false;
 			}
-
-			// Do not output mip level that are more than supported.
-			if (PermutationVector.Get<FReduceMipCount>() > (MaxMipLevelCount + 1))
-			{
-				return false;
-			}
 		}
-		else
+		
+		// Do not output mip level that are more than supported.
+		int32 ReduceMipCount = PermutationVector.Get<FReduceMipCount>();
+		if (GetNearestMipLevelReductionCount(ReduceMipCount, bAlphaChannel) != ReduceMipCount)
 		{
-			// Do not output mip level that are more than supported.
-			if (PermutationVector.Get<FReduceMipCount>() > MaxMipLevelCount)
-			{
-				return false;
-			}
+			return false;
 		}
 
 		// Do not compile storing Coc independently of RGB if not supported.
@@ -2069,10 +2078,8 @@ bool DiaphragmDOF::AddPasses(
 		FIntPoint SrcSize = HalfResGatherInputDescs.SceneColor.Extent;
 		
 		// Compute the number of mip level required by the gathering pass.
-		const int32 MipLevelCount = FMath::Max(FMath::CeilToInt(FMath::Log2(MaxBluringRadius * 0.5 / HalfResRingCount)) + (bUseLowAccumulatorQuality ? 1 : 0), 2);
-
-		// Maximum number of mip level that can be done.
-		const int32 MaxReductionMipLevelCount = FDiaphragmDOFReduceCS::GetMaxReductionMipLevelCount(IsPostProcessingWithAlphaChannelSupported());
+		const int32 MipLevelsRequested = FMath::CeilToInt(FMath::Log2(MaxBluringRadius * 0.5 / HalfResRingCount)) + (bUseLowAccumulatorQuality ? 1 : 0);
+		const int32 MipLevelCount = FMath::Clamp(MipLevelsRequested, 1, CVarGatherMaxMipMapCount.GetValueOnRenderThread());
 
 		// Maximum number of scattering group per draw instance.
 		// TODO: depends.
@@ -2171,7 +2178,7 @@ bool DiaphragmDOF::AddPasses(
 
 		// Adds the first reduce pass.
 		{
-			int32 ProcessingMipLevelCount = FMath::Min(MipLevelCount, MaxReductionMipLevelCount);
+			int32 ProcessingMipLevelCount = FDiaphragmDOFReduceCS::GetNearestMipLevelReductionCount(MipLevelCount, IsPostProcessingWithAlphaChannelSupported());
 
 			FIntPoint PassViewSize = PreprocessViewSize;
 
