[DOF] Follow-up CL to 47532889, fixing issue where GetMaxReductionMipLevelCount did not account for FDisableOutputMip0. If mip 0 is disabled, an extra UAV is available.

diff --git a/Engine/Source/Runtime/Renderer/Private/PostProcess/DiaphragmDOF.cpp b/Engine/Source/Runtime/Renderer/Private/PostProcess/DiaphragmDOF.cpp
--- a/Engine/Source/Runtime/Renderer/Private/PostProcess/DiaphragmDOF.cpp
+++ b/Engine/Source/Runtime/Renderer/Private/PostProcess/DiaphragmDOF.cpp
@@ -896,16 +896,17 @@ class FDiaphragmDOFReduceCS : public FDiaphragmDOFShader
 
 
 	/** Returns the number of mip level the reduce pass is able to output. */
-	static int32 GetMaxReductionMipLevelCount(bool bSupportAlphaChannel)
+	static int32 GetMaxReductionMipLevelCount(bool bSupportAlphaChannel, bool bMip0Disabled)
 	{
 		// Can only have 8 UAVs, but need to output 3x UAV for hybird scatter + 2 UAV per mips for RGBA + SeparateCoc.
-		return bSupportAlphaChannel ? 2 : kMaxMipLevelCount;
+		int32 MaxLevelsWithAlphaChannel = bMip0Disabled ? 3 : 2;
+		return bSupportAlphaChannel ? MaxLevelsWithAlphaChannel : kMaxMipLevelCount;
 	}
 
 	/** Returns the nearest number of mip levels the reduce pass is able to output */
-	static int32 GetNearestMipLevelReductionCount(int ReductionCount, bool bSupportAlphaChannel)
+	static int32 GetNearestMipLevelReductionCount(int ReductionCount, bool bSupportAlphaChannel, bool bMip0Disabled)
 	{
-		int32 MaxCount = GetMaxReductionMipLevelCount(bSupportAlphaChannel);
+		int32 MaxCount = GetMaxReductionMipLevelCount(bSupportAlphaChannel, bMip0Disabled);
 		ReductionCount = FMath::Min(ReductionCount, MaxCount);
 
 		return ReductionCount;
@@ -915,8 +916,8 @@ class FDiaphragmDOFReduceCS : public FDiaphragmDOFShader
 	{
 		FPermutationDomain PermutationVector(Parameters.PermutationId);
 
-		const int32 MaxMipLevelCount = GetMaxReductionMipLevelCount(PermutationVector.Get<FDiaphragmDOFShader::FAlphaChannelDim>());
 		bool bAlphaChannel = PermutationVector.Get<FDiaphragmDOFShader::FAlphaChannelDim>();
+		bool bDisableOutputMip0 = PermutationVector.Get<FDisableOutputMip0>();
 
 		if (PermutationVector.Get<FDisableOutputMip0>())
 		{
@@ -929,7 +930,7 @@ class FDiaphragmDOFReduceCS : public FDiaphragmDOFShader
 		
 		// Do not output mip level that are more than supported.
 		int32 ReduceMipCount = PermutationVector.Get<FReduceMipCount>();
-		if (GetNearestMipLevelReductionCount(ReduceMipCount, bAlphaChannel) != ReduceMipCount)
+		if (GetNearestMipLevelReductionCount(ReduceMipCount, bAlphaChannel, bDisableOutputMip0) != ReduceMipCount)
 		{
 			return false;
 		}
@@ -2178,7 +2179,8 @@ bool DiaphragmDOF::AddPasses(
 
 		// Adds the first reduce pass.
 		{
-			int32 ProcessingMipLevelCount = FDiaphragmDOFReduceCS::GetNearestMipLevelReductionCount(MipLevelCount, IsPostProcessingWithAlphaChannelSupported());
+			const bool bMip0Disabled = false;
+			int32 ProcessingMipLevelCount = FDiaphragmDOFReduceCS::GetNearestMipLevelReductionCount(MipLevelCount, IsPostProcessingWithAlphaChannelSupported(), bMip0Disabled);
 
 			FIntPoint PassViewSize = PreprocessViewSize;
 
