Translucent Materials no long compile the non-skylight version of base pass pixel shader to reduce shader permutations and instead always use the skylight version w/ a dynamic branch.

- This leverages the existing branch in `GetSkyLighting` (`if (UseBasePassSkylight > 0)`) to determine if we should render the skylight.  When Lumen is enabled we don't even take this code path.
- For opaque materials we already just pick the non-skylight version (and skip compiling the skylight version).
- In `FBasePassMeshProcessor::Process` if we are translucent we need to sure to always render with the skylight version of the shader.
- In CreateTranslucentBasePassUniformBuffer we set the uniform (UseBasePassSkylight) whether or not the scene has/wants a skylight.  Like before if Lumen is enabled we never take this branch.

[RN] minor | Rendering - Materials and Shaders

diff --git a/Engine/Source/Runtime/Renderer/Private/BasePassRendering.cpp b/Engine/Source/Runtime/Renderer/Private/BasePassRendering.cpp
--- a/Engine/Source/Runtime/Renderer/Private/BasePassRendering.cpp
+++ b/Engine/Source/Runtime/Renderer/Private/BasePassRendering.cpp
@@ -1915,7 +1915,13 @@ bool FBasePassMeshProcessor::Process(
 {
 	const FVertexFactory* VertexFactory = MeshBatch.VertexFactory;
 
-	const bool bRenderSkylight = Scene && Scene->ShouldRenderSkylightInBasePass(bIsTranslucent) && ShadingModels.IsLit();
+	bool bRenderSkylight = Scene && Scene->ShouldRenderSkylightInBasePass(bIsTranslucent) && ShadingModels.IsLit();
+	if (bIsTranslucent)
+	{
+		// Force translucent to always run the skylight version, but then to always do the dynamic branch.
+		// This needs to match the logic in TBasePassPS::ShouldCompilePermutation
+		bRenderSkylight = true;
+	}
 
 	TMeshProcessorShaders<
 		TBasePassVertexShaderPolicyParamType<LightMapPolicyType>,
diff --git a/Engine/Source/Runtime/Renderer/Private/BasePassRendering.h b/Engine/Source/Runtime/Renderer/Private/BasePassRendering.h
--- a/Engine/Source/Runtime/Renderer/Private/BasePassRendering.h
+++ b/Engine/Source/Runtime/Renderer/Private/BasePassRendering.h
@@ -669,6 +669,13 @@ public:
 			}
 		}
 
+		// Don't compile the non-skylight version for translucent materials.  Translucent materials always compile the skylight version and
+		// will dynamically branch on whether a skylight is enabled.
+		if (bTranslucent && (bEnableSkyLight == false))
+		{
+			return false;
+		}
+
 		const bool bCacheShaders = !bEnableSkyLight
 			//translucent materials need to compile skylight support to support MOVABLE skylights also.
 			|| bTranslucent
diff --git a/Engine/Source/Runtime/Renderer/Private/TranslucentRendering.cpp b/Engine/Source/Runtime/Renderer/Private/TranslucentRendering.cpp
--- a/Engine/Source/Runtime/Renderer/Private/TranslucentRendering.cpp
+++ b/Engine/Source/Runtime/Renderer/Private/TranslucentRendering.cpp
@@ -1141,7 +1141,9 @@ TRDGUniformBufferRef<FTranslucentBasePassUniformParameters> CreateTranslucentBas
 	const bool bLumenGIHandlingSkylight = bLumenGIEnabled
 		&& BasePassParameters.LumenParameters.TranslucencyGIGridSize.Z > 0;
 
-	BasePassParameters.Shared.UseBasePassSkylight = bLumenGIHandlingSkylight ? 0 : 1;
+	// Dynamically pick if we should render skylight in translucent base pass if skylights are enabled.
+	const bool bShouldRenderSkylightInBasePass = Scene ? Scene->ShouldRenderSkylightInBasePass(true) : true;
+	BasePassParameters.Shared.UseBasePassSkylight = bLumenGIHandlingSkylight ? 0 : bShouldRenderSkylightInBasePass;
 
 	BasePassParameters.SceneColorCopyTexture = SystemTextures.Black;
 	BasePassParameters.SceneColorCopySampler = TStaticSamplerState<SF_Bilinear, AM_Clamp, AM_Clamp, AM_Clamp>::GetRHI();
