[Lit Translucency] On-screen stats showing number of visible and marked voxels.

diff --git a/Engine/Shaders/Private/TranslucentLightingShaders.usf b/Engine/Shaders/Private/TranslucentLightingShaders.usf
--- a/Engine/Shaders/Private/TranslucentLightingShaders.usf
+++ b/Engine/Shaders/Private/TranslucentLightingShaders.usf
@@ -282,41 +282,72 @@ void InjectAmbientCubemapMainPS(
 #endif // FEATURE_LEVEL
 
 
-#ifdef DebugTranslucencyLightingVolumeCS
+#ifdef FrustumCullVoxelsCS
 
-#include "/Engine/Private/ShaderPrint.ush"
+#define MAX_CAMERA_PLANES (5)
+
+RWStructuredBuffer<uint> RWVoxelAllocator0;
+RWStructuredBuffer<uint> RWVoxelAllocator1;
+
+float4 CameraFrustumPlanes[MAX_CAMERA_PLANES];
 
-void PrintTranslucencyVolume(inout FShaderPrintContext Context, float3 Min, float3 Max, float4 Color)
+uint NumVoxels;
+
+[numthreads(THREADGROUP_SIZE, THREADGROUP_SIZE, THREADGROUP_SIZE)]
+void FrustumCullVoxelsCS(
+	uint3 GroupId : SV_GroupID,
+	uint3 DispatchThreadId : SV_DispatchThreadID,
+	uint3 GroupThreadId : SV_GroupThreadID)
 {
-	float3 C000 = float3(Min.x, Min.y, Min.z);
-	float3 C001 = float3(Min.x, Min.y, Max.z);
-	float3 C010 = float3(Min.x, Max.y, Min.z);
-	float3 C011 = float3(Min.x, Max.y, Max.z);
-	float3 C100 = float3(Max.x, Min.y, Min.z);
-	float3 C101 = float3(Max.x, Min.y, Max.z);
-	float3 C110 = float3(Max.x, Max.y, Min.z);
-	float3 C111 = float3(Max.x, Max.y, Max.z);
-
-	AddLineTWS(Context, C000, C001, Color);
-	AddLineTWS(Context, C001, C011, Color);
-	AddLineTWS(Context, C011, C010, Color);
-	AddLineTWS(Context, C010, C000, Color);
-
-	AddLineTWS(Context, C100, C101, Color);
-	AddLineTWS(Context, C101, C111, Color);
-	AddLineTWS(Context, C111, C110, Color);
-	AddLineTWS(Context, C110, C100, Color);
-
-	AddLineTWS(Context, C001, C101, Color);
-	AddLineTWS(Context, C100, C000, Color);
-	AddLineTWS(Context, C011, C111, Color);
-	AddLineTWS(Context, C110, C010, Color);
+	const float3 VoxelCenter0 = GetTranslatedWorldPositionFromVolumeCoord(DispatchThreadId, 0);
+	const float3 VoxelExtent0 = View.TranslucencyLightingVolumeInvSize[0].w / 2;
+
+	bool bNearFrustum0 = true;
+	for (uint PlaneIndex = 0; PlaneIndex < MAX_CAMERA_PLANES; ++PlaneIndex)
+	{
+		float DistanceFromVoxelCenterToPlane = dot(CameraFrustumPlanes[PlaneIndex], float4(VoxelCenter0, 1.0f));
+		float Radius = dot(VoxelExtent0, abs(CameraFrustumPlanes[PlaneIndex].xyz));
+		bNearFrustum0 = bNearFrustum0 && (DistanceFromVoxelCenterToPlane <= Radius);
+	}
+
+	const float3 VoxelCenter1 = GetTranslatedWorldPositionFromVolumeCoord(DispatchThreadId, 1);
+	const float3 VoxelExtent1 = View.TranslucencyLightingVolumeInvSize[1].w / 2;
+
+	bool bNearFrustum1 = true;
+	for (uint PlaneIndex = 0; PlaneIndex < MAX_CAMERA_PLANES; ++PlaneIndex)
+	{
+		float DistanceFromVoxelCenterToPlane = dot(CameraFrustumPlanes[PlaneIndex], float4(VoxelCenter1, 1.0f));
+		float Radius = dot(VoxelExtent1, abs(CameraFrustumPlanes[PlaneIndex].xyz));
+		bNearFrustum1 = bNearFrustum1 && (DistanceFromVoxelCenterToPlane <= Radius);
+	}
+
+	if (bNearFrustum0)
+	{
+		InterlockedAdd(RWVoxelAllocator0[0], 1);
+	}
+
+	if (bNearFrustum1)
+	{
+		InterlockedAdd(RWVoxelAllocator1[0], 1);
+	}
 }
+#endif // FrustumCullVoxelsCS
 
-Texture3D<uint> InnerVolumeMarkTexture;
-Texture3D<uint> OuterVolumeMarkTexture;
 
-StructuredBuffer<uint> VoxelAllocator;
+#ifdef DebugTranslucencyLightingVolumeCS
+
+#include "/Engine/Private/ShaderPrint.ush"
+
+Texture3D<uint> VolumeMarkTexture0;
+Texture3D<uint> VolumeMarkTexture1;
+
+StructuredBuffer<uint> VoxelAllocator0;
+StructuredBuffer<uint> VoxelAllocator1;
+
+StructuredBuffer<uint> FrustumCullVoxelAllocator0;
+StructuredBuffer<uint> FrustumCullVoxelAllocator1;
+
+uint NumVoxels;
 
 [numthreads(THREADGROUP_SIZE, THREADGROUP_SIZE, THREADGROUP_SIZE)]
 void DebugTranslucencyLightingVolumeCS(
@@ -324,10 +355,10 @@ void DebugTranslucencyLightingVolumeCS(
 	uint3 DispatchThreadId : SV_DispatchThreadID,
 	uint3 GroupThreadId : SV_GroupThreadID)
 {
+	FShaderPrintContext Context = InitShaderPrintContext(true, uint2(50, 50));
+
 	if (all(DispatchThreadId == 0))
 	{
-		FShaderPrintContext Context = InitShaderPrintContext(true, uint2(50, 50));
-
 		const float3 VolMin0 = View.TranslucencyLightingVolumeMin[0].xyz;
 		const float3 VolMin1 = View.TranslucencyLightingVolumeMin[1].xyz;
 		const float3 VolSize0 = 1.0f / View.TranslucencyLightingVolumeInvSize[0].xyz;
@@ -335,45 +366,88 @@ void DebugTranslucencyLightingVolumeCS(
 		const float3 VolMax0 = VolMin0 + VolSize0;
 		const float3 VolMax1 = VolMin1 + VolSize1;
 
-		PrintTranslucencyVolume(Context, VolMin0, VolMax0, float4(0, 1, 0, 1));
-		PrintTranslucencyVolume(Context, VolMin1, VolMax1, float4(1, 1, 0, 1));
+		AddAABBTWS(Context, VolMin0, VolMax0, float4(0, 1, 0, 1));
+		AddAABBTWS(Context, VolMin1, VolMax1, float4(1, 1, 0, 1));
+
+		Print(Context, TEXT("Translucency Lighting Volume"), FontWhite);
+		Newline(Context);
+
+		Print(Context, TEXT("- Cascade 0:"), FontWhite);
+		Newline(Context);
+
+		Print(Context, TEXT("    - Size = "), FontWhite);
+		Print(Context, VolSize0, FontYellow);
+		Newline(Context);
+
+		const uint VisibleVoxels0 = FrustumCullVoxelAllocator0[0];
+
+		Print(Context, TEXT("    - # Visible Voxels = "), FontWhite);
+		Print(Context, VisibleVoxels0, FontYellow);
+		Print(Context, TEXT(" / "), FontYellow);
+		Print(Context, NumVoxels, FontYellow);
+		Print(Context, TEXT(" ("), FontYellow);
+		Print(Context, VisibleVoxels0 * 100.0f / NumVoxels, FontYellow, 3, 1);
+		Print(Context, TEXT("%)"), FontYellow);
+		Newline(Context);
+
+		const uint MarkedVoxels0 = VoxelAllocator0[0];
 
-		Print(Context, TEXT("Translucency Light Volume0 Size = "), FontWhite);
-		Print(Context, VolSize0.x, FontYellow, 10, 1);
-		Print(Context, TEXT(", "), FontWhite);
-		Print(Context, VolSize0.y, FontYellow, 10, 1);
-		Print(Context, TEXT(", "), FontWhite);
-		Print(Context, VolSize0.y, FontYellow, 10, 1);
+		Print(Context, TEXT("    - # Marked Voxels = "), FontWhite);
+		Print(Context, MarkedVoxels0, FontYellow);
+		Print(Context, TEXT(" / "), FontYellow);
+		Print(Context, NumVoxels, FontYellow);
+		Print(Context, TEXT(" ("), FontYellow);
+		Print(Context, MarkedVoxels0 * 100.0f / NumVoxels, FontYellow, 3, 1);
+		Print(Context, TEXT("%)"), FontYellow);
 		Newline(Context);
-		Print(Context, TEXT("Translucency Light Volume1 Size = "), FontWhite);
-		Print(Context, VolSize1.x, FontYellow, 10, 1);
-		Print(Context, TEXT(", "), FontWhite);
-		Print(Context, VolSize1.y, FontYellow, 10, 1);
-		Print(Context, TEXT(", "), FontWhite);
-		Print(Context, VolSize1.y, FontYellow, 10, 1);
+
+		Print(Context, TEXT("- Cascade 1:"), FontWhite);
+		Newline(Context);
+
+		Print(Context, TEXT("    - Size = "), FontWhite);
+		Print(Context, VolSize1, FontYellow);
+		Newline(Context);
+
+		const uint VisibleVoxels1 = FrustumCullVoxelAllocator1[0];
+
+		Print(Context, TEXT("    - # Visible Voxels = "), FontWhite);
+		Print(Context, VisibleVoxels1, FontYellow);
+		Print(Context, TEXT(" / "), FontYellow);
+		Print(Context, NumVoxels, FontYellow);
+		Print(Context, TEXT(" ("), FontYellow);
+		Print(Context, VisibleVoxels1 * 100.0f / NumVoxels, FontYellow, 3, 1);
+		Print(Context, TEXT("%)"), FontYellow);
+		Newline(Context);
+
+		const uint MarkedVoxels1 = VoxelAllocator1[0];
+
+		Print(Context, TEXT("    - # Marked Voxels = "), FontWhite);
+		Print(Context, MarkedVoxels1, FontYellow);
+		Print(Context, TEXT(" / "), FontYellow);
+		Print(Context, NumVoxels, FontYellow);
+		Print(Context, TEXT(" ("), FontYellow);
+		Print(Context, MarkedVoxels1 * 100.0f / NumVoxels, FontYellow, 3, 1);
+		Print(Context, TEXT("%)"), FontYellow);
+
+		AddTextBackground(Context, float4(0, 0, 0, 0.25f));
 	}
 
 	// draw marked voxels
 
-	if (InnerVolumeMarkTexture[DispatchThreadId] > 0)
+	if (VolumeMarkTexture0[DispatchThreadId] > 0)
 	{
-		FShaderPrintContext Context = InitShaderPrintContext(true, uint2(50, 50));
-
 		const float3 TranslatedWorldPosition = GetTranslatedWorldPositionFromVolumeCoord(DispatchThreadId, 0);
 		const float3 VoxelExtent = View.TranslucencyLightingVolumeInvSize[0].w / 2;
 
-		PrintTranslucencyVolume(Context, TranslatedWorldPosition - VoxelExtent, TranslatedWorldPosition + VoxelExtent, float4(0, 1, 0, 1));
+		AddAABBTWS(Context, TranslatedWorldPosition - VoxelExtent, TranslatedWorldPosition + VoxelExtent, float4(0, 1, 0, 1));
 	}
 
-	if (OuterVolumeMarkTexture[DispatchThreadId] > 0)
+	if (VolumeMarkTexture1[DispatchThreadId] > 0)
 	{
-		FShaderPrintContext Context = InitShaderPrintContext(true, uint2(50, 50));
-
 		const float3 TranslatedWorldPosition = GetTranslatedWorldPositionFromVolumeCoord(DispatchThreadId, 1);
 		const float3 VoxelExtent = View.TranslucencyLightingVolumeInvSize[1].w / 2;
 
-		PrintTranslucencyVolume(Context, TranslatedWorldPosition - VoxelExtent, TranslatedWorldPosition + VoxelExtent, float4(1, 1, 0, 1));
+		AddAABBTWS(Context, TranslatedWorldPosition - VoxelExtent, TranslatedWorldPosition + VoxelExtent, float4(1, 1, 0, 1));
 	}
 }
 #endif // DebugTranslucencyLightingVolumeCS
-
diff --git a/Engine/Source/Runtime/Renderer/Private/TranslucentLighting.cpp b/Engine/Source/Runtime/Renderer/Private/TranslucentLighting.cpp
--- a/Engine/Source/Runtime/Renderer/Private/TranslucentLighting.cpp
+++ b/Engine/Source/Runtime/Renderer/Private/TranslucentLighting.cpp
@@ -2222,6 +2222,48 @@ void InjectSimpleTranslucencyLightingVolumeArray(
 	}
 }
 
+class FFrustumCullVoxelsCS : public FGlobalShader
+{
+public:
+	static const int32 NumCameraFrustumPlanes = 5;
+
+	DECLARE_GLOBAL_SHADER(FFrustumCullVoxelsCS)
+	SHADER_USE_PARAMETER_STRUCT(FFrustumCullVoxelsCS, FGlobalShader)
+
+	BEGIN_SHADER_PARAMETER_STRUCT(FParameters, )
+		SHADER_PARAMETER_STRUCT_REF(FViewUniformShaderParameters, ViewUniformBuffer)
+		SHADER_PARAMETER_RDG_BUFFER_UAV(RWStructuredBuffer, RWVoxelAllocator0)
+		SHADER_PARAMETER_RDG_BUFFER_UAV(RWStructuredBuffer, RWVoxelAllocator1)
+		SHADER_PARAMETER(uint32, NumVoxels)
+		SHADER_PARAMETER_ARRAY(FVector4f, CameraFrustumPlanes, [NumCameraFrustumPlanes])
+	END_SHADER_PARAMETER_STRUCT()
+
+	using FPermutationDomain = TShaderPermutationDomain<>;
+
+	static int32 GetGroupSize()
+	{
+		return 4;
+	}
+
+	static bool ShouldCompilePermutation(const FGlobalShaderPermutationParameters& Parameters)
+	{
+		return IsFeatureLevelSupported(Parameters.Platform, ERHIFeatureLevel::SM5);
+	}
+
+	static void ModifyCompilationEnvironment(const FGlobalShaderPermutationParameters& Parameters, FShaderCompilerEnvironment& OutEnvironment)
+	{
+		FGlobalShader::ModifyCompilationEnvironment(Parameters, OutEnvironment);
+		OutEnvironment.SetDefine(TEXT("THREADGROUP_SIZE"), GetGroupSize());
+	}
+
+	static EShaderPermutationPrecacheRequest ShouldPrecachePermutation(const FGlobalShaderPermutationParameters& Parameters)
+	{
+		return EShaderPermutationPrecacheRequest::NotPrecached;
+	}
+};
+
+IMPLEMENT_GLOBAL_SHADER(FFrustumCullVoxelsCS, "/Engine/Private/TranslucentLightingShaders.usf", "FrustumCullVoxelsCS", SF_Compute);
+
 class FDebugTranslucencyLightingVolumeCS : public FGlobalShader
 {
 	DECLARE_GLOBAL_SHADER(FDebugTranslucencyLightingVolumeCS)
@@ -2230,8 +2272,13 @@ class FDebugTranslucencyLightingVolumeCS : public FGlobalShader
 	BEGIN_SHADER_PARAMETER_STRUCT(FParameters, )
 		SHADER_PARAMETER_STRUCT_REF(FViewUniformShaderParameters, ViewUniformBuffer)
 		SHADER_PARAMETER_STRUCT_INCLUDE(ShaderPrint::FShaderParameters, ShaderPrintUniformBuffer)
-		SHADER_PARAMETER_RDG_TEXTURE(Texture3D, InnerVolumeMarkTexture)
-		SHADER_PARAMETER_RDG_TEXTURE(Texture3D, OuterVolumeMarkTexture)
+		SHADER_PARAMETER_RDG_TEXTURE(Texture3D, VolumeMarkTexture0)
+		SHADER_PARAMETER_RDG_TEXTURE(Texture3D, VolumeMarkTexture1)
+		SHADER_PARAMETER_RDG_BUFFER_SRV(StructuredBuffer, VoxelAllocator0)
+		SHADER_PARAMETER_RDG_BUFFER_SRV(StructuredBuffer, VoxelAllocator1)
+		SHADER_PARAMETER_RDG_BUFFER_SRV(StructuredBuffer, FrustumCullVoxelAllocator0)
+		SHADER_PARAMETER_RDG_BUFFER_SRV(StructuredBuffer, FrustumCullVoxelAllocator1)
+		SHADER_PARAMETER(uint32, NumVoxels)
 	END_SHADER_PARAMETER_STRUCT()
 
 	using FPermutationDomain = TShaderPermutationDomain<>;
@@ -2271,8 +2318,10 @@ static void DrawDebugTranslucencyLightingVolume(
 	}
 
 	const FIntVector TranslucencyLightingVolumeDim(GetTranslucencyLightingVolumeDim());
+	const uint32 NumVoxels = TranslucencyLightingVolumeDim.X * TranslucencyLightingVolumeDim.Y * TranslucencyLightingVolumeDim.Z;
 
 	FRDGTextureRef DummyTexture = GSystemTextures.GetVolumetricBlackUintDummy(GraphBuilder);
+	FRDGBufferRef DummyAllocatorBuffer = GSystemTextures.GetDefaultStructuredBuffer(GraphBuilder, sizeof(uint32), IsTranslucencyLightingVolumeUsingVoxelMarking() ? 0 : NumVoxels);
 
 	for (int32 ViewIndex = 0; ViewIndex < Views.Num(); ++ViewIndex)
 	{
@@ -2280,27 +2329,77 @@ static void DrawDebugTranslucencyLightingVolume(
 		RDG_GPU_MASK_SCOPE(GraphBuilder, View.GPUMask);
 		RDG_EVENT_SCOPE(GraphBuilder, "View%d", ViewIndex);
 
-		FDebugTranslucencyLightingVolumeCS::FParameters* PassParameters = GraphBuilder.AllocParameters<FDebugTranslucencyLightingVolumeCS::FParameters>();
-		PassParameters->ViewUniformBuffer = View.ViewUniformBuffer;
-		PassParameters->InnerVolumeMarkTexture = View.TranslucencyVolumeMarkData[0].MarkTexture ? View.TranslucencyVolumeMarkData[0].MarkTexture : DummyTexture;
-		PassParameters->OuterVolumeMarkTexture = View.TranslucencyVolumeMarkData[1].MarkTexture ? View.TranslucencyVolumeMarkData[1].MarkTexture : DummyTexture;
-
 		ShaderPrint::SetEnabled(true);
-		ShaderPrint::RequestSpaceForLines(128u);
+		ShaderPrint::RequestSpaceForLines(12 * int32(NumVoxels * 0.1f));
 		ShaderPrint::RequestSpaceForCharacters(128u);
-		ShaderPrint::RequestSpaceForTriangles(64u);
-		ShaderPrint::SetParameters(GraphBuilder, View.ShaderPrintData, PassParameters->ShaderPrintUniformBuffer);
 
-		FDebugTranslucencyLightingVolumeCS::FPermutationDomain PermutationVector;
-		auto ComputeShader = View.ShaderMap->GetShader<FDebugTranslucencyLightingVolumeCS>(PermutationVector);
-		const FIntVector GroupCount = FComputeShaderUtils::GetGroupCount(TranslucencyLightingVolumeDim, FDebugTranslucencyLightingVolumeCS::GetGroupSize());
+		const FViewInfo::FTranslucencyVolumeMarkData& VolumeMarkData0 = View.TranslucencyVolumeMarkData[0];
+		const FViewInfo::FTranslucencyVolumeMarkData& VolumeMarkData1 = View.TranslucencyVolumeMarkData[1];
 
-		FComputeShaderUtils::AddPass(
-			GraphBuilder,
-			RDG_EVENT_NAME("DebugTranslucencyLightingVolumeCS"),
-			ComputeShader,
-			PassParameters,
-			GroupCount);
+		FRDGBufferRef FrustumCullVoxelAllocator0 = GraphBuilder.CreateBuffer(FRDGBufferDesc::CreateStructuredDesc(sizeof(uint32), 1), TEXT("TranslucencyLightingVolume.FrustumCullVoxelAllocator"));
+		FRDGBufferRef FrustumCullVoxelAllocator1 = GraphBuilder.CreateBuffer(FRDGBufferDesc::CreateStructuredDesc(sizeof(uint32), 1), TEXT("TranslucencyLightingVolume.FrustumCullVoxelAllocator"));
+
+		AddClearUAVPass(GraphBuilder, GraphBuilder.CreateUAV(FrustumCullVoxelAllocator0), 0);
+		AddClearUAVPass(GraphBuilder, GraphBuilder.CreateUAV(FrustumCullVoxelAllocator1), 0);
+
+		{
+			FMatrix TranslatedViewProjMatrix = View.ViewMatrices.GetTranslatedViewProjectionMatrix();
+
+			FPlane FrustumPlanes[FFrustumCullVoxelsCS::NumCameraFrustumPlanes];
+			TranslatedViewProjMatrix.GetFrustumNearPlane(FrustumPlanes[0]);
+			TranslatedViewProjMatrix.GetFrustumLeftPlane(FrustumPlanes[1]);
+			TranslatedViewProjMatrix.GetFrustumRightPlane(FrustumPlanes[2]);
+			TranslatedViewProjMatrix.GetFrustumTopPlane(FrustumPlanes[3]);
+			TranslatedViewProjMatrix.GetFrustumBottomPlane(FrustumPlanes[4]);
+			static_assert(FFrustumCullVoxelsCS::NumCameraFrustumPlanes == 5, "Update code above for the new number of frustum planes");
+
+			FFrustumCullVoxelsCS::FParameters* PassParameters = GraphBuilder.AllocParameters<FFrustumCullVoxelsCS::FParameters>();
+			PassParameters->ViewUniformBuffer = View.ViewUniformBuffer;
+			PassParameters->RWVoxelAllocator0 = GraphBuilder.CreateUAV(FrustumCullVoxelAllocator0);
+			PassParameters->RWVoxelAllocator1 = GraphBuilder.CreateUAV(FrustumCullVoxelAllocator1);
+			PassParameters->NumVoxels = NumVoxels;
+
+			for (uint32 PlaneIndex = 0; PlaneIndex < FFrustumCullVoxelsCS::NumCameraFrustumPlanes; ++PlaneIndex)
+			{
+				PassParameters->CameraFrustumPlanes[PlaneIndex] = FVector4f(FrustumPlanes[PlaneIndex].X, FrustumPlanes[PlaneIndex].Y, FrustumPlanes[PlaneIndex].Z, FrustumPlanes[PlaneIndex].W);
+			}
+
+			FFrustumCullVoxelsCS::FPermutationDomain PermutationVector;
+			auto ComputeShader = View.ShaderMap->GetShader<FFrustumCullVoxelsCS>(PermutationVector);
+			const FIntVector GroupCount = FComputeShaderUtils::GetGroupCount(TranslucencyLightingVolumeDim, FFrustumCullVoxelsCS::GetGroupSize());
+
+			FComputeShaderUtils::AddPass(
+				GraphBuilder,
+				RDG_EVENT_NAME("FrustumCullVoxelsCS"),
+				ComputeShader,
+				PassParameters,
+				GroupCount);
+		}
+
+		{
+			FDebugTranslucencyLightingVolumeCS::FParameters* PassParameters = GraphBuilder.AllocParameters<FDebugTranslucencyLightingVolumeCS::FParameters>();
+			PassParameters->ViewUniformBuffer = View.ViewUniformBuffer;
+			PassParameters->VolumeMarkTexture0 = VolumeMarkData0.MarkTexture ? VolumeMarkData0.MarkTexture : DummyTexture;
+			PassParameters->VolumeMarkTexture1 = VolumeMarkData1.MarkTexture ? VolumeMarkData1.MarkTexture : DummyTexture;
+			PassParameters->VoxelAllocator0 = GraphBuilder.CreateSRV(VolumeMarkData0.VoxelAllocator ? VolumeMarkData0.VoxelAllocator : DummyAllocatorBuffer);
+			PassParameters->VoxelAllocator1 = GraphBuilder.CreateSRV(VolumeMarkData1.VoxelAllocator ? VolumeMarkData1.VoxelAllocator : DummyAllocatorBuffer);
+			PassParameters->FrustumCullVoxelAllocator0 = GraphBuilder.CreateSRV(FrustumCullVoxelAllocator0);
+			PassParameters->FrustumCullVoxelAllocator1 = GraphBuilder.CreateSRV(FrustumCullVoxelAllocator1);
+			PassParameters->NumVoxels = NumVoxels;
+
+			ShaderPrint::SetParameters(GraphBuilder, View.ShaderPrintData, PassParameters->ShaderPrintUniformBuffer);
+
+			FDebugTranslucencyLightingVolumeCS::FPermutationDomain PermutationVector;
+			auto ComputeShader = View.ShaderMap->GetShader<FDebugTranslucencyLightingVolumeCS>(PermutationVector);
+			const FIntVector GroupCount = FComputeShaderUtils::GetGroupCount(TranslucencyLightingVolumeDim, FDebugTranslucencyLightingVolumeCS::GetGroupSize());
+
+			FComputeShaderUtils::AddPass(
+				GraphBuilder,
+				RDG_EVENT_NAME("DebugTranslucencyLightingVolumeCS"),
+				ComputeShader,
+				PassParameters,
+				GroupCount);
+		}
 	}
 }
 
