Fix shader filename case errors

#rnx

diff --git a/Engine/Plugins/FX/Niagara/Shaders/Private/Stateless/NiagaraStatelessSimulationDefault.usf b/Engine/Plugins/FX/Niagara/Shaders/Private/Stateless/NiagaraStatelessSimulationDefault.usf
--- a/Engine/Plugins/FX/Niagara/Shaders/Private/Stateless/NiagaraStatelessSimulationDefault.usf
+++ b/Engine/Plugins/FX/Niagara/Shaders/Private/Stateless/NiagaraStatelessSimulationDefault.usf
@@ -85,7 +85,7 @@
 #include "/Plugin/FX/Niagara/Private/Stateless/Modules/NiagaraStatelessModule_ScaleSpriteSize.ush"
 #include "/Plugin/FX/Niagara/Private/Stateless/Modules/NiagaraStatelessModule_ScaleSpriteSizeBySpeed.ush"
 #include "/Plugin/FX/Niagara/Private/Stateless/Modules/NiagaraStatelessModule_ScaleMeshSize.ush"
-#include "/Plugin/FX/Niagara/Private/Stateless/Modules/NiagaraStatelessModule_ScaleMeshSizebySpeed.ush"
+#include "/Plugin/FX/Niagara/Private/Stateless/Modules/NiagaraStatelessModule_ScaleMeshSizeBySpeed.ush"
 #include "/Plugin/FX/Niagara/Private/Stateless/Modules/NiagaraStatelessModule_SolveVelocitiesAndForces.ush"
 #include "/Plugin/FX/Niagara/Private/Stateless/Modules/NiagaraStatelessModule_SpriteFacingAndAlignment.ush"
 #include "/Plugin/FX/Niagara/Private/Stateless/Modules/NiagaraStatelessModule_SpriteRotationRate.ush"
diff --git a/Engine/Plugins/FX/Niagara/Source/Niagara/Private/Stateless/Modules/NiagaraStatelessModule_ScaleMeshSizeBySpeed.cpp b/Engine/Plugins/FX/Niagara/Source/Niagara/Private/Stateless/Modules/NiagaraStatelessModule_ScaleMeshSizeBySpeed.cpp
--- a/Engine/Plugins/FX/Niagara/Source/Niagara/Private/Stateless/Modules/NiagaraStatelessModule_ScaleMeshSizeBySpeed.cpp
+++ b/Engine/Plugins/FX/Niagara/Source/Niagara/Private/Stateless/Modules/NiagaraStatelessModule_ScaleMeshSizeBySpeed.cpp
@@ -126,7 +126,7 @@ void UNiagaraStatelessModule_ScaleMeshSizeBySpeed::SetShaderParameters(const FNi
 #if WITH_EDITORONLY_DATA
 const TCHAR* UNiagaraStatelessModule_ScaleMeshSizeBySpeed::GetShaderTemplatePath() const
 {
-	return TEXT("/Plugin/FX/Niagara/Private/Stateless/Modules/NiagaraStatelessModule_ScaleMeshSizebySpeed.ush");
+	return TEXT("/Plugin/FX/Niagara/Private/Stateless/Modules/NiagaraStatelessModule_ScaleMeshSizeBySpeed.ush");
 }
 
 void UNiagaraStatelessModule_ScaleMeshSizeBySpeed::GetOutputVariables(TArray<FNiagaraVariableBase>& OutVariables, EVariableFilter Filter) const
diff --git a/Engine/Shaders/Private/RayTracing/RayTracingDebug.usf b/Engine/Shaders/Private/RayTracing/RayTracingDebug.usf
--- a/Engine/Shaders/Private/RayTracing/RayTracingDebug.usf
+++ b/Engine/Shaders/Private/RayTracing/RayTracingDebug.usf
@@ -25,7 +25,7 @@
 #include "RayTracingDebugUtils.ush"
 #include "RayTracingLightCullingCommon.ush"
 
-#include "/Engine/Shared/RayTracingDebugDefinitions.h"
+#include "/Engine/Shared/RaytracingDebugDefinitions.h"
 
 RWTexture2D<float4> Output;
 RWTexture2D<float> OutputDepth;
diff --git a/Engine/Shaders/Private/RayTracing/RayTracingDebugHitStats.usf b/Engine/Shaders/Private/RayTracing/RayTracingDebugHitStats.usf
--- a/Engine/Shaders/Private/RayTracing/RayTracingDebugHitStats.usf
+++ b/Engine/Shaders/Private/RayTracing/RayTracingDebugHitStats.usf
@@ -3,7 +3,7 @@
 #include "../Common.ush"
 #include "RayTracingDebugUtils.ush"
 
-#include "/Engine/Shared/RayTracingDebugDefinitions.h"
+#include "/Engine/Shared/RaytracingDebugDefinitions.h"
 #include "/Engine/Shared/RayTracingDebugTypes.h"
 
 uint OpaqueOnly;
diff --git a/Engine/Shaders/Private/RayTracing/RayTracingDebugTraversal.usf b/Engine/Shaders/Private/RayTracing/RayTracingDebugTraversal.usf
--- a/Engine/Shaders/Private/RayTracing/RayTracingDebugTraversal.usf
+++ b/Engine/Shaders/Private/RayTracing/RayTracingDebugTraversal.usf
@@ -12,7 +12,7 @@
 #include "../Substrate/Substrate.ush"
 #include "../Lumen/LumenPosition.ush"
 
-#include "/Engine/Shared/RayTracingDebugDefinitions.h"
+#include "/Engine/Shared/RaytracingDebugDefinitions.h"
 
 #include "TraceRayInline.ush"
 
diff --git a/Engine/Source/Runtime/Renderer/Private/RayTracing/RayTracingInstanceBufferUtil.cpp b/Engine/Source/Runtime/Renderer/Private/RayTracing/RayTracingInstanceBufferUtil.cpp
--- a/Engine/Source/Runtime/Renderer/Private/RayTracing/RayTracingInstanceBufferUtil.cpp
+++ b/Engine/Source/Runtime/Renderer/Private/RayTracing/RayTracingInstanceBufferUtil.cpp
@@ -629,7 +629,7 @@ struct FRayTracingBuildInstanceBufferCS : public FGlobalShader
 	}
 };
 
-IMPLEMENT_GLOBAL_SHADER(FRayTracingBuildInstanceBufferCS, "/Engine/Private/Raytracing/RayTracingInstanceBufferUtil.usf", "RayTracingBuildInstanceBufferCS", SF_Compute);
+IMPLEMENT_GLOBAL_SHADER(FRayTracingBuildInstanceBufferCS, "/Engine/Private/RayTracing/RayTracingInstanceBufferUtil.usf", "RayTracingBuildInstanceBufferCS", SF_Compute);
 
 void BuildRayTracingInstanceBuffer(
 	FRHICommandList& RHICmdList,
diff --git a/Engine/Source/Runtime/Renderer/Private/RayTracing/RayTracingScene.cpp b/Engine/Source/Runtime/Renderer/Private/RayTracing/RayTracingScene.cpp
--- a/Engine/Source/Runtime/Renderer/Private/RayTracing/RayTracingScene.cpp
+++ b/Engine/Source/Runtime/Renderer/Private/RayTracing/RayTracingScene.cpp
@@ -604,7 +604,7 @@ struct FRayTracingProcessFeedbackCS : public FGlobalShader
 	}
 };
 
-IMPLEMENT_GLOBAL_SHADER(FRayTracingProcessFeedbackCS, "/Engine/Private/Raytracing/RayTracingFeedback.usf", "RayTracingProcessFeedbackCS", SF_Compute);
+IMPLEMENT_GLOBAL_SHADER(FRayTracingProcessFeedbackCS, "/Engine/Private/RayTracing/RayTracingFeedback.usf", "RayTracingProcessFeedbackCS", SF_Compute);
 
 struct FRayTracingUpdateGeometryHitCountCS : public FGlobalShader
 {
@@ -640,7 +640,7 @@ struct FRayTracingUpdateGeometryHitCountCS : public FGlobalShader
 	}
 };
 
-IMPLEMENT_GLOBAL_SHADER(FRayTracingUpdateGeometryHitCountCS, "/Engine/Private/Raytracing/RayTracingFeedback.usf", "RayTracingUpdateGeometryHitCountCS", SF_Compute);
+IMPLEMENT_GLOBAL_SHADER(FRayTracingUpdateGeometryHitCountCS, "/Engine/Private/RayTracing/RayTracingFeedback.usf", "RayTracingUpdateGeometryHitCountCS", SF_Compute);
 
 BEGIN_SHADER_PARAMETER_STRUCT(FFeedbackReadbackPassParameters, )
 	RDG_BUFFER_ACCESS(HandleBuffer, ERHIAccess::CopySrc)
diff --git a/Engine/Source/Runtime/Renderer/Private/RayTracing/RaytracingSkylight.cpp b/Engine/Source/Runtime/Renderer/Private/RayTracing/RaytracingSkylight.cpp
--- a/Engine/Source/Runtime/Renderer/Private/RayTracing/RaytracingSkylight.cpp
+++ b/Engine/Source/Runtime/Renderer/Private/RayTracing/RaytracingSkylight.cpp
@@ -268,7 +268,7 @@ class FRayTracingSkyLightRGS : public FGlobalShader
 	END_SHADER_PARAMETER_STRUCT()
 };
 
-IMPLEMENT_GLOBAL_SHADER(FRayTracingSkyLightRGS, "/Engine/Private/Raytracing/RaytracingSkylightRGS.usf", "SkyLightRGS", SF_RayGen);
+IMPLEMENT_GLOBAL_SHADER(FRayTracingSkyLightRGS, "/Engine/Private/RayTracing/RayTracingSkyLightRGS.usf", "SkyLightRGS", SF_RayGen);
 
 void FDeferredShadingSceneRenderer::PrepareRayTracingSkyLight(const FViewInfo& View, const FScene& Scene, TArray<FRHIRayTracingShader*>& OutRayGenShaders)
 {
diff --git a/Engine/Source/Runtime/Renderer/Private/SubsurfaceTiles.cpp b/Engine/Source/Runtime/Renderer/Private/SubsurfaceTiles.cpp
--- a/Engine/Source/Runtime/Renderer/Private/SubsurfaceTiles.cpp
+++ b/Engine/Source/Runtime/Renderer/Private/SubsurfaceTiles.cpp
@@ -55,7 +55,7 @@ public:
 	}
 };
 
-IMPLEMENT_GLOBAL_SHADER(FClearUAVBuildIndirectDispatchBufferCS, "/Engine/Private/PostprocessSubsurfaceTile.usf", "BuildIndirectDispatchArgsCS", SF_Compute);
+IMPLEMENT_GLOBAL_SHADER(FClearUAVBuildIndirectDispatchBufferCS, "/Engine/Private/PostProcessSubsurfaceTile.usf", "BuildIndirectDispatchArgsCS", SF_Compute);
 
 class FClearUAVCS : public FGlobalShader
 {
@@ -77,7 +77,7 @@ public:
 	}
 };
 
-IMPLEMENT_GLOBAL_SHADER(FClearUAVCS, "/Engine/Private/PostprocessSubsurfaceTile.usf", "ClearUAV", SF_Compute);
+IMPLEMENT_GLOBAL_SHADER(FClearUAVCS, "/Engine/Private/PostProcessSubsurfaceTile.usf", "ClearUAV", SF_Compute);
 
 
 void AddConditionalClearBlackUAVPass(FRDGBuilder& GraphBuilder, FRDGEventName&& Name, 
@@ -249,7 +249,7 @@ public:
 	}
 };
 
-IMPLEMENT_GLOBAL_SHADER(FSubsurfaceTileBuildIndirectDispatchArgsCS, "/Engine/Private/PostprocessSubsurfaceTile.usf", "SubsurfaceTileBuildIndirectDispatchArgsCS", SF_Compute);
+IMPLEMENT_GLOBAL_SHADER(FSubsurfaceTileBuildIndirectDispatchArgsCS, "/Engine/Private/PostProcessSubsurfaceTile.usf", "SubsurfaceTileBuildIndirectDispatchArgsCS", SF_Compute);
 
 // Returns the [0, N] clamped value of the 'r.SSS.Scale' CVar.
 static float GetSubsurfaceRadiusScaleForTiling()
