[PostProcessing] Add ColorParade, VectorScope and ColorHistogram color grading visualizations. Can be enabled with ShowFlag.VisualizeColorGrading

diff --git a/Engine/Shaders/Private/Tools/ColorGradingHistogram.usf b/Engine/Shaders/Private/Tools/ColorGradingHistogram.usf
new file mode 100644
--- /dev/null
+++ b/Engine/Shaders/Private/Tools/ColorGradingHistogram.usf
@@ -0,0 +1,185 @@
+// Copyright Epic Games, Inc. All Rights Reserved.
+
+/*=============================================================================
+	ColorGradingColorHistogram.usf: Render ColorHistogram color grading view.
+	Charts the intensity of the color in every pixel in the input as a histogram.
+=============================================================================*/
+
+
+#include "/Engine/Private/Common.ush"
+
+#define THREADGROUP_THREADCOUNT (THREADGROUP_SIZE_X * THREADGROUP_SIZE_Y)
+
+int IntensityToIntensityBin(float NormalizedIntensity)
+{
+	int Bin = round(NormalizedIntensity * (PER_CHANNEL_RESOLUTION - 1));
+	return Bin;
+}
+
+#ifdef BuildColorHistogramCS
+
+Texture2D InputColorTexture;
+int4 InputColorViewRect;
+float2 MinMaxValue;
+RWTexture2D<uint> OutputHistogramTexture;
+RWStructuredBuffer<uint> OutputMetaDataBuffer;
+
+#define NUM_CHANNELS (3)
+
+groupshared uint LDS[PER_CHANNEL_RESOLUTION];
+groupshared uint LDS_Max[NUM_CHANNELS];
+
+void ClearLDS(uint GroupThreadIndex)
+{
+	for (uint Index = GroupThreadIndex; Index < PER_CHANNEL_RESOLUTION; Index += THREADGROUP_THREADCOUNT)
+	{
+		LDS[Index] = 0;
+	}
+	for (uint Index = GroupThreadIndex; Index < NUM_CHANNELS; Index += THREADGROUP_THREADCOUNT)
+	{
+		LDS_Max[Index] = 0;
+	}
+}
+
+void FlushAndClearLDS(uint Channel, uint GroupThreadIndex)
+{
+	for (uint Bin = GroupThreadIndex; Bin < PER_CHANNEL_RESOLUTION; Bin += THREADGROUP_THREADCOUNT)
+	{
+		uint Value = LDS[Bin];
+
+		if (Value > 0)
+		{
+			LDS[Bin] = 0;
+
+			uint PrevValue;
+			InterlockedAdd(OutputHistogramTexture[int2(Bin, Channel)], Value, PrevValue);
+			uint CurValue = PrevValue + Value;
+			InterlockedMax(LDS_Max[Channel], CurValue);
+		}
+	}
+}
+
+[numthreads(THREADGROUP_SIZE_X, THREADGROUP_SIZE_Y, 1)]
+void BuildColorHistogramCS(
+	uint2 DispatchThreadId : SV_DispatchThreadID,
+	uint2 GroupId : SV_GroupID,
+	uint2 GroupThreadId : SV_GroupThreadID,
+	uint GroupThreadIndex : SV_GroupIndex)
+{
+	float MinValue = MinMaxValue[0];
+	float MaxValue = MinMaxValue[1];
+
+	int2 InputColorViewSize = InputColorViewRect.zw - InputColorViewRect.xy;
+	int2 PixelPosition = DispatchThreadId;
+	int2 BufferPixelPosition = InputColorViewRect.xy + PixelPosition;
+
+	bool bHasValidPixel = all(PixelPosition < InputColorViewSize);
+
+	float3 InputSample = 0;
+	if (bHasValidPixel)
+	{
+		InputSample = InputColorTexture.Load(int3(BufferPixelPosition, 0)).rgb;
+	}
+
+	float3 NormalizedSample = (InputSample - MinValue) / (MaxValue - MinValue);
+
+	ClearLDS(GroupThreadIndex);
+	GroupMemoryBarrierWithGroupSync();
+
+	for (int Channel = 0; Channel < NUM_CHANNELS; ++Channel)
+	{
+		int Bin = IntensityToIntensityBin(NormalizedSample[Channel]);
+
+		InterlockedAdd(LDS[Bin], 1);
+
+		GroupMemoryBarrierWithGroupSync();
+
+		FlushAndClearLDS(Channel, GroupThreadIndex);
+
+		GroupMemoryBarrierWithGroupSync();
+	}
+
+	if (GroupThreadIndex == 0)
+	{
+		for (int Channel = 0; Channel < NUM_CHANNELS; ++Channel)
+		{
+			InterlockedMax(OutputMetaDataBuffer[Channel], LDS_Max[Channel]);
+		}
+	}
+}
+
+#endif
+
+#ifdef DrawColorHistogramCS
+
+Texture2D<uint> HistogramTexture;
+StructuredBuffer<uint> MetaDataBuffer;
+RWTexture2D<float3> OutputTexture;
+int4 OutputViewRect;
+int4 InputColorViewRect;
+uint Flags;
+
+[numthreads(THREADGROUP_SIZE_X, THREADGROUP_SIZE_Y, 1)]
+void DrawColorHistogramCS(
+	uint2 DispatchThreadId : SV_DispatchThreadID,
+	uint2 GroupId : SV_GroupID,
+	uint2 GroupThreadId : SV_GroupThreadID,
+	uint GroupThreadIndex : SV_GroupIndex)
+{
+	int2 OutputViewSize = OutputViewRect.zw - OutputViewRect.xy;
+	int2 PixelPosition = int2(DispatchThreadId.x, DispatchThreadId.y);
+	int2 BufferPixelPosition = OutputViewRect.xy + PixelPosition;
+	float2 UV = (float2(PixelPosition) + 0.5) / OutputViewSize;
+	UV.y = 1.0 - UV.y; // flip y
+
+	int Bin = IntensityToIntensityBin(UV.x);
+
+	float3 Output = 0;
+
+	bool bColorize 					= (Flags & FLAG_COLORIZE) != 0;
+	bool bDrawSeparateRowPerChannel = (Flags & FLAG_DRAW_SEPARATE_ROW_PER_CHANNEL) != 0;
+
+	if (bDrawSeparateRowPerChannel)
+	{
+		int Channel = int(UV.y * 3);
+		UV.y = frac(UV.y * 3);
+
+		uint NumSamples = HistogramTexture[int2(Bin, Channel)];
+		uint LargestBinCount = MetaDataBuffer[Channel];
+
+		float BarHeight = float(NumSamples) / LargestBinCount; 
+		float Intensity = UV.y < BarHeight;
+		//Intensity *= lerp(1.0, 0.3, saturate((BarHeight - UV.y) * 20.0));
+		Output[Channel] = Intensity;
+
+		if (!bColorize)
+		{
+			Output = dot(Output, 1);
+		}
+	}
+	else
+	{
+		uint3 NumSamples = 0;
+		NumSamples[0] = HistogramTexture[int2(Bin, 0)];
+		NumSamples[1] = HistogramTexture[int2(Bin, 1)];
+		NumSamples[2] = HistogramTexture[int2(Bin, 2)];
+
+		uint3 LargestBinCount = 0;
+		LargestBinCount[0] = MetaDataBuffer[0];
+		LargestBinCount[1] = MetaDataBuffer[1];
+		LargestBinCount[2] = MetaDataBuffer[2];
+
+		Output = float3(NumSamples) / LargestBinCount;
+
+		Output = UV.yyy < Output;
+
+		if (!bColorize)
+		{
+			Output = dot(Output, 0.33333f);
+		}
+	}
+
+	OutputTexture[BufferPixelPosition] = Output;
+}
+
+#endif
diff --git a/Engine/Shaders/Private/Tools/ColorGradingParade.usf b/Engine/Shaders/Private/Tools/ColorGradingParade.usf
new file mode 100644
--- /dev/null
+++ b/Engine/Shaders/Private/Tools/ColorGradingParade.usf
@@ -0,0 +1,274 @@
+// Copyright Epic Games, Inc. All Rights Reserved.
+
+/*=============================================================================
+	ColorGradingParade.usf: Render RGB Parade and Waveform color grading views.
+	These views give a histogram of color channel intensity, per image column.
+=============================================================================*/
+
+#define MATERIAL_DOMAIN_UI 1 // required for LuminanceFactors()
+#include "/Engine/Private/Common.ush"
+
+#ifndef DIM_WAVE_OPS
+#define DIM_WAVE_OPS 0
+#endif
+
+float2 MinMaxValue;
+uint Flags;
+
+#define QUANT_FACTOR 1024.0 // Quantization factor, used to convert float to fixed precision int
+#define QUANT(x) (int((x) * QUANT_FACTOR))
+#define DEQUANT(x) (float(x) / QUANT_FACTOR)
+
+#define COLUMN_FLAG_HAS_NAN (1U << 0)
+#define COLUMN_FLAG_HAS_INF (1U << 1)
+
+int ImageColumnToHistogramColumn(int Column, int ChannelIndex, int2 InputViewRectSize)
+{
+	return Column*3 + ChannelIndex;
+	//return ChannelIndex * InputViewRectSize.x + Column;
+}
+
+#ifdef BuildColorParadeHistogramCS
+
+#define THREADGROUP_THREADCOUNT (THREADGROUP_SIZE_X * THREADGROUP_SIZE_Y)
+#define NUM_BINS_PER_THREAD (NUM_BINS / THREADGROUP_SIZE_Y)
+
+Texture2D InputColorTexture;
+RWTexture2D<uint> OutputTexture;
+RWStructuredBuffer<uint> OutputColumnFlagsBuffer;
+float4 InputColorViewRect;
+int2 OutputHistogramSize; // x is image columns, y is histogram bins
+
+groupshared uint Bins[NUM_BINS];
+
+void ClearLDS(uint GroupThreadIndex)
+{
+	for (uint Index = GroupThreadIndex; Index < NUM_BINS; Index += THREADGROUP_THREADCOUNT)
+	{
+		Bins[Index] = 0;
+	}
+}
+
+[numthreads(THREADGROUP_SIZE_X, THREADGROUP_SIZE_Y, 1)]
+void BuildColorParadeHistogramCS(
+	uint2 DispatchThreadId : SV_DispatchThreadID,
+	uint2 GroupId : SV_GroupID,
+	uint2 GroupThreadId : SV_GroupThreadID,
+	uint GroupThreadIndex : SV_GroupIndex)
+{
+	const float MinValue = MinMaxValue.x;
+	const float MaxValue = MinMaxValue.y;
+
+	float2 InputViewRectSize = InputColorViewRect.zw - InputColorViewRect.xy;
+
+	// Read input value. If the output graph is smaller than the input image, aggregate multiple pixels
+	float3 IntensityAcc = (float3)0;
+	bool bHasInf = false;
+	bool bHasNan = false;
+	{
+		int InputColumnsPerOutputColumn = ceil(InputViewRectSize.x / OutputHistogramSize.x);
+		int StartColumn = InputColorViewRect.x + DispatchThreadId.x * InputColumnsPerOutputColumn;
+		int NumColumnsToRead = min(InputViewRectSize.x - StartColumn, InputColumnsPerOutputColumn);
+		int EndColumn = min(StartColumn + NumColumnsToRead, InputColorViewRect.z);
+		int Row = min(InputColorViewRect.y + DispatchThreadId.y, InputColorViewRect.w-1);
+
+		for (int Column = StartColumn; Column < EndColumn; ++Column)
+		{
+			float3 InputValue = InputColorTexture.Load(int3(Column, Row, 0)).rgb;
+
+			bHasInf = bHasInf || any(!IsFinite(InputValue));
+			bHasNan = bHasNan || any(isnan(InputValue));
+			IntensityAcc += InputValue;
+		}
+		IntensityAcc /= NumColumnsToRead;
+	}
+
+#if DIM_WAVE_OPS
+	bHasInf = WaveActiveAnyTrue(bHasInf);
+	bHasNan = WaveActiveAnyTrue(bHasNan);
+#endif
+
+	// Build local intensity histogram in LDS
+	float3 HistogramValues[NUM_BINS_PER_THREAD];
+	for (int ComponentIndex = 0; ComponentIndex < 3; ++ComponentIndex)
+	{
+		ClearLDS(GroupThreadIndex);
+		GroupMemoryBarrierWithGroupSync();
+
+		// Write histogram
+		if (DispatchThreadId.y < InputViewRectSize.y)
+		{
+			float Component = IntensityAcc[ComponentIndex];
+			float UnitIntensity = saturate((Component - MinValue) / (MaxValue - MinValue));
+			float BinIndexFloat = UnitIntensity * (NUM_BINS - 1);
+			int BinIndex = round(BinIndexFloat);
+
+			int PrevValue;
+
+			// faster but with aliasing
+			//InterlockedAdd(Bins[BinIndex], QUANT(1), PrevValue);
+
+			// slower but higher quality
+			InterlockedAdd(Bins[floor(BinIndexFloat)], QUANT(1.0 - frac(BinIndexFloat)), PrevValue);
+			InterlockedAdd(Bins[ceil(BinIndexFloat)], QUANT(frac(BinIndexFloat)), PrevValue);
+		}
+
+		GroupMemoryBarrierWithGroupSync();
+
+		// Read histogram into VGPR
+		for (int Bin = 0; Bin < NUM_BINS_PER_THREAD; ++Bin)
+		{
+			int BinIndex = (GroupThreadIndex * NUM_BINS_PER_THREAD) + Bin;
+			HistogramValues[Bin][ComponentIndex] = DEQUANT(Bins[BinIndex]);
+		}
+
+		GroupMemoryBarrierWithGroupSync();
+	}
+
+	// Build global intensity histogram, accumulating histograms across threadgroups
+	for (int Bin = 0; Bin < NUM_BINS_PER_THREAD; ++Bin)
+	{
+		int BinIndex = (GroupThreadIndex * NUM_BINS_PER_THREAD) + Bin;
+		float3 BinValue = HistogramValues[Bin];
+
+		if (any(BinValue != 0))
+		{
+			for (int ComponentIndex = 0; ComponentIndex < 3; ++ComponentIndex)
+			{
+				int HistogramColumn = ImageColumnToHistogramColumn(DispatchThreadId.x, ComponentIndex, InputViewRectSize);
+				InterlockedAdd(OutputTexture[int2(HistogramColumn, BinIndex)], QUANT(BinValue[ComponentIndex]));
+			}
+		}
+	}
+
+	uint ColumnFlags = 0;
+	ColumnFlags |= bHasNan ? COLUMN_FLAG_HAS_NAN : 0;
+	ColumnFlags |= bHasInf ? COLUMN_FLAG_HAS_INF : 0;
+	if (ColumnFlags != 0
+#if DIM_WAVE_OPS
+		&& WaveIsFirstLane()
+#endif
+		)
+	{
+		int HistogramColumn = ImageColumnToHistogramColumn(DispatchThreadId.x, 0, InputViewRectSize);
+		InterlockedOr(OutputColumnFlagsBuffer[HistogramColumn], ColumnFlags);
+	}
+}
+
+#endif
+
+#ifdef DrawColorParadeCS
+
+#define THREADGROUP_THREADCOUNT (THREADGROUP_SIZE_X * THREADGROUP_SIZE_Y)
+
+Texture2D<uint> HistogramTexture;
+SamplerState HistogramTextureSampler;
+StructuredBuffer<uint> ColumnFlagsBuffer;
+RWTexture2D<float3> OutputTexture;
+int2 InputViewSize;
+int2 HistogramSize;
+int4 OutputViewRect;
+float Brightness;
+float Gamma;
+float DesaturationPower;
+
+float SampleHistogram(const float2 UV, int Channel)
+{
+	float2 HistogramBufferSize = float2(HistogramSize.x * 3, HistogramSize.y);
+
+	int InputColumnsPerOutputColumn = ceil(float(InputViewSize.x) / HistogramSize.x);
+	float HistogramActualWidth = InputViewSize.x / InputColumnsPerOutputColumn;
+
+	float HistogramColumn = ImageColumnToHistogramColumn(floor(UV.x * HistogramActualWidth), Channel, InputViewSize);
+	float2 BufferUV = float2((HistogramColumn + 0.5) / HistogramBufferSize.x, UV.y);
+
+	return DEQUANT(HistogramTexture.Load(int3(BufferUV.x * HistogramBufferSize.x, BufferUV.y * HistogramBufferSize.y, 0)).r);
+}
+
+uint GetColumnFlags(const float2 UV, int Channel)
+{
+	int InputColumnsPerOutputColumn = ceil(float(InputViewSize.x) / HistogramSize.x);
+	float HistogramActualWidth = InputViewSize.x / InputColumnsPerOutputColumn;
+
+	int HistogramColumn = ImageColumnToHistogramColumn(floor(UV.x * HistogramActualWidth), Channel, InputViewSize);
+	return ColumnFlagsBuffer[HistogramColumn];
+}
+
+[numthreads(THREADGROUP_SIZE_X, THREADGROUP_SIZE_Y, 1)]
+void DrawColorParadeCS(
+	uint2 DispatchThreadId : SV_DispatchThreadID,
+	uint2 GroupId : SV_GroupID,
+	uint2 GroupThreadId : SV_GroupThreadID,
+	uint GroupThreadIndex : SV_GroupIndex)
+{
+	if (any(OutputViewRect.xy + DispatchThreadId > OutputViewRect.zw))
+	{
+		return;
+	}
+
+	bool bDrawSeparateColumnPerChannel = (Flags & FLAG_DRAW_SEPARATE_COLUMN_PER_CHANNEL) != 0;
+	bool bColorize 					   = (Flags & FLAG_COLORIZE) != 0;
+	bool bHighlightNaNsAndInfs		   = (Flags & FLAG_HIGHLIGHT_NANS_AND_INFS) != 0;
+
+
+	float2 OutputUV = (DispatchThreadId.xy + 0.5) / (OutputViewRect.zw - OutputViewRect.xy);
+	float2 InputUV = OutputUV;
+	InputUV.y = 1.0 - InputUV.y; // flip Y
+
+	if (bDrawSeparateColumnPerChannel)
+	{
+		InputUV.x = frac(InputUV.x * 3.0);
+	}
+	
+	float3 HistogramValue = (float3)0;
+	HistogramValue.r = SampleHistogram(InputUV, 0);
+	HistogramValue.g = SampleHistogram(InputUV, 1);
+	HistogramValue.b = SampleHistogram(InputUV, 2);
+
+	HistogramValue /= InputViewSize.y;
+
+	uint ColumnFlags = GetColumnFlags(InputUV, 0);
+	bool bHasNan = (ColumnFlags & COLUMN_FLAG_HAS_NAN) > 0;
+	bool bHasInf = (ColumnFlags & COLUMN_FLAG_HAS_INF) > 0;
+	
+	float3 OutputValue = HistogramValue * Brightness;
+
+	if (bDrawSeparateColumnPerChannel)
+	{
+		// Filter out the relevant RGB primary to draw for this column
+		// but bleed the value a bit across channels so clipping/blown-out values show up as white
+		int ComponentIndex = int(OutputUV.x * 3.0);
+		float Component = OutputValue[ComponentIndex] * 0.99;
+		OutputValue = pow(Component, DesaturationPower);
+		OutputValue[ComponentIndex] = Component;
+
+		if (!bColorize)
+		{
+			// no Luminance() here because we want to compare the intensity across columns, 
+			// rather than accurately reinterpret color to luminance
+			OutputValue = dot(OutputValue, 0.3333);
+		}
+	}
+	else
+	{
+		if (!bColorize)
+		{
+			OutputValue = Luminance(OutputValue);
+		}
+	}
+
+	// Tonemap
+	OutputValue = pow(OutputValue, Gamma);
+
+	if (bHighlightNaNsAndInfs)
+	{
+		const float3 InfColor = float3(1,1,0) * 100.0;
+		const float3 NanColor = float3(1,0,1) * 100.0;
+		OutputValue = bHasInf ? InfColor : OutputValue;
+		OutputValue = bHasNan ? NanColor : OutputValue;
+	}
+
+	OutputTexture[OutputViewRect.xy + DispatchThreadId] = OutputValue;
+}
+
+#endif
diff --git a/Engine/Shaders/Private/Tools/ColorGradingVectorScope.usf b/Engine/Shaders/Private/Tools/ColorGradingVectorScope.usf
new file mode 100644
--- /dev/null
+++ b/Engine/Shaders/Private/Tools/ColorGradingVectorScope.usf
@@ -0,0 +1,188 @@
+// Copyright Epic Games, Inc. All Rights Reserved.
+
+/*=============================================================================
+	ColorGradingVectorScope.usf: Render Vector Scope color grading view.
+	This view is a scatterplot of the hue and saturation of each pixel in the image.
+	To be more accurate, it is a plot of Cb and Cr in the BT.709 YCbCr representation.
+=============================================================================*/
+
+#define MATERIAL_DOMAIN_UI 1 // required for LuminanceFactors()
+#include "/Engine/Private/Common.ush"
+
+#define THREADGROUP_THREADCOUNT (THREADGROUP_SIZE_X * THREADGROUP_SIZE_Y)
+
+float3 RGBToYCbCr(float3 In)
+{
+	// BT.709 
+	float Y  =  0.212600 * In.r +  0.715200 * In.g +  0.072200 * In.b;
+	float Cb = -0.114572 * In.r + -0.385428 * In.g +  0.500000 * In.b;
+	float Cr =  0.500000 * In.r + -0.454153 * In.g + -0.045847 * In.b;
+	return float3(Y, Cb, Cr);
+}
+
+float3 YCbCrToRGB(float3 In)
+{
+	// BT.709
+	float R =  1.000000 * In.r +  0.000000 * In.g +  1.574800 * In.b;
+	float G =  1.000000 * In.r + -0.187324 * In.g + -0.468124 * In.b;
+	float B =  1.000000 * In.r +  1.855600 * In.g +  0.000000 * In.b;
+	return float3(R, G, B);
+}
+
+float3 RGBToNormalizedYCbCr(float3 In)
+{
+	float3 YCbCrSample = RGBToYCbCr(In);
+	float Y = YCbCrSample[0];
+	float2 CbCr = float2(YCbCrSample[1], YCbCrSample[2]);
+
+	// Convert hue and saturation to range [0;1]
+	const float2 CbCrMinMax = float2(-0.5f, 0.5f);
+	float2 NormalizedCbCr = (CbCr - CbCrMinMax[0]) / (CbCrMinMax[1] - CbCrMinMax[0]);
+
+	return float3(Y, NormalizedCbCr[0], NormalizedCbCr[1]);
+}
+
+float3 NormalizedYCbCrToRGB(float3 In)
+{
+	float2 CbCr = In.yz - 0.5f;
+	float3 YCbCr = float3(In.x, CbCr[0], CbCr[1]);
+
+	return YCbCrToRGB(YCbCr);
+}
+
+int HueToHueBin(float Hue)
+{
+	int HueBin = round(Hue * (HUE_RESOLUTION - 1));
+	return HueBin;
+}
+
+int SaturationToSaturationBin(float Saturation)
+{
+	int SaturationBin = round(Saturation * (SATURATION_RESOLUTION - 1));
+	return SaturationBin;
+}
+
+#ifdef BuildVectorScopeHistogramCS
+
+Texture2D InputColorTexture;
+int4 InputColorViewRect;
+float2 LuminanceFilterMinMax;
+RWTexture2D<uint> OutputHistogramTexture;
+
+[numthreads(THREADGROUP_SIZE_X, THREADGROUP_SIZE_Y, 1)]
+void BuildVectorScopeHistogramCS(
+	uint2 DispatchThreadId : SV_DispatchThreadID,
+	uint2 GroupId : SV_GroupID,
+	uint2 GroupThreadId : SV_GroupThreadID,
+	uint GroupThreadIndex : SV_GroupIndex)
+{
+	int2 InputColorViewSize = InputColorViewRect.zw - InputColorViewRect.xy;
+
+	uint TotalPixelCount = uint(InputColorViewSize.x) * uint(InputColorViewSize.y);
+	uint PixelIndex = DispatchThreadId.x;
+	uint GroupStartPixelIndex = GroupId.x * THREADGROUP_THREADCOUNT;
+	uint GroupEndPixelIndex = min(GroupStartPixelIndex + THREADGROUP_THREADCOUNT, TotalPixelCount);
+	int2 BufferPixelPosition = InputColorViewRect.xy + int2(PixelIndex % InputColorViewSize.x, PixelIndex / InputColorViewSize.x);
+
+	bool bHasValidPixel = PixelIndex < TotalPixelCount;
+
+	float3 NormalizedYCbCr = 0;
+	if (bHasValidPixel)
+	{
+		float3 RGBSample = InputColorTexture.Load(int3(BufferPixelPosition, 0)).rgb;
+		NormalizedYCbCr = RGBToNormalizedYCbCr(RGBSample);
+	}
+
+	float Luma = NormalizedYCbCr[0];
+	if (Luma < LuminanceFilterMinMax[0] || Luma > LuminanceFilterMinMax[1])
+	{
+		return;
+	}
+
+	float Hue = NormalizedYCbCr[1];
+	int HueBin = HueToHueBin(Hue);
+	float Saturation = NormalizedYCbCr[2];
+	int SaturationBin = SaturationToSaturationBin(Saturation);
+	InterlockedAdd(OutputHistogramTexture[int2(SaturationBin, HueBin)], 1);
+}
+
+#endif
+
+#ifdef DrawVectorScopeCS
+
+Texture2D<uint> HistogramTexture;
+RWTexture2D<float3> OutputTexture;
+int4 OutputViewRect;
+int4 InputColorViewRect;
+uint Flags;
+float Rotation;
+float Brightness;
+float Gamma;
+float DesaturationPower;
+
+[numthreads(THREADGROUP_SIZE_X, THREADGROUP_SIZE_Y, 1)]
+void DrawVectorScopeCS(
+	uint2 DispatchThreadId : SV_DispatchThreadID,
+	uint2 GroupId : SV_GroupID,
+	uint2 GroupThreadId : SV_GroupThreadID,
+	uint GroupThreadIndex : SV_GroupIndex)
+{
+	int2 OutputViewSize = OutputViewRect.zw - OutputViewRect.xy;
+	int2 PixelPosition = int2(DispatchThreadId.x, DispatchThreadId.y);
+	int2 BufferPixelPosition = OutputViewRect.xy + PixelPosition;
+	float2 UV = (float2(PixelPosition) + 0.5) / OutputViewSize;
+	float2 Clip = UV * 2 - 1;
+
+	if (any(PixelPosition >= OutputViewSize))
+	{
+		return;
+	}
+
+	// Rotate
+	{
+		float Rot = PI*0.5;
+		float SinRot;
+		float CosRot;
+		sincos(Rot, SinRot, CosRot);
+		Clip = float2(
+			dot(float2(CosRot, -SinRot), Clip),
+			dot(float2(SinRot,  CosRot), Clip)
+		);
+		UV = Clip * 0.5 + 0.5;
+	}
+
+	float HueNorm = UV.y;
+	int HueBin = HueToHueBin(HueNorm);
+
+	float SaturationNorm = UV.x;
+	int SaturationBin = SaturationToSaturationBin(SaturationNorm);
+
+	float3 ColorNormYCbCr = float3(
+		// Arbitrary fade to white at center
+		0.2 * Pow2(1.0 - length(Clip)), 
+		HueNorm, SaturationNorm);
+	float3 Color = clamp(NormalizedYCbCrToRGB(ColorNormYCbCr), 0, 1);
+	float Intensity = dot(Color, 0.333);
+	float Desat = pow(Intensity, DesaturationPower);
+	Color = max(Color, (float3)Desat);
+
+	bool bColorize = (Flags & FLAG_COLORIZE) != 0;
+
+	uint NumSamples = HistogramTexture.Load(int3(SaturationBin, HueBin, 0)).r;
+
+	uint2 InputColorViewSize = InputColorViewRect.zw - InputColorViewRect.xy;
+	uint NumPixels = InputColorViewSize.x * InputColorViewSize.y;
+
+	float3 Output = bColorize ? Color : Luminance(Color);
+
+	Output *= float(NumSamples) / NumPixels;
+	Output *= InputColorViewSize.x; // Somewhat arbitrary, but rescaling with input width makes brightness scale the same as in ColorParade
+	Output *= Brightness;
+
+	// Tonemap
+	Output = pow(Output, Gamma);
+
+	OutputTexture[BufferPixelPosition] = Output;
+}
+
+#endif
diff --git a/Engine/Source/Runtime/Engine/Private/ShowFlags.cpp b/Engine/Source/Runtime/Engine/Private/ShowFlags.cpp
--- a/Engine/Source/Runtime/Engine/Private/ShowFlags.cpp
+++ b/Engine/Source/Runtime/Engine/Private/ShowFlags.cpp
@@ -703,6 +703,11 @@ void EngineShowFlagOverride(EShowFlagInitMode ShowFlagInitMode, EViewModeIndex V
 		EngineShowFlags.SetBloom(false);
 	}
 
+	if (EngineShowFlags.VisualizeColorGrading)
+	{
+		EngineShowFlags.SetCompositeEditorPrimitives(false);
+	}
+
 #if !(UE_BUILD_SHIPPING || UE_BUILD_TEST)
 	{
 		static const auto ICVar = IConsoleManager::Get().FindTConsoleVariableDataInt(TEXT("r.LimitRenderingFeatures"));
diff --git a/Engine/Source/Runtime/Engine/Public/ShowFlags.h b/Engine/Source/Runtime/Engine/Public/ShowFlags.h
--- a/Engine/Source/Runtime/Engine/Public/ShowFlags.h
+++ b/Engine/Source/Runtime/Engine/Public/ShowFlags.h
@@ -414,6 +414,7 @@ private:
 		SetVisualizeVirtualTexture(false);
 		SetVectorFields(false);
 		SetGBufferHints(false);
+		SetVisualizeColorGrading(false);
 		SetCompositeEditorPrimitives(InitMode == ESFIM_Editor || InitMode == ESFIM_VREditing);
 #if UE_ENABLE_DEBUG_DRAWING
 		SetCompositeDebugPrimitives(true);
diff --git a/Engine/Source/Runtime/Engine/Public/ShowFlagsValues.inl b/Engine/Source/Runtime/Engine/Public/ShowFlagsValues.inl
--- a/Engine/Source/Runtime/Engine/Public/ShowFlagsValues.inl
+++ b/Engine/Source/Runtime/Engine/Public/ShowFlagsValues.inl
@@ -94,6 +94,8 @@ SHOWFLAG_ALWAYS_ACCESSIBLE(SpotLights, SFG_LightTypes, NSLOCTEXT("UnrealEd", "Sp
 SHOWFLAG_ALWAYS_ACCESSIBLE(RectLights, SFG_LightTypes, NSLOCTEXT("UnrealEd", "RectLightsSF", "Rect Lights"))
 /** Color correction after tone mapping */
 SHOWFLAG_FIXED_IN_SHIPPING(1, ColorGrading, SFG_PostProcess, NSLOCTEXT("UnrealEd", "ColorGradingSF", "Color Grading"))
+/** Whether to show the color grading visualization displays. */
+SHOWFLAG_ALWAYS_ACCESSIBLE(VisualizeColorGrading, SFG_Visualize, NSLOCTEXT("UnrealEd", "VisualizeColorGrading", "Visualize Color Grading"))
 /** Visualize vector fields. */
 SHOWFLAG_FIXED_IN_SHIPPING(0, VectorFields, SFG_Developer, NSLOCTEXT("UnrealEd", "VectorFieldsSF", "Vector Fields"))
 /** Depth of Field */
diff --git a/Engine/Source/Runtime/Renderer/Internal/PostProcess/ColorGradingScopes.h b/Engine/Source/Runtime/Renderer/Internal/PostProcess/ColorGradingScopes.h
new file mode 100644
--- /dev/null
+++ b/Engine/Source/Runtime/Renderer/Internal/PostProcess/ColorGradingScopes.h
@@ -0,0 +1,119 @@
+// Copyright Epic Games, Inc. All Rights Reserved.
+
+#include "Math/MathFwd.h"
+#include "ScreenPass.h"
+
+class FRDGBuilder;
+class FRDGTexture;
+class FViewInfo;
+class FRDGTextureSRV;
+enum EShaderPlatform : uint16;
+
+namespace UE::ColorParade
+{
+	struct FDrawParameters
+	{
+		// If false, the channels are overlayed on top of eachother (waveform instead of parade)
+		bool bDrawSeparateColumnPerChannel = true;
+		// If false, the graph is grayscale
+		bool bColorize = true;
+		// If true, columns with NaNs and Infs are highlighted in the graph
+		bool bHighlightNaNsAndInfs = true;
+
+		// Range of values that is plotted on the Y-axis
+		float MinValue = 0.0f;
+		float MaxValue = 1.0f;
+
+		// Parameters to control how the histogram values are tonemapped to graph colors
+		float Brightness = 50.0f;
+		float Gamma = 1.0f / 2.2f;
+		float DesaturationPower = 2.5f;
+	};
+
+	bool IsPlatformSupported(EShaderPlatform Platform);
+
+	void Draw(
+		FRDGBuilder& GraphBuilder,
+		FRDGTextureSRV* InputSRV, FIntRect InputViewRect,
+		FRDGTexture* Output, FIntRect OutputViewRect,
+		const FDrawParameters& Parameters);
+
+	void Draw(
+		FRDGBuilder& GraphBuilder,
+		FRDGTextureSRV* InputSRV, FIntRect InputViewRect,
+		FRDGTexture* Output, FIntRect OutputViewRect)
+	{
+		FDrawParameters DrawParameters{};
+		Draw(GraphBuilder, InputSRV, InputViewRect, Output, OutputViewRect, DrawParameters);
+	}
+} // namespace UE::ColorParade
+
+namespace UE::VectorScope
+{
+	struct FDrawParameters
+	{
+		// If false, the graph is grayscale
+		bool bColorize = true;
+
+		// Rotation of the graph, in radians
+		float Rotation = PI / 2.0f;
+
+		// Include only pixel values with luminance in a certain range
+		float LuminanceFilterMin = 0;
+		float LuminanceFilterMax = 1;
+
+		// Parameters to control how the histogram values are tonemapped to graph colors
+		float Brightness = 50.0f;
+		float Gamma = 1.0f / 2.2f;
+		float DesaturationPower = 2.5f;
+	};
+
+	bool IsPlatformSupported(EShaderPlatform Platform);
+
+	void Draw(
+		FRDGBuilder& GraphBuilder,
+		FRDGTextureSRV* InputSRV, FIntRect InputViewRect,
+		FRDGTexture* Output, FIntRect OutputViewRect,
+		const FDrawParameters& Parameters);
+
+	void Draw(
+		FRDGBuilder& GraphBuilder,
+		FRDGTextureSRV* InputSRV, FIntRect InputViewRect,
+		FRDGTexture* Output, FIntRect OutputViewRect)
+	{
+		FDrawParameters DrawParameters{};
+		Draw(GraphBuilder, InputSRV, InputViewRect, Output, OutputViewRect, DrawParameters);
+	}
+} // namespace UE::VectorScope
+
+namespace UE::ColorHistogram
+{
+	struct FDrawParameters
+	{
+		// If false, the channels are overlayed on top of eachother
+		bool bDrawSeparateRowPerChannel = true;
+		// If false, the graph is grayscale
+		bool bColorize = true;
+
+		// Range of values that is plotted on the Y-axis
+		float MinValue = 0.0f;
+		float MaxValue = 1.0f;
+	};
+
+	bool IsPlatformSupported(EShaderPlatform Platform);
+
+	void Draw(
+		FRDGBuilder& GraphBuilder,
+		FRDGTextureSRV* InputSRV, FIntRect InputViewRect,
+		FRDGTexture* Output, FIntRect OutputViewRect,
+		const FDrawParameters& Parameters);
+
+	void Draw(
+		FRDGBuilder& GraphBuilder,
+		FRDGTextureSRV* InputSRV, FIntRect InputViewRect,
+		FRDGTexture* Output, FIntRect OutputViewRect)
+	{
+		FDrawParameters DrawParameters{};
+		Draw(GraphBuilder, InputSRV, InputViewRect, Output, OutputViewRect, DrawParameters);
+	}
+} // namespace UE::ColorHistogram
diff --git a/Engine/Source/Runtime/Renderer/Private/PostProcess/VisualizeColorGrading.cpp b/Engine/Source/Runtime/Renderer/Private/PostProcess/VisualizeColorGrading.cpp
new file mode 100644
--- /dev/null
+++ b/Engine/Source/Runtime/Renderer/Private/PostProcess/VisualizeColorGrading.cpp
@@ -0,0 +1,857 @@
+// Copyright Epic Games, Inc. All Rights Reserved.
+
+#include "PostProcess/VisualizeColorGrading.h"
+#include "PostProcess/ColorGradingScopes.h"
+#include "GlobalShader.h"
+#include "RenderGraphBuilder.h"
+#include "ScreenPass.h"
+#include "RenderGraphUtils.h"
+#include "SceneRendering.h"
+#include "GenerateMips.h"
+#include "CanvasItem.h"
+#include "SceneViewExtension.h"
+#include "PostProcess/PostProcessMaterialInputs.h"
+
+static bool UseWaveOps(EShaderPlatform ShaderPlatform)
+{
+	const ERHIFeatureSupport UseWaveOpsSupport = FDataDrivenShaderPlatformInfo::GetSupportsWaveOperations(ShaderPlatform);
+	return (GRHISupportsWaveOperations && UseWaveOpsSupport == ERHIFeatureSupport::RuntimeDependent) || (UseWaveOpsSupport == ERHIFeatureSupport::RuntimeGuaranteed);
+}
+
+namespace UE::ColorParade
+{
+	constexpr uint32 FlagDrawSeparateColumnPerChannel = 1<<0;
+	constexpr uint32 FlagColorize 					  = 1<<1;
+	constexpr uint32 FlagHighlightNaNsAndInfs		  = 1<<2;
+
+	bool IsPlatformSupported(EShaderPlatform Platform)
+	{
+		return IsFeatureLevelSupported(Platform, ERHIFeatureLevel::SM6);
+	}
+
+	class FBuildColorParadeHistogramCS : public FGlobalShader
+	{
+		DECLARE_GLOBAL_SHADER(FBuildColorParadeHistogramCS);
+		SHADER_USE_PARAMETER_STRUCT(FBuildColorParadeHistogramCS, FGlobalShader);
+
+		class FDimWaveOps : SHADER_PERMUTATION_BOOL("DIM_WAVE_OPS");
+		using FPermutationDomain = TShaderPermutationDomain<FDimWaveOps>;
+
+		inline static const FIntPoint GroupSize{ 1, 256 };
+	
+		// Number of histogram bins per column
+		static constexpr uint32 NumBins = 1024;
+	
+		BEGIN_SHADER_PARAMETER_STRUCT(FParameters, )
+			SHADER_PARAMETER_RDG_TEXTURE_SRV(Texture2D, InputColorTexture)
+			SHADER_PARAMETER_RDG_TEXTURE_UAV(RWTexture2D<uint>, OutputTexture)
+			SHADER_PARAMETER_RDG_BUFFER_UAV(RWStructuredBuffer<uint>, OutputColumnFlagsBuffer)
+			SHADER_PARAMETER(FVector4f, InputColorViewRect)
+			SHADER_PARAMETER(FIntPoint, OutputHistogramSize)
+
+			SHADER_PARAMETER(FVector2f, MinMaxValue)
+			SHADER_PARAMETER(uint32, Flags)
+		END_SHADER_PARAMETER_STRUCT()
+
+		static void ModifyCompilationEnvironment(const FGlobalShaderPermutationParameters& Parameters, FShaderCompilerEnvironment& OutEnvironment)
+		{
+			FGlobalShader::ModifyCompilationEnvironment(Parameters, OutEnvironment);
+
+			OutEnvironment.SetDefine(TEXT("THREADGROUP_SIZE_X"), GroupSize.X);
+			OutEnvironment.SetDefine(TEXT("THREADGROUP_SIZE_Y"), GroupSize.Y);
+
+			OutEnvironment.SetDefine(TEXT("NUM_BINS"), NumBins);
+
+			OutEnvironment.SetDefine(TEXT("FLAG_DRAW_SEPARATE_COLUMN_PER_CHANNEL"), FlagDrawSeparateColumnPerChannel);
+			OutEnvironment.SetDefine(TEXT("FLAG_COLORIZE"), FlagColorize);
+			OutEnvironment.SetDefine(TEXT("FLAG_HIGHLIGHT_NANS_AND_INFS"), FlagHighlightNaNsAndInfs);
+		}
+
+		static bool ShouldCompilePermutation(const FGlobalShaderPermutationParameters& Parameters)
+		{
+			FPermutationDomain PermutationVector(Parameters.PermutationId);
+
+			ERHIFeatureSupport WaveOpsSupport = FDataDrivenShaderPlatformInfo::GetSupportsWaveOperations(Parameters.Platform);
+			if (PermutationVector.Get<FDimWaveOps>())
+			{
+				if (WaveOpsSupport == ERHIFeatureSupport::Unsupported)
+				{
+					return false;
+				}
+			}
+			else
+			{
+				if (WaveOpsSupport == ERHIFeatureSupport::RuntimeGuaranteed)
+				{
+					return false;
+				}
+			}
+
+			return IsPlatformSupported(Parameters.Platform);
+		}
+	};
+	IMPLEMENT_GLOBAL_SHADER(FBuildColorParadeHistogramCS, "/Engine/Private/Tools/ColorGradingParade.usf", "BuildColorParadeHistogramCS", SF_Compute);
+
+	class FDrawColorParadeCS : public FGlobalShader
+	{
+		DECLARE_GLOBAL_SHADER(FDrawColorParadeCS);
+		SHADER_USE_PARAMETER_STRUCT(FDrawColorParadeCS, FGlobalShader);
+
+		inline static const FIntPoint GroupSize{ 8, 8 };
+
+		BEGIN_SHADER_PARAMETER_STRUCT(FParameters, )
+			SHADER_PARAMETER_RDG_TEXTURE(Texture2D<uint>, HistogramTexture)
+			SHADER_PARAMETER_SAMPLER(SamplerState, HistogramTextureSampler)
+			SHADER_PARAMETER_RDG_BUFFER_SRV(StructuredBuffer<uint>, ColumnFlagsBuffer)
+			SHADER_PARAMETER_RDG_TEXTURE_UAV(RWTexture2D<float3>, OutputTexture)
+			SHADER_PARAMETER(FIntPoint, InputViewSize)
+			SHADER_PARAMETER(FIntPoint, HistogramSize)
+			SHADER_PARAMETER(FIntVector4, OutputViewRect)
+			SHADER_PARAMETER(float, Brightness)
+			SHADER_PARAMETER(float, Gamma)
+			SHADER_PARAMETER(float, DesaturationPower)
+
+			SHADER_PARAMETER(FVector2f, MinMaxValue)
+			SHADER_PARAMETER(uint32, Flags)
+		END_SHADER_PARAMETER_STRUCT()
+
+		static void ModifyCompilationEnvironment(const FGlobalShaderPermutationParameters& Parameters, FShaderCompilerEnvironment& OutEnvironment)
+		{
+			FGlobalShader::ModifyCompilationEnvironment(Parameters, OutEnvironment);
+
+			OutEnvironment.SetDefine(TEXT("THREADGROUP_SIZE_X"), GroupSize.X);
+			OutEnvironment.SetDefine(TEXT("THREADGROUP_SIZE_Y"), GroupSize.Y);
+
+			OutEnvironment.SetDefine(TEXT("NUM_BINS"), FBuildColorParadeHistogramCS::NumBins);
+
+			OutEnvironment.SetDefine(TEXT("FLAG_DRAW_SEPARATE_COLUMN_PER_CHANNEL"), FlagDrawSeparateColumnPerChannel);
+			OutEnvironment.SetDefine(TEXT("FLAG_COLORIZE"), FlagColorize);
+			OutEnvironment.SetDefine(TEXT("FLAG_HIGHLIGHT_NANS_AND_INFS"), FlagHighlightNaNsAndInfs);
+		}
+
+		static bool ShouldCompilePermutation(const FGlobalShaderPermutationParameters& Parameters)
+		{
+			return IsPlatformSupported(Parameters.Platform);
+		}
+	};
+	IMPLEMENT_GLOBAL_SHADER(FDrawColorParadeCS, "/Engine/Private/Tools/ColorGradingParade.usf", "DrawColorParadeCS", SF_Compute);
+
+	void Draw(
+		FRDGBuilder& GraphBuilder,
+		FRDGTextureSRVRef InputSRV, FIntRect InputViewRect,
+		FRDGTextureRef Output, FIntRect OutputViewRect,
+		const FDrawParameters& Parameters)
+	{
+		constexpr int NumComponents = 3;
+
+		check(IsPlatformSupported(GMaxRHIShaderPlatform));
+		check(InputSRV);
+		check(Output);
+		// The histogram is stored as a uint32 texture, with one value per channel, to support atomics.
+		// This defines a limit on the output rect.
+		constexpr int MaxWidth = (1U << 16U) / NumComponents;
+		constexpr int MaxHeight = (1U << 16U);
+		check(OutputViewRect.Size().X < MaxWidth);
+		check(OutputViewRect.Size().Y < MaxHeight);
+		check(InputViewRect.Size().X > 0);
+		check(InputViewRect.Size().Y > 0);
+		check(OutputViewRect.Size().X > 0);
+		check(OutputViewRect.Size().Y > 0);
+
+		RDG_EVENT_SCOPE(GraphBuilder, "ColorParade");
+
+		FIntPoint OutputSize = OutputViewRect.Size();
+		FIntPoint HistogramSize(OutputSize.X, FBuildColorParadeHistogramCS::NumBins);
+
+		FGlobalShaderMap* ShaderMap = GetGlobalShaderMap(GMaxRHIFeatureLevel);
+
+		uint32 Flags = 0;
+		Flags |= Parameters.bDrawSeparateColumnPerChannel ? FlagDrawSeparateColumnPerChannel : 0U;
+		Flags |= Parameters.bColorize ? FlagColorize : 0U;
+		Flags |= Parameters.bHighlightNaNsAndInfs ? FlagHighlightNaNsAndInfs : 0U;
+
+		FVector2f MinMaxValue{ Parameters.MinValue, Parameters.MaxValue };
+
+		// Build a per-column histogram of pixel values.
+		FRDGTextureRef Histogram = nullptr;
+		FRDGBufferRef ColumnFlags = nullptr;
+		{
+			FIntPoint HistogramTextureSize(HistogramSize.X * NumComponents, HistogramSize.Y);
+
+			FRDGTextureDesc HistogramDesc = FRDGTextureDesc::Create2D(
+				HistogramTextureSize,
+				PF_R32_UINT, // used for support of atomics
+				FClearValueBinding::Black,
+				TexCreate_ShaderResource | TexCreate_UAV);
+
+			Histogram = GraphBuilder.CreateTexture(HistogramDesc, TEXT("ColorParadeHistogram"));
+			FRDGTextureUAVRef HistogramUAV = GraphBuilder.CreateUAV(Histogram);
+
+			FRDGBufferDesc ColumnFlagsDesc = FRDGBufferDesc::CreateStructuredDesc(sizeof(uint32), HistogramTextureSize.X);
+			ColumnFlags = GraphBuilder.CreateBuffer(ColumnFlagsDesc, TEXT("ColumnFlags"));
+			FRDGBufferUAVRef ColumnFlagsUAV = GraphBuilder.CreateUAV(ColumnFlags);
+
+			AddClearUAVPass(GraphBuilder, HistogramUAV, 0U);
+			AddClearUAVPass(GraphBuilder, ColumnFlagsUAV, 0U);
+
+			FBuildColorParadeHistogramCS::FParameters* PassParameters = GraphBuilder.AllocParameters<FBuildColorParadeHistogramCS::FParameters>();
+			PassParameters->InputColorTexture = InputSRV;
+			PassParameters->OutputTexture = HistogramUAV;
+			PassParameters->OutputColumnFlagsBuffer = ColumnFlagsUAV;
+			PassParameters->InputColorViewRect = FVector4f(InputViewRect.Min, InputViewRect.Max);
+			PassParameters->OutputHistogramSize = HistogramSize;
+			PassParameters->MinMaxValue = MinMaxValue;
+			PassParameters->Flags = Flags;
+
+			FBuildColorParadeHistogramCS::FPermutationDomain PermutationVector;
+			PermutationVector.Set<FBuildColorParadeHistogramCS::FDimWaveOps>(UseWaveOps(GMaxRHIShaderPlatform));
+			auto ComputeShader = ShaderMap->GetShader<FBuildColorParadeHistogramCS>(PermutationVector);
+
+			FComputeShaderUtils::AddPass(GraphBuilder,
+										 RDG_EVENT_NAME("BuildColorParadeHistogramCS"),
+										 ComputeShader,
+										 PassParameters,
+										 // One thread handles N input values, where N = InputWidth/OutputWidth.
+										 FComputeShaderUtils::GetGroupCount(FIntPoint(OutputSize.X, InputViewRect.Size().Y), FBuildColorParadeHistogramCS::GroupSize));
+		}
+
+		// Draw this histogram data into a nice graph
+		{
+			FRDGTextureUAVRef RGBParadeUAV = GraphBuilder.CreateUAV(Output);
+
+			FDrawColorParadeCS::FParameters* PassParameters = GraphBuilder.AllocParameters<FDrawColorParadeCS::FParameters>();
+			PassParameters->HistogramTexture = Histogram;
+			PassParameters->HistogramTextureSampler = TStaticSamplerState<SF_Bilinear, AM_Clamp, AM_Clamp, AM_Clamp>::GetRHI();
+			PassParameters->ColumnFlagsBuffer = GraphBuilder.CreateSRV(ColumnFlags);
+			PassParameters->OutputTexture = RGBParadeUAV;
+			PassParameters->InputViewSize = InputViewRect.Size();
+			PassParameters->HistogramSize = HistogramSize;
+			PassParameters->OutputViewRect = FIntVector4(OutputViewRect.Min.X, OutputViewRect.Min.Y, OutputViewRect.Max.X, OutputViewRect.Max.Y);
+			PassParameters->Brightness = Parameters.Brightness;
+			PassParameters->Gamma = Parameters.Gamma;
+			PassParameters->DesaturationPower = Parameters.DesaturationPower;
+			PassParameters->MinMaxValue = MinMaxValue;
+			PassParameters->Flags = Flags;
+
+			auto ComputeShader = ShaderMap->GetShader<FDrawColorParadeCS>();
+
+			FComputeShaderUtils::AddPass(GraphBuilder,
+										 RDG_EVENT_NAME("DrawColorParadeCS"),
+										 ComputeShader,
+										 PassParameters,
+										 FComputeShaderUtils::GetGroupCount(OutputSize, FDrawColorParadeCS::GroupSize));
+		}
+	}
+} // namespace UE::ColorParade
+
+namespace UE::VectorScope
+{
+	constexpr uint32 FlagColorize = 1<<0;
+	constexpr uint32 HueResolution = 1024;
+	constexpr uint32 SaturationResolution = 1024;
+
+	bool IsPlatformSupported(EShaderPlatform Platform)
+	{
+		return IsFeatureLevelSupported(Platform, ERHIFeatureLevel::SM6);
+	}
+
+	class FBuildVectorScopeHistogramCS : public FGlobalShader
+	{
+		DECLARE_GLOBAL_SHADER(FBuildVectorScopeHistogramCS);
+		SHADER_USE_PARAMETER_STRUCT(FBuildVectorScopeHistogramCS, FGlobalShader);
+
+		inline static const FIntPoint GroupSize{ 256, 1 };
+		
+		BEGIN_SHADER_PARAMETER_STRUCT(FParameters, )
+			SHADER_PARAMETER_RDG_TEXTURE_SRV(Texture2D, InputColorTexture)
+			SHADER_PARAMETER(FIntVector4, InputColorViewRect)
+			SHADER_PARAMETER(FVector2f, LuminanceFilterMinMax)
+			SHADER_PARAMETER_RDG_TEXTURE_UAV(RWTexture2D<uint>, OutputHistogramTexture)
+		END_SHADER_PARAMETER_STRUCT()
+
+		static void ModifyCompilationEnvironment(const FGlobalShaderPermutationParameters& Parameters, FShaderCompilerEnvironment& OutEnvironment)
+		{
+			FGlobalShader::ModifyCompilationEnvironment(Parameters, OutEnvironment);
+
+			OutEnvironment.SetDefine(TEXT("THREADGROUP_SIZE_X"), GroupSize.X);
+			OutEnvironment.SetDefine(TEXT("THREADGROUP_SIZE_Y"), GroupSize.Y);
+
+			OutEnvironment.SetDefine(TEXT("HUE_RESOLUTION"), HueResolution);
+			OutEnvironment.SetDefine(TEXT("SATURATION_RESOLUTION"), SaturationResolution);
+		}
+
+		static bool ShouldCompilePermutation(const FGlobalShaderPermutationParameters& Parameters)
+		{
+			return IsPlatformSupported(Parameters.Platform);
+		}
+	};
+	IMPLEMENT_GLOBAL_SHADER(FBuildVectorScopeHistogramCS, "/Engine/Private/Tools/ColorGradingVectorScope.usf", "BuildVectorScopeHistogramCS", SF_Compute);
+
+	class FDrawVectorScopeCS : public FGlobalShader
+	{
+		DECLARE_GLOBAL_SHADER(FDrawVectorScopeCS);
+		SHADER_USE_PARAMETER_STRUCT(FDrawVectorScopeCS, FGlobalShader);
+
+		inline static const FIntPoint GroupSize{ 8, 8 };
+	
+		BEGIN_SHADER_PARAMETER_STRUCT(FParameters, )
+			SHADER_PARAMETER_RDG_TEXTURE(Texture2D<uint>, HistogramTexture)
+			SHADER_PARAMETER_RDG_TEXTURE_UAV(RWTexture2D<float3>, OutputTexture)
+			SHADER_PARAMETER(FIntVector4, OutputViewRect)
+			SHADER_PARAMETER(FIntVector4, InputColorViewRect)
+			SHADER_PARAMETER(uint32, Flags)
+			SHADER_PARAMETER(float, Rotation)
+			SHADER_PARAMETER(float, Brightness)
+			SHADER_PARAMETER(float, Gamma)
+			SHADER_PARAMETER(float, DesaturationPower)
+		END_SHADER_PARAMETER_STRUCT()
+
+		static void ModifyCompilationEnvironment(const FGlobalShaderPermutationParameters& Parameters, FShaderCompilerEnvironment& OutEnvironment)
+		{
+			FGlobalShader::ModifyCompilationEnvironment(Parameters, OutEnvironment);
+
+			OutEnvironment.SetDefine(TEXT("THREADGROUP_SIZE_X"), GroupSize.X);
+			OutEnvironment.SetDefine(TEXT("THREADGROUP_SIZE_Y"), GroupSize.Y);
+
+			OutEnvironment.SetDefine(TEXT("HUE_RESOLUTION"), HueResolution);
+			OutEnvironment.SetDefine(TEXT("SATURATION_RESOLUTION"), SaturationResolution);
+
+			OutEnvironment.SetDefine(TEXT("FLAG_COLORIZE"), FlagColorize);
+		}
+
+		static bool ShouldCompilePermutation(const FGlobalShaderPermutationParameters& Parameters)
+		{
+			return IsPlatformSupported(Parameters.Platform);
+		}
+	};
+	IMPLEMENT_GLOBAL_SHADER(FDrawVectorScopeCS, "/Engine/Private/Tools/ColorGradingVectorScope.usf", "DrawVectorScopeCS", SF_Compute);
+
+	void Draw(
+		FRDGBuilder& GraphBuilder,
+		FRDGTextureSRVRef InputSRV, FIntRect InputViewRect,
+		FRDGTextureRef Output, FIntRect OutputViewRect,
+		const FDrawParameters& Parameters)
+	{
+		check(IsPlatformSupported(GMaxRHIShaderPlatform));
+		check(InputSRV);
+		check(Output);
+		check(InputViewRect.Size().X > 0);
+		check(InputViewRect.Size().Y > 0);
+		check(OutputViewRect.Size().X > 0);
+		check(OutputViewRect.Size().Y > 0);
+		
+		RDG_EVENT_SCOPE(GraphBuilder, "VectorScope");
+
+		FGlobalShaderMap* ShaderMap = GetGlobalShaderMap(GMaxRHIFeatureLevel);
+
+		FIntPoint InputColorViewSize = InputViewRect.Size();
+
+		uint32 Flags = 0;
+		Flags |= Parameters.bColorize ? FlagColorize : 0;
+
+		FRDGTextureRef HistogramTexture = nullptr;
+		{
+			FRDGTextureDesc HistogramTextureDesc = FRDGTextureDesc::Create2D(
+				FIntPoint(SaturationResolution, HueResolution),
+				PF_R32_UINT,
+				FClearValueBinding::Black,
+				TexCreate_ShaderResource | TexCreate_UAV);
+
+			HistogramTexture = GraphBuilder.CreateTexture(HistogramTextureDesc, TEXT("VectorScopeHistogram"));
+
+			AddClearUAVPass(GraphBuilder, GraphBuilder.CreateUAV(HistogramTexture), 0U);
+
+			FBuildVectorScopeHistogramCS::FParameters* PassParameters = GraphBuilder.AllocParameters<FBuildVectorScopeHistogramCS::FParameters>();
+			PassParameters->InputColorTexture = InputSRV;
+			PassParameters->InputColorViewRect = FIntVector4(InputViewRect.Min.X, InputViewRect.Min.Y, InputViewRect.Max.X, InputViewRect.Max.Y);
+			PassParameters->LuminanceFilterMinMax = FVector2f(Parameters.LuminanceFilterMin, Parameters.LuminanceFilterMax);
+			PassParameters->OutputHistogramTexture = GraphBuilder.CreateUAV(HistogramTexture);
+
+			auto ComputeShader = ShaderMap->GetShader<FBuildVectorScopeHistogramCS>();
+
+			FComputeShaderUtils::AddPass(GraphBuilder,
+										 RDG_EVENT_NAME("BuildVectorScopeHistogramCS"),
+										 ComputeShader,
+										 PassParameters,
+										 FComputeShaderUtils::GetGroupCount(InputColorViewSize.X * InputColorViewSize.Y, FBuildVectorScopeHistogramCS::GroupSize.X));
+		}
+
+		{
+			FDrawVectorScopeCS::FParameters* PassParameters = GraphBuilder.AllocParameters<FDrawVectorScopeCS::FParameters>();
+			PassParameters->HistogramTexture = HistogramTexture;
+			PassParameters->OutputTexture = GraphBuilder.CreateUAV(Output);
+			PassParameters->OutputViewRect = FIntVector4(OutputViewRect.Min.X, OutputViewRect.Min.Y, OutputViewRect.Max.X, OutputViewRect.Max.Y);
+			PassParameters->InputColorViewRect = FIntVector4(InputViewRect.Min.X, InputViewRect.Min.Y, InputViewRect.Max.X, InputViewRect.Max.Y);
+			PassParameters->Rotation = Parameters.Rotation;
+			PassParameters->Brightness = Parameters.Brightness;
+			PassParameters->Gamma = Parameters.Gamma;
+			PassParameters->DesaturationPower = Parameters.DesaturationPower;
+			PassParameters->Flags = Flags;
+
+			auto ComputeShader = ShaderMap->GetShader<FDrawVectorScopeCS>();
+
+			FComputeShaderUtils::AddPass(GraphBuilder,
+										 RDG_EVENT_NAME("DrawVectorScopeCS"),
+										 ComputeShader,
+										 PassParameters,
+										 FComputeShaderUtils::GetGroupCount(OutputViewRect.Size(), FDrawVectorScopeCS::GroupSize));
+		}
+	}
+} // namespace UE::VectorScope
+
+namespace UE::ColorHistogram
+{
+	constexpr uint32 FlagDrawSeparateRowPerChannel = 1<<0;
+	constexpr uint32 FlagColorize 				   = 1<<1;
+
+	constexpr uint32 PerChannelResolution = 1024;
+
+	bool IsPlatformSupported(EShaderPlatform Platform)
+	{
+		return IsFeatureLevelSupported(Platform, ERHIFeatureLevel::SM6);
+	}
+
+	class FBuildColorHistogramCS : public FGlobalShader
+	{
+		DECLARE_GLOBAL_SHADER(FBuildColorHistogramCS);
+		SHADER_USE_PARAMETER_STRUCT(FBuildColorHistogramCS, FGlobalShader);
+
+		inline static const FIntPoint GroupSize{ 16, 16 };
+		
+		BEGIN_SHADER_PARAMETER_STRUCT(FParameters, )
+			SHADER_PARAMETER_RDG_TEXTURE_SRV(Texture2D, InputColorTexture)
+			SHADER_PARAMETER(FIntVector4, InputColorViewRect)
+			SHADER_PARAMETER(FVector2f, MinMaxValue)
+			SHADER_PARAMETER_RDG_TEXTURE_UAV(RWTexture2D<uint>, OutputHistogramTexture)
+			SHADER_PARAMETER_RDG_BUFFER_UAV(RWStructuredBuffer<uint>, OutputMetaDataBuffer)
+		END_SHADER_PARAMETER_STRUCT()
+
+		static void ModifyCompilationEnvironment(const FGlobalShaderPermutationParameters& Parameters, FShaderCompilerEnvironment& OutEnvironment)
+		{
+			FGlobalShader::ModifyCompilationEnvironment(Parameters, OutEnvironment);
+
+			OutEnvironment.SetDefine(TEXT("THREADGROUP_SIZE_X"), GroupSize.X);
+			OutEnvironment.SetDefine(TEXT("THREADGROUP_SIZE_Y"), GroupSize.Y);
+
+			OutEnvironment.SetDefine(TEXT("PER_CHANNEL_RESOLUTION"), PerChannelResolution);
+		}
+
+		static bool ShouldCompilePermutation(const FGlobalShaderPermutationParameters& Parameters)
+		{
+			return IsPlatformSupported(Parameters.Platform);
+		}
+	};
+	IMPLEMENT_GLOBAL_SHADER(FBuildColorHistogramCS, "/Engine/Private/Tools/ColorGradingHistogram.usf", "BuildColorHistogramCS", SF_Compute);
+
+	class FDrawColorHistogramCS : public FGlobalShader
+	{
+		DECLARE_GLOBAL_SHADER(FDrawColorHistogramCS);
+		SHADER_USE_PARAMETER_STRUCT(FDrawColorHistogramCS, FGlobalShader);
+
+		inline static const FIntPoint GroupSize{ 8, 8 };
+	
+		BEGIN_SHADER_PARAMETER_STRUCT(FParameters, )
+			SHADER_PARAMETER_RDG_TEXTURE(Texture2D<uint>, HistogramTexture)
+			SHADER_PARAMETER_RDG_BUFFER_SRV(StructuredBuffer<uint>, MetaDataBuffer)
+			SHADER_PARAMETER_RDG_TEXTURE_UAV(RWTexture2D<float3>, OutputTexture)
+			SHADER_PARAMETER(FIntVector4, OutputViewRect)
+			SHADER_PARAMETER(FIntVector4, InputColorViewRect)
+			SHADER_PARAMETER(uint32, Flags)
+		END_SHADER_PARAMETER_STRUCT()
+
+		static void ModifyCompilationEnvironment(const FGlobalShaderPermutationParameters& Parameters, FShaderCompilerEnvironment& OutEnvironment)
+		{
+			FGlobalShader::ModifyCompilationEnvironment(Parameters, OutEnvironment);
+
+			OutEnvironment.SetDefine(TEXT("THREADGROUP_SIZE_X"), GroupSize.X);
+			OutEnvironment.SetDefine(TEXT("THREADGROUP_SIZE_Y"), GroupSize.Y);
+
+			OutEnvironment.SetDefine(TEXT("PER_CHANNEL_RESOLUTION"), PerChannelResolution);
+
+			OutEnvironment.SetDefine(TEXT("FLAG_COLORIZE"), FlagColorize);
+			OutEnvironment.SetDefine(TEXT("FLAG_DRAW_SEPARATE_ROW_PER_CHANNEL"), FlagDrawSeparateRowPerChannel);
+		}
+
+		static bool ShouldCompilePermutation(const FGlobalShaderPermutationParameters& Parameters)
+		{
+			return IsPlatformSupported(Parameters.Platform);
+		}
+	};
+	IMPLEMENT_GLOBAL_SHADER(FDrawColorHistogramCS, "/Engine/Private/Tools/ColorGradingHistogram.usf", "DrawColorHistogramCS", SF_Compute);
+
+	void Draw(
+		FRDGBuilder& GraphBuilder,
+		FRDGTextureSRVRef InputSRV, FIntRect InputViewRect,
+		FRDGTextureRef Output, FIntRect OutputViewRect,
+		const FDrawParameters& Parameters)
+	{
+		check(IsPlatformSupported(GMaxRHIShaderPlatform));
+		check(InputSRV);
+		check(Output);
+		check(InputViewRect.Size().X > 0);
+		check(InputViewRect.Size().Y > 0);
+		check(OutputViewRect.Size().X > 0);
+		check(OutputViewRect.Size().Y > 0);
+		
+		RDG_EVENT_SCOPE(GraphBuilder, "ColorHistogram");
+
+		FGlobalShaderMap* ShaderMap = GetGlobalShaderMap(GMaxRHIFeatureLevel);
+
+		FIntPoint InputColorViewSize = InputViewRect.Size();
+
+		uint32 Flags = 0;
+		Flags |= Parameters.bDrawSeparateRowPerChannel ? FlagDrawSeparateRowPerChannel : 0;
+		Flags |= Parameters.bColorize ? FlagColorize : 0;
+
+		FVector2f MinMaxValue{ Parameters.MinValue, Parameters.MaxValue };
+
+		FRDGTextureRef HistogramTexture = nullptr;
+		FRDGBufferRef MetaDataBuffer = nullptr;
+		{
+			FRDGTextureDesc HistogramTextureDesc = FRDGTextureDesc::Create2D(
+				FIntPoint(PerChannelResolution, 3),
+				PF_R32_UINT,
+				FClearValueBinding::Black,
+				TexCreate_ShaderResource | TexCreate_UAV);
+
+			HistogramTexture = GraphBuilder.CreateTexture(HistogramTextureDesc, TEXT("ColorHistogram"));
+			FRDGTextureUAVRef HistogramTextureUAV = GraphBuilder.CreateUAV(HistogramTexture);
+			
+			FRDGBufferDesc MetaDataBufferDesc = FRDGBufferDesc::CreateStructuredDesc(sizeof(uint32), 3);
+			MetaDataBuffer = GraphBuilder.CreateBuffer(MetaDataBufferDesc, TEXT("MetaData"));
+			FRDGBufferUAVRef MetaDataBufferUAV =  GraphBuilder.CreateUAV(MetaDataBuffer);
+
+			AddClearUAVPass(GraphBuilder, GraphBuilder.CreateUAV(HistogramTexture), 0U);
+			AddClearUAVPass(GraphBuilder, MetaDataBufferUAV, 0U);
+
+			FBuildColorHistogramCS::FParameters* PassParameters = GraphBuilder.AllocParameters<FBuildColorHistogramCS::FParameters>();
+			PassParameters->InputColorTexture = InputSRV;
+			PassParameters->InputColorViewRect = FIntVector4(InputViewRect.Min.X, InputViewRect.Min.Y, InputViewRect.Max.X, InputViewRect.Max.Y);
+			PassParameters->MinMaxValue = MinMaxValue;
+			PassParameters->OutputHistogramTexture = HistogramTextureUAV;
+			PassParameters->OutputMetaDataBuffer = MetaDataBufferUAV;
+
+			auto ComputeShader = ShaderMap->GetShader<FBuildColorHistogramCS>();
+
+			FComputeShaderUtils::AddPass(GraphBuilder,
+										 RDG_EVENT_NAME("BuildColorHistogramCS"),
+										 ComputeShader,
+										 PassParameters,
+										 FComputeShaderUtils::GetGroupCount(InputColorViewSize, FBuildColorHistogramCS::GroupSize));
+		}
+
+		{
+			FDrawColorHistogramCS::FParameters* PassParameters = GraphBuilder.AllocParameters<FDrawColorHistogramCS::FParameters>();
+			PassParameters->HistogramTexture = HistogramTexture;
+			PassParameters->MetaDataBuffer = GraphBuilder.CreateSRV(MetaDataBuffer);
+			PassParameters->OutputTexture = GraphBuilder.CreateUAV(Output);
+			PassParameters->OutputViewRect = FIntVector4(OutputViewRect.Min.X, OutputViewRect.Min.Y, OutputViewRect.Max.X, OutputViewRect.Max.Y);
+			PassParameters->InputColorViewRect = FIntVector4(InputViewRect.Min.X, InputViewRect.Min.Y, InputViewRect.Max.X, InputViewRect.Max.Y);
+			PassParameters->Flags = Flags;
+
+			auto ComputeShader = ShaderMap->GetShader<FDrawColorHistogramCS>();
+
+			FComputeShaderUtils::AddPass(GraphBuilder,
+										 RDG_EVENT_NAME("DrawColorHistogramCS"),
+										 ComputeShader,
+										 PassParameters,
+										 FComputeShaderUtils::GetGroupCount(OutputViewRect.Size(), FDrawColorHistogramCS::GroupSize));
+		}
+	}
+} // namespace UE::ColorHistogram
+
+static FScreenPassTexture VisualizeColorGrading(FRDGBuilder& GraphBuilder, const FSceneView& View, FScreenPassTexture& SceneColor, FScreenPassRenderTarget& OverrideOutput)
+{
+	RDG_EVENT_SCOPE(GraphBuilder, "VisualizeColorGrading");
+
+	auto FinalizeOutput = [](FRDGBuilder& GraphBuilder, const FSceneView& View, FScreenPassTexture& Output, FScreenPassRenderTarget& OverrideOutput) -> FScreenPassTexture
+	{
+		if (OverrideOutput.IsValid())
+		{
+			AddDrawTexturePass(GraphBuilder, View, Output, OverrideOutput);
+			return OverrideOutput;
+		}
+
+		return Output;
+	};
+
+	FRDGTextureSRVRef InputSRV = GraphBuilder.CreateSRV(SceneColor.Texture);
+	FScreenPassTexture Output = MoveTemp(SceneColor);
+
+	FRHISamplerState* BilinearClampSampler = TStaticSamplerState<SF_Bilinear, AM_Clamp, AM_Clamp, AM_Clamp>::GetRHI();
+
+	// Color Parade
+	UE::ColorParade::FDrawParameters RGBParadeDrawParameters
+	{
+		.MinValue = -0.1f,
+		.MaxValue =  1.1f,
+		.Gamma 	  = 1.0f / GEngine->DisplayGamma
+	};
+	FRDGTextureRef RGBParade = nullptr;
+	FIntRect RGBParadeRect{};
+	if (UE::ColorParade::IsPlatformSupported(GMaxRHIShaderPlatform))
+	{
+		int DownsampleFactor = 4;
+
+		int Width = FMath::Min(Output.ViewRect.Size().X, Output.ViewRect.Size().Y) / 2;
+		FIntPoint Size {Width, (Output.ViewRect.Size().Y * Width) / Output.ViewRect.Size().X};
+		FIntPoint RenderSize = Size * DownsampleFactor;
+		RenderSize = FIntPoint { FMath::Max(RenderSize.X, 1), FMath::Max(RenderSize.Y, 1) };
+		uint8 NumMips = DownsampleFactor;
+
+		FIntPoint OutputPosition{ Output.ViewRect.Max.X - Size.X, Output.ViewRect.Max.Y - Size.Y };
+		RGBParadeRect = FIntRect(OutputPosition, OutputPosition + Size);
+
+		FRDGTextureDesc RGBParadeDesc = FRDGTextureDesc::Create2D(
+			RenderSize,
+			PF_FloatRGB,
+			FClearValueBinding::Black,
+			TexCreate_ShaderResource | TexCreate_UAV,
+			NumMips);
+
+		RGBParade = GraphBuilder.CreateTexture(RGBParadeDesc, TEXT("ColorParade"));
+
+		UE::ColorParade::Draw(GraphBuilder, InputSRV, Output.ViewRect, RGBParade, FIntRect(0, 0, RenderSize.X, RenderSize.Y), RGBParadeDrawParameters);
+
+		FGenerateMips::Execute(GraphBuilder, GMaxRHIFeatureLevel, RGBParade, BilinearClampSampler);
+	}
+
+	// VectorScope
+	UE::VectorScope::FDrawParameters VectorScopeDrawParameters
+	{
+		.Gamma = 1.0f / GEngine->DisplayGamma
+	};
+	FRDGTextureRef VectorScope = nullptr;
+	FIntRect VectorScopeRect {};
+	if (UE::VectorScope::IsPlatformSupported(GMaxRHIShaderPlatform))
+	{
+		int DownsampleFactor = 4;
+
+		int Width = FMath::Min(Output.ViewRect.Size().X, Output.ViewRect.Size().Y) / 3;
+		FIntPoint Size{ Width, Width };
+		FIntPoint RenderSize = Size * DownsampleFactor;
+		RenderSize = FIntPoint { FMath::Max(RenderSize.X, 1), FMath::Max(RenderSize.Y, 1) };
+		uint8 NumMips = DownsampleFactor;
+
+		FIntPoint OutputPosition{ 0, Output.ViewRect.Max.Y - Size.Y };
+		VectorScopeRect = FIntRect(OutputPosition, OutputPosition + Size);
+
+		FRDGTextureDesc VectorScopeDesc = FRDGTextureDesc::Create2D(
+			RenderSize,
+			PF_FloatRGB,
+			FClearValueBinding::Black,
+			TexCreate_ShaderResource | TexCreate_UAV,
+			NumMips);
+
+		VectorScope = GraphBuilder.CreateTexture(VectorScopeDesc, TEXT("VectorScope"));
+
+		UE::VectorScope::Draw(GraphBuilder, InputSRV, Output.ViewRect, VectorScope, FIntRect(0, 0, RenderSize.X, RenderSize.Y), VectorScopeDrawParameters);
+
+		FGenerateMips::Execute(GraphBuilder, GMaxRHIFeatureLevel, VectorScope, BilinearClampSampler);
+	}
+
+	// ColorHistogram
+	UE::ColorHistogram::FDrawParameters ColorHistogramDrawParameters{};
+	FRDGTextureRef ColorHistogram = nullptr;
+	FIntRect ColorHistogramRect {};
+	if (UE::ColorHistogram::IsPlatformSupported(GMaxRHIShaderPlatform))
+	{
+		int DownsampleFactor = 4;
+
+		int Width = FMath::Min(Output.ViewRect.Size().X, Output.ViewRect.Size().Y) / 3;
+		FIntPoint Size{ Width, Width / 2 };
+		FIntPoint RenderSize = Size * DownsampleFactor;
+		RenderSize = FIntPoint { FMath::Max(RenderSize.X, 1), FMath::Max(RenderSize.Y, 1) };
+		uint8 NumMips = DownsampleFactor;
+
+		FIntPoint OutputPosition{ Output.ViewRect.Max.X - Size.X, 0 };
+		ColorHistogramRect = FIntRect(OutputPosition, OutputPosition + Size);
+
+		FRDGTextureDesc ColorHistogramDesc = FRDGTextureDesc::Create2D(
+			RenderSize,
+			PF_FloatRGB,
+			FClearValueBinding::Black,
+			TexCreate_ShaderResource | TexCreate_UAV,
+			NumMips);
+
+		ColorHistogram = GraphBuilder.CreateTexture(ColorHistogramDesc, TEXT("ColorHistogramTexture"));
+
+		UE::ColorHistogram::Draw(GraphBuilder, InputSRV, Output.ViewRect, ColorHistogram, FIntRect(0, 0, RenderSize.X, RenderSize.Y), ColorHistogramDrawParameters);
+
+		FGenerateMips::Execute(GraphBuilder, GMaxRHIFeatureLevel, ColorHistogram, BilinearClampSampler);
+	}
+
+	if (RGBParade)
+	{
+		FIntPoint InputPosition = FIntPoint::ZeroValue;
+		AddDrawTexturePass(GraphBuilder, View, RGBParade, Output.Texture, InputPosition, RGBParade->Desc.Extent, RGBParadeRect.Min, RGBParadeRect.Size(), BilinearClampSampler);
+
+		AddDrawCanvasPass(GraphBuilder, RDG_EVENT_NAME("ColorParade UI"), View, FScreenPassRenderTarget(Output, ERenderTargetLoadAction::ELoad),
+		[&](FCanvas& Canvas)
+		{
+			UFont* Font = GEngine->GetMonospaceFont();
+
+			const float DPIScale = Canvas.GetDPIScale();
+			Canvas.SetBaseTransform(FMatrix(
+				FScaleMatrix(DPIScale)
+				* FTranslationMatrix(FVector(RGBParadeRect.Min.X, RGBParadeRect.Min.Y, 0))
+				* Canvas.CalcBaseTransform2D(Canvas.GetViewRect().Width(), Canvas.GetViewRect().Height())));
+
+			auto DrawShadowedString = [&](float X, float Y, FStringView Text, const FLinearColor& Color = FLinearColor::White)
+			{
+				Canvas.DrawShadowedString(X / DPIScale, Y / DPIScale, Text, Font, Color);
+			};
+		
+			auto GetStringSize = [&](FStringView StringView) -> FVector2f
+			{
+				int32 XL = 0;
+				int32 YL = 0;
+				StringSize(Font, XL, YL, StringView);
+				return FVector2f(XL, YL);
+			};
+
+			FVector2f TickLabelSize = GetStringSize(TEXT("-0.00 ")) * FVector2f(1.2f, 1.0f); // Monospace, only character count matters
+			Canvas.DrawTile(-TickLabelSize.X, 0, TickLabelSize.X, RGBParadeRect.Size().Y, 0, 0, 0, 0, FLinearColor::Black, nullptr, false);
+
+			int NumTicks = 13;
+			for (int TickIndex = 0; TickIndex < NumTicks; ++TickIndex)
+			{
+				float Progress = float(TickIndex) / (NumTicks - 1);
+				float TickValue = FMath::Lerp(RGBParadeDrawParameters.MinValue, RGBParadeDrawParameters.MaxValue, Progress);
+				float Y = (1.0f - Progress) * RGBParadeRect.Size().Y;
+				FCanvasLineItem Line { FVector2D(0, Y) / DPIScale, FVector2D(5, Y) / DPIScale };
+				Line.SetColor(FLinearColor(1.0, 1.0, 1.0, 0.2));
+				Canvas.DrawItem(Line);
+				FString Label = FString::Printf(TEXT("% 0.2f"), TickValue);
+				if (TickIndex == 0)
+				{
+					DrawShadowedString(-TickLabelSize.X, Y - TickLabelSize.Y, Label);
+				}
+				else if (TickIndex == NumTicks - 1)
+				{
+					DrawShadowedString(-TickLabelSize.X, Y, Label);
+				}
+				else
+				{
+					DrawShadowedString(-TickLabelSize.X, Y - (TickLabelSize.Y / 2), Label);
+				}
+			}
+		});
+	}
+
+	if (VectorScope)
+	{
+		FIntPoint InputPosition = FIntPoint::ZeroValue;
+		AddDrawTexturePass(GraphBuilder, View, VectorScope, Output.Texture, InputPosition, VectorScope->Desc.Extent, VectorScopeRect.Min, VectorScopeRect.Size(), BilinearClampSampler);
+
+		AddDrawCanvasPass(GraphBuilder, RDG_EVENT_NAME("VectorScope UI"), View, FScreenPassRenderTarget(Output, ERenderTargetLoadAction::ELoad),
+		[&](FCanvas& Canvas)
+		{
+			UFont* Font = GEngine->GetMonospaceFont();
+
+			const float DPIScale = Canvas.GetDPIScale();
+			Canvas.SetBaseTransform(FMatrix(
+				FScaleMatrix(DPIScale)
+				* FTranslationMatrix(FVector(VectorScopeRect.Min.X, VectorScopeRect.Min.Y, 0))
+				* Canvas.CalcBaseTransform2D(Canvas.GetViewRect().Width(), Canvas.GetViewRect().Height())));
+
+			auto DrawShadowedString = [&](float X, float Y, FStringView Text, const FLinearColor& Color = FLinearColor::White)
+			{
+				Canvas.DrawShadowedString(X / DPIScale, Y / DPIScale, Text, Font, Color);
+			};
+		
+			auto GetStringSize = [&](FStringView StringView) -> FVector2f
+			{
+				int32 XL = 0;
+				int32 YL = 0;
+				StringSize(Font, XL, YL, StringView);
+				return FVector2f(XL, YL);
+			};
+
+			int NumTicks = 72;
+			for (int TickIndex = 0; TickIndex < NumTicks; ++TickIndex)
+			{
+				float Progress = float(TickIndex) / (NumTicks - 1);
+
+				FVector2D Start(FMath::Sin(Progress * 2.0f * PI), FMath::Cos(Progress * 2.0f * PI));
+				bool bIsEven = (TickIndex & 1) == 0;
+				FVector2D End = Start * (bIsEven ? 0.98f : 0.95f);
+
+				FVector2D ChartRadius = VectorScopeRect.Size() / 2;
+				FVector2D StartPos = ChartRadius + Start * ChartRadius;
+				FVector2D EndPos = ChartRadius + End * ChartRadius;
+				FCanvasLineItem Line { StartPos / DPIScale, EndPos / DPIScale };
+				Line.SetColor(FLinearColor(0.3, 0.3, 0.3));
+				Canvas.DrawItem(Line);
+			}
+		});
+	}
+
+	if (ColorHistogram)
+	{
+		FIntPoint InputPosition = FIntPoint::ZeroValue;
+		AddDrawTexturePass(GraphBuilder, View, ColorHistogram, Output.Texture, InputPosition, ColorHistogram->Desc.Extent, ColorHistogramRect.Min, ColorHistogramRect.Size(), BilinearClampSampler);
+
+		AddDrawCanvasPass(GraphBuilder, RDG_EVENT_NAME("ColorHistogram UI"), View, FScreenPassRenderTarget(Output, ERenderTargetLoadAction::ELoad),
+		[&](FCanvas& Canvas)
+		{
+			UFont* Font = GEngine->GetMonospaceFont();
+
+			const float DPIScale = Canvas.GetDPIScale();
+			Canvas.SetBaseTransform(FMatrix(
+				FScaleMatrix(DPIScale)
+				* FTranslationMatrix(FVector(ColorHistogramRect.Min.X, ColorHistogramRect.Min.Y, 0))
+				* Canvas.CalcBaseTransform2D(Canvas.GetViewRect().Width(), Canvas.GetViewRect().Height())));
+
+			for (float YNorm : { 1.0f/3, 2.0f/3 })
+			{
+				int Y = YNorm * ColorHistogramRect.Size().Y;
+				FCanvasLineItem Line { FVector2D(0, Y) / DPIScale, FVector2D(ColorHistogramRect.Size().X, Y) / DPIScale };
+				Line.SetColor(FLinearColor(1.0, 1.0, 1.0, 0.2));
+				Canvas.DrawItem(Line);
+			}
+		});
+	}
+
+	return FinalizeOutput(GraphBuilder, View, Output, OverrideOutput);
+}
+
+
+class FVisualizeColorGradingExtension : public FSceneViewExtensionBase
+{
+public:
+	FVisualizeColorGradingExtension( const FAutoRegister& AutoRegister)
+	: FSceneViewExtensionBase( AutoRegister )
+	{
+		FCoreDelegates::OnEnginePreExit.AddLambda([]()
+			{ 
+				Instance = nullptr;
+			});
+	}
+
+	virtual void SubscribeToPostProcessingPass(EPostProcessingPass Pass, const FSceneView& InView, FAfterPassCallbackDelegateArray& InOutPassCallbacks, bool bIsPassEnabled) override
+	{
+		if (!InView.Family->EngineShowFlags.VisualizeColorGrading)
+		{
+			return;
+		}
+
+		if (Pass == EPostProcessingPass::VisualizeDepthOfField)
+		{
+			InOutPassCallbacks.AddDefaulted_GetRef().BindLambda(
+				[this](FRDGBuilder& GraphBuilder, const FSceneView& View, const FPostProcessMaterialInputs& Inputs)
+				{
+					FScreenPassTexture InputSceneColor = FScreenPassTexture{ Inputs.GetInput(EPostProcessMaterialInput::SceneColor) };
+					FScreenPassRenderTarget Output = Inputs.OverrideOutput;
+					return VisualizeColorGrading(GraphBuilder, View, InputSceneColor, Output);
+				});
+		}
+	}
+
+	inline static TSharedPtr<FVisualizeColorGradingExtension,ESPMode::ThreadSafe> Instance = nullptr;
+};
+
+void InitVisualizeColorGradingExtension()
+{
+	FVisualizeColorGradingExtension::Instance = FSceneViewExtensions::NewExtension<FVisualizeColorGradingExtension>(); 
+}
diff --git a/Engine/Source/Runtime/Renderer/Private/PostProcess/VisualizeColorGrading.h b/Engine/Source/Runtime/Renderer/Private/PostProcess/VisualizeColorGrading.h
new file mode 100644
--- /dev/null
+++ b/Engine/Source/Runtime/Renderer/Private/PostProcess/VisualizeColorGrading.h
@@ -0,0 +1,3 @@
+// Copyright Epic Games, Inc. All Rights Reserved.
+
+void InitVisualizeColorGradingExtension();
diff --git a/Engine/Source/Runtime/Renderer/Private/Renderer.cpp b/Engine/Source/Runtime/Renderer/Private/Renderer.cpp
--- a/Engine/Source/Runtime/Renderer/Private/Renderer.cpp
+++ b/Engine/Source/Runtime/Renderer/Private/Renderer.cpp
@@ -52,6 +52,7 @@
 #include "LightFunctionAtlas.h"
 #include "SceneRenderBuilder.h"
 #include "MeshEdgesRendering.h"
+#include "PostProcess/VisualizeColorGrading.h"
 #include "MaterialCache/MaterialCacheTagProvider.h"
 #include "Skinning/SkinnedMeshDebugView.h"
 #include "InstanceDataSceneProxy.h"
@@ -129,6 +130,7 @@ void FRendererModule::StartupModule()
 		{
 			InitMeshEdgesViewExtension();
 			InitSkinnedMeshDebugViewExtension();
+			InitVisualizeColorGradingExtension();
 		});
 	}
 }
diff --git a/Engine/Source/Runtime/Renderer/Private/ScreenPass.cpp b/Engine/Source/Runtime/Renderer/Private/ScreenPass.cpp
--- a/Engine/Source/Runtime/Renderer/Private/ScreenPass.cpp
+++ b/Engine/Source/Runtime/Renderer/Private/ScreenPass.cpp
@@ -295,7 +295,8 @@ void AddDrawTexturePass(
 	FIntPoint InputPosition,
 	FIntPoint InputSize,
 	FIntPoint OutputPosition,
-	FIntPoint OutputSize)
+	FIntPoint OutputSize,
+	FRHISamplerState* Sampler)
 {
 	const FRDGTextureDesc& InputDesc = InputTexture->Desc;
 	const FRDGTextureDesc& OutputDesc = OutputTexture->Desc;
@@ -324,7 +325,7 @@ void AddDrawTexturePass(
 
 	FCopyRectPS::FParameters* Parameters = GraphBuilder.AllocParameters<FCopyRectPS::FParameters>();
 	Parameters->InputTexture = InputTexture;
-	Parameters->InputSampler = TStaticSamplerState<>::GetRHI();
+	Parameters->InputSampler = Sampler ? Sampler : TStaticSamplerState<>::GetRHI();
 	Parameters->RenderTargets[0] = FRenderTargetBinding(OutputTexture, LoadAction);
 
 	AddDrawScreenPass(GraphBuilder, RDG_EVENT_NAME("DrawTexture"), ViewInfo, OutputViewport, InputViewport, PixelShader, Parameters);
@@ -337,7 +338,8 @@ void AddDrawTexturePass(
 	FRDGTextureRef OutputTexture,
 	FIntPoint InputPosition,
 	FIntPoint OutputPosition,
-	FIntPoint Size)
+	FIntPoint Size,
+	FRHISamplerState* Sampler)
 {
 	AddDrawTexturePass(
 		GraphBuilder,
@@ -347,7 +349,8 @@ void AddDrawTexturePass(
 		InputPosition,
 		Size,
 		OutputPosition,
-		Size);
+		Size,
+		Sampler);
 }
 
 void AddDrawTexturePass(
diff --git a/Engine/Source/Runtime/Renderer/Public/ScreenPass.h b/Engine/Source/Runtime/Renderer/Public/ScreenPass.h
--- a/Engine/Source/Runtime/Renderer/Public/ScreenPass.h
+++ b/Engine/Source/Runtime/Renderer/Public/ScreenPass.h
@@ -659,7 +659,8 @@ void RENDERER_API AddDrawTexturePass(
 	FRDGTextureRef OutputTexture,
 	FIntPoint InputPosition = FIntPoint::ZeroValue,
 	FIntPoint OutputPosition = FIntPoint::ZeroValue,
-	FIntPoint Size = FIntPoint::ZeroValue);
+	FIntPoint Size = FIntPoint::ZeroValue,
+	FRHISamplerState* Sampler = nullptr);
 
 void RENDERER_API AddDrawTexturePass(
 	FRDGBuilder& GraphBuilder,
@@ -669,7 +670,8 @@ void RENDERER_API AddDrawTexturePass(
 	FIntPoint InputPosition,
 	FIntPoint InputSize,
 	FIntPoint OutputPosition,
-	FIntPoint OutputSize);
+	FIntPoint OutputSize,
+	FRHISamplerState* Sampler = nullptr);
 
 /** Helper variant which takes a shared viewport instead of unique input / output positions. */
 inline void AddDrawTexturePass(
