Heterogeneous Volumes: Modifications to preshading pipeline

1. Defining specialized preshading uniform buffer
2. Removing dependency on sparse voxel build and uniform buffer
4. Removes tiled rendering indirect dispatch
5. Incorporating LOD controls to dictate preshading voxel resolution. Optionally coupling lighting cache resolution
6. Adding uniform buffer utils API and preferring API over raw variable access

Speedup benefits:
1. Removing sparse voxel build and tiled rendering indirect dispatch improves overall performance by 10-15%.
2. Adding CLOD controls to preshading voxel resolution provides a significant runtime benefit as volumes move away from camera
3. Coupling preshading voxel resolution to lighting cache resolution is always faster than live shading, but adds overblurring.

diff --git a/Engine/Shaders/Private/HeterogeneousVolumes/HeterogeneousVolumesHardwareRayTracing.usf b/Engine/Shaders/Private/HeterogeneousVolumes/HeterogeneousVolumesHardwareRayTracing.usf
--- a/Engine/Shaders/Private/HeterogeneousVolumes/HeterogeneousVolumesHardwareRayTracing.usf
+++ b/Engine/Shaders/Private/HeterogeneousVolumes/HeterogeneousVolumesHardwareRayTracing.usf
@@ -391,6 +391,10 @@ float VolumetricScatteringIntensity;
 
 // Shadowing
 uint VirtualShadowMapId;
+uint GetVirtualShadowMapId()
+{
+	return VirtualShadowMapId;
+}
 
 float3 ComputeTransmittanceHardwareRayTracing(
 	float3 Origin,
diff --git a/Engine/Shaders/Private/HeterogeneousVolumes/HeterogeneousVolumesLiveShadingPipeline.usf b/Engine/Shaders/Private/HeterogeneousVolumes/HeterogeneousVolumesLiveShadingPipeline.usf
--- a/Engine/Shaders/Private/HeterogeneousVolumes/HeterogeneousVolumesLiveShadingPipeline.usf
+++ b/Engine/Shaders/Private/HeterogeneousVolumes/HeterogeneousVolumesLiveShadingPipeline.usf
@@ -89,6 +89,11 @@ float VolumetricScatteringIntensity;
 
 // Shadowing
 uint VirtualShadowMapId;
+uint GetVirtualShadowMapId()
+{
+	return VirtualShadowMapId;
+}
+
 float ShadowStepSize;
 float ShadowStepFactor;
 
diff --git a/Engine/Shaders/Private/HeterogeneousVolumes/HeterogeneousVolumesPreshadingPipeline.usf b/Engine/Shaders/Private/HeterogeneousVolumes/HeterogeneousVolumesPreshadingPipeline.usf
--- a/Engine/Shaders/Private/HeterogeneousVolumes/HeterogeneousVolumesPreshadingPipeline.usf
+++ b/Engine/Shaders/Private/HeterogeneousVolumes/HeterogeneousVolumesPreshadingPipeline.usf
@@ -1,6 +1,5 @@
 // Copyright Epic Games, Inc. All Rights Reserved.
 
-#include "HeterogeneousVolumesSparseVoxelUtils.ush"
 #include "../Common.ush"
 
 #define SUPPORT_CONTACT_SHADOWS 0
@@ -11,65 +10,28 @@
 #define SHADOW_QUALITY 2
 #define NO_TRANSLUCENCY_AVAILABLE
 
-#if VIRTUAL_SHADOW_MAP
-#include "../VirtualShadowMaps/VirtualShadowMapProjectionCommon.ush"
-#endif
-
-#include "HeterogeneousVolumesSparseVoxelUniformBufferUtils.ush"
-#include "HeterogeneousVolumesTracingUtils.ush"
-#include "HeterogeneousVolumesLightingUtils.ush"
-
-#ifndef THREADGROUP_SIZE_1D
-#define THREADGROUP_SIZE_1D 1
-#endif // THREADGROUP_SIZE_1D
-
-#ifndef THREADGROUP_SIZE_2D
-#define THREADGROUP_SIZE_2D 1
-#endif // THREADGROUP_SIZE_2D
-
-#ifndef THREADGROUP_SIZE_3D
-#define THREADGROUP_SIZE_3D 1
-#endif // THREADGROUP_SIZE_3D
-
-#ifndef DIM_SPARSE_VOXEL_TRACING
-#define DIM_SPARSE_VOXEL_TRACING 1
-#endif // DIM_SPARSE_VOXEL_TRACING
-
-#ifndef DIM_VOXEL_CULLING
-#define DIM_VOXEL_CULLING 1
-#endif // DIM_VOXEL_CULLING
-
 #ifndef DIM_APPLY_SHADOW_TRANSMITTANCE
 #define DIM_APPLY_SHADOW_TRANSMITTANCE 0
 #endif // DIM_APPLY_SHADOW_TRANSMITTANCE
 
-// Lighting
-int bApplyEmissionAndTransmittance;
-int bApplyDirectLighting;
-int bApplyShadowTransmittance;
-int LightType;
-float VolumetricScatteringIntensity;
-
-// Shadowing
-uint VirtualShadowMapId;
+#ifndef THREADGROUP_SIZE_3D
+#define THREADGROUP_SIZE_3D 1
+#endif // THREADGROUP_SIZE_3D
 
-// Ray data
-float MaxTraceDistance;
-float MaxShadowTraceDistance;
-float StepSize;
-int MaxStepCount;
-int bJitter;
+#ifndef THREADGROUP_SIZE_2D
+#define THREADGROUP_SIZE_2D 1
+#endif // THREADGROUP_SIZE_2D
 
-// Dispatch data
-int3 GroupCount;
-int DownsampleFactor;
+#if VIRTUAL_SHADOW_MAP
+#include "../VirtualShadowMaps/VirtualShadowMapProjectionCommon.ush"
+#endif
 
-// Volume data
-int MipLevel;
+#include "HeterogeneousVolumesPreshadingUniformBufferUtils.ush"
+#include "HeterogeneousVolumesPreshadingUtils.ush"
+#include "HeterogeneousVolumesTracingUtils.ush"
+#include "HeterogeneousVolumesLightingUtils.ush"
 
-// Output
-RWTexture2D<float4> RWLightingTexture;
-RWTexture2D<float4> RWVelocityTexture;
+#include "HeterogeneousVolumesRayMarchingUtils.ush"
 
 int3 TextureResolution;
 Texture3D InputTexture;
@@ -137,337 +99,6 @@ void GenerateMin3D(
 	}
 }
 
-StructuredBuffer<FVoxelDataPacked> VoxelBuffer;
-Buffer<uint> NumVoxelsBuffer;
-
-// A screen-tile is an 8x8 block of computation for ray marching
-struct FRayMarchingTile
-{
-	uint2 PixelOffset;
-	uint Voxels[2];
-
-	// TODO: For tracking purposes..
-	uint Id;
-	uint Padding[3];
-};
-
-FRayMarchingTile CreateRayMarchingTile(uint TileId, uint2 PixelOffset, uint VoxelMin, uint VoxelMax)
-{
-	FRayMarchingTile Tile = (FRayMarchingTile) 0;
-	Tile.PixelOffset = PixelOffset;
-
-	Tile.Voxels[0] = VoxelMin;
-	Tile.Voxels[1] = VoxelMax;
-
-	Tile.Id = TileId;
-	return Tile;
-}
-
-RWStructuredBuffer<FRayMarchingTile> RWRayMarchingTilesBuffer;
-RWBuffer<uint> RWNumRayMarchingTilesBuffer;
-
-RWStructuredBuffer<FVoxelDataPacked> RWVoxelsPerTileBuffer;
-RWBuffer<uint> RWNumVoxelsPerTileBuffer;
-
-struct FRayMarchingDebug
-{
-	float4 Planes[5];
-	float4 BBox[2];
-
-	float Padding[4];
-};
-
-RWStructuredBuffer<FRayMarchingDebug> RWRayMarchingDebugBuffer;
-
-float4 CreatePlane(
-	float3 P_000,
-	float3 P_100,
-	float3 P_010
-)
-{
-	float3 X = P_100 - P_000;
-	float3 Y = P_010 - P_000;
-
-	float3 N = normalize(cross(X, Y));
-	float d = -dot(N, P_000);
-	return float4(N, d);
-}
-
-float CalcSignedDistance(float3 P, float4 Plane)
-{
-	return dot(P, Plane.xyz) + Plane.w;
-}
-
-bool PointIntersectsFrustum(float3 P, float4 Planes[5])
-{
-	for (int i = 0; i < 5; ++i)
-	{
-		if (CalcSignedDistance(P, Planes[i]) > 0.0)
-		{
-			return false;
-		}
-	}
-
-	return true;
-}
-
-bool BoxIntersectsFrustum(float3 BoxOrigin, float3 BoxExtent, float4 Planes[5])
-{
-	for (int i = 0; i < 5; ++i)
-	{
-		float Distance = CalcSignedDistance(BoxOrigin, Planes[i]);
-		float ProjectedExtent = dot(BoxExtent, abs(Planes[i].xyz));
-
-		if (Distance > ProjectedExtent)
-		{
-			return false;
-		}
-	}
-
-	return true;
-}
-
-[numthreads(THREADGROUP_SIZE_2D, THREADGROUP_SIZE_2D, 1)]
-void GenerateRayMarchingTiles(
-	uint2 DispatchThreadId : SV_DispatchThreadID
-)
-{
-	if (!any(DispatchThreadId))
-	{
-		RWNumRayMarchingTilesBuffer[1] = 1;
-		RWNumRayMarchingTilesBuffer[2] = 1;
-	}
-
-	uint4 Viewport = uint4(DispatchThreadId.xy * THREADGROUP_SIZE_2D, (DispatchThreadId.xy + 1) * THREADGROUP_SIZE_2D);
-	uint2 PixelOffset = Viewport.xy;
-
-	// Scale to full-res
-	Viewport *= DownsampleFactor;
-	uint2 ViewPixelOffset = Viewport.xy;
-	if (any(ViewPixelOffset >= View.ViewSizeAndInvSize.xy))
-	{
-		return;
-	}
-
-	// Build tile-frustum
-	float Epsilon = 1.0e-6;
-#if HAS_INVERTED_Z_BUFFER
-	float NearDepth = 1.0;
-	float FarDepth = Epsilon;
-#else
-	float NearDepth = 0.0;
-	float FarDepth = 1.0 - Epsilon;
-#endif // HAS_INVERTED_Z_BUFFER
-
-	FDFInverseMatrix WorldToLocal = DFPromoteInverse(GetWorldToLocal()); // LWC_TODO
-	float4x4 TranslatedWorldToLocal = DFFastToTranslatedWorld(WorldToLocal, PrimaryView.PreViewTranslation);
-	float3 W_000 = SvPositionToTranslatedWorld(float4(Viewport.xw, NearDepth, 1));
-	float3 W_100 = SvPositionToTranslatedWorld(float4(Viewport.zw, NearDepth, 1));
-	float3 W_010 = SvPositionToTranslatedWorld(float4(Viewport.xy, NearDepth, 1));
-	float3 W_110 = SvPositionToTranslatedWorld(float4(Viewport.zy, NearDepth, 1));
-	float3 W_001 = SvPositionToTranslatedWorld(float4(Viewport.xw, FarDepth, 1));
-	float3 W_101 = SvPositionToTranslatedWorld(float4(Viewport.zw, FarDepth, 1));
-	float3 W_011 = SvPositionToTranslatedWorld(float4(Viewport.xy, FarDepth, 1));
-	float3 W_111 = SvPositionToTranslatedWorld(float4(Viewport.zy, FarDepth, 1));
-
-	float3 L_000 = mul(float4(W_000, 1), TranslatedWorldToLocal).xyz;
-	float3 L_100 = mul(float4(W_100, 1), TranslatedWorldToLocal).xyz;
-	float3 L_010 = mul(float4(W_010, 1), TranslatedWorldToLocal).xyz;
-	float3 L_110 = mul(float4(W_110, 1), TranslatedWorldToLocal).xyz;
-	float3 L_001 = mul(float4(W_001, 1), TranslatedWorldToLocal).xyz;
-	float3 L_101 = mul(float4(W_101, 1), TranslatedWorldToLocal).xyz;
-	float3 L_011 = mul(float4(W_011, 1), TranslatedWorldToLocal).xyz;
-	float3 L_111 = mul(float4(W_111, 1), TranslatedWorldToLocal).xyz;
-
-	float4 Planes[5] = {
-		// Front
-		CreatePlane(L_000, L_010, L_100),
-		// Back
-		//CreatePlane(L_101, L_001, L_111),
-		// +X
-		CreatePlane(L_100, L_110, L_101),
-		// +Y
-		CreatePlane(L_010, L_011, L_110),
-		// -X
-		CreatePlane(L_000, L_001, L_010),
-		// -Y
-		CreatePlane(L_000, L_100, L_001)
-	};
-
-	uint DebugIndex = DispatchThreadId.y * THREADGROUP_SIZE_2D + DispatchThreadId.x;
-#if DIM_DEBUG
-	for (uint i = 0; i < 5; ++i)
-	{
-		RWRayMarchingDebugBuffer[DebugIndex].Planes[i] = Planes[i];
-	}
-
-	RWRayMarchingDebugBuffer[DebugIndex].BBox[0] = float4(GetLocalBoundsOrigin() - GetLocalBoundsExtent(), 0.0);
-	RWRayMarchingDebugBuffer[DebugIndex].BBox[1] = float4(GetLocalBoundsOrigin() + GetLocalBoundsExtent(), 0.0);
-#endif // DIM_DEBUG
-
-	// Intersect frustum against bounding volume
-	if (!BoxIntersectsFrustum(GetLocalBoundsOrigin(), GetLocalBoundsExtent(), Planes))
-	{
-		return;
-	}
-
-#if DIM_VOXEL_CULLING
-	float3 LocalBoundsMin = GetLocalBoundsOrigin() - GetLocalBoundsExtent();
-	float3 LocalBoundsMax = GetLocalBoundsOrigin() + GetLocalBoundsExtent();
-
-	// Intersect against voxels to determine allocation
-	uint VoxelCount = 0;
-	for (uint VoxelIndex = 0; VoxelIndex < GetNumVoxels(); ++VoxelIndex)
-	{
-		FVoxelData VoxelData = DecodeVoxelData(GetVoxelDataPacked(VoxelIndex), GetVolumeResolution());
-		FBox VoxelBoundingBox = GetVoxelBoundingBox(VoxelData, GetVolumeResolution(), LocalBoundsMin, LocalBoundsMax);
-
-		float3 BoxOrigin = (VoxelBoundingBox.Max + VoxelBoundingBox.Min) / 2.0;
-		float3 BoxExtent = VoxelBoundingBox.Max - VoxelBoundingBox.Min;
-		if (BoxIntersectsFrustum(BoxOrigin, BoxExtent, Planes))
-		{
-			VoxelCount++;
-		}
-	}
-
-	// Allocate contiguous span of voxel indices
-	uint VoxelEntryIndex;
-	if (VoxelCount > 0)
-	{
-		InterlockedAdd(RWNumVoxelsPerTileBuffer[0], VoxelCount, VoxelEntryIndex);
-
-		// Intersect and store voxels
-		VoxelCount = 0;
-		for (uint VoxelIndex = 0; VoxelIndex < GetNumVoxels(); ++VoxelIndex)
-		{
-			FVoxelData VoxelData = DecodeVoxelData(GetVoxelDataPacked(VoxelIndex), GetVolumeResolution());
-			FBox VoxelBoundingBox = GetVoxelBoundingBox(VoxelData, GetVolumeResolution(), LocalBoundsMin, LocalBoundsMax);
-
-			float3 BoxOrigin = (VoxelBoundingBox.Max + VoxelBoundingBox.Min) / 2.0;
-			float3 BoxExtent = VoxelBoundingBox.Max - VoxelBoundingBox.Min;
-			if (BoxIntersectsFrustum(BoxOrigin, BoxExtent, Planes))
-			{
-				RWVoxelsPerTileBuffer[VoxelEntryIndex + VoxelCount] = GetVoxelDataPacked(VoxelIndex);
-				VoxelCount++;
-			}
-		}
-	}
-#else
-	uint VoxelEntryIndex = 0;
-	uint VoxelCount = GetNumVoxels();
-#endif
-
-	// Allocate tile
-	if (VoxelCount > 0)
-	{
-		uint TileId;
-		InterlockedAdd(RWNumRayMarchingTilesBuffer[0], 1, TileId);
-
-		RWRayMarchingTilesBuffer[TileId] = CreateRayMarchingTile(TileId, PixelOffset, VoxelEntryIndex, VoxelEntryIndex + VoxelCount);
-	}
-}
-
-struct FTraceVoxelResult
-{
-	FVoxelData VoxelData;
-	float2 HitT;
-};
-
-FTraceVoxelResult TraceVoxels(float3 LocalRayOrigin, float3 LocalRayDirection, float LocalRayTMin, float LocalRayTMax,
-	float3 LocalBoundsMin, float3 LocalBoundsMax)
-{
-	FTraceVoxelResult TraceVoxelResult;
-	TraceVoxelResult.HitT.x = LocalRayTMax;
-	TraceVoxelResult.HitT.y = LocalRayTMax;
-
-	for (uint VoxelIndex = 0; VoxelIndex < GetNumVoxels(); ++VoxelIndex)
-	{
-		FVoxelData VoxelData = DecodeVoxelData(GetVoxelDataPacked(VoxelIndex), GetVolumeResolution());
-
-		FBox VoxelBoundingBox = GetVoxelBoundingBox(VoxelData, GetVolumeResolution(), LocalBoundsMin, LocalBoundsMax);
-		float2 HitT = IntersectAABB(LocalRayOrigin, LocalRayDirection, LocalRayTMin, LocalRayTMax, VoxelBoundingBox.Min, VoxelBoundingBox.Max);
-
-		bool bIsHit = (HitT.y - HitT.x) > 0.0;
-		if (bIsHit && (HitT.x < TraceVoxelResult.HitT.x))
-		{
-			TraceVoxelResult.HitT = HitT;
-			TraceVoxelResult.VoxelData = VoxelData;
-		}
-	}
-
-	return TraceVoxelResult;
-}
-
-StructuredBuffer<FVoxelDataPacked> VoxelsPerTileBuffer;
-
-FTraceVoxelResult TraceVoxels(float3 LocalRayOrigin, float3 LocalRayDirection, float LocalRayTMin, float LocalRayTMax,
-	float3 LocalBoundsMin, float3 LocalBoundsMax,
-	uint VoxelStartIndex, uint VoxelEndIndex)
-{
-	FTraceVoxelResult TraceVoxelResult;
-	TraceVoxelResult.HitT.x = LocalRayTMax;
-	TraceVoxelResult.HitT.y = LocalRayTMax;
-
-	for (uint VoxelIndex = VoxelStartIndex; VoxelIndex < VoxelEndIndex; ++VoxelIndex)
-	{
-		FVoxelData VoxelData = DecodeVoxelData(VoxelsPerTileBuffer[VoxelIndex], GetVolumeResolution());
-
-		FBox VoxelBoundingBox = GetVoxelBoundingBox(VoxelData, GetVolumeResolution(), LocalBoundsMin, LocalBoundsMax);
-		float2 HitT = IntersectAABB(LocalRayOrigin, LocalRayDirection, LocalRayTMin, LocalRayTMax, VoxelBoundingBox.Min, VoxelBoundingBox.Max);
-
-		bool bIsHit = (HitT.y - HitT.x) > 0.0;
-		if (bIsHit && (HitT.x < TraceVoxelResult.HitT.x))
-		{
-			TraceVoxelResult.HitT = HitT;
-			TraceVoxelResult.VoxelData = VoxelData;
-		}
-	}
-
-	return TraceVoxelResult;
-}
-
-// TODO: Implement ambient occlusion for preshading pipeline
-float EvalAmbientOcclusion(float3 WorldPosition)
-{
-	return 1.0;
-}
-
-// Cinematic feature options
-bool UseInscatteringVolume()
-{
-	return true;
-}
-
-#define INDIRECT_LIGHTING_MODE_OFF 0
-#define INDIRECT_LIGHTING_MODE_LIGHTING_CACHE 1
-#define INDIRECT_LIGHTING_MODE_SINGLE_SCATTERING 2
-
-int GetIndirectLightingMode()
-{
-	return INDIRECT_LIGHTING_MODE_OFF;
-}
-
-bool ShouldWriteVelocity()
-{
-	return false;
-}
-
-bool IsOfflineRender()
-{
-	return false;
-}
-
-#define FOG_INSCATTERING_MODE_OFF 0
-#define FOG_INSCATTERING_MODE_REFERENCE 1
-#define FOG_INSCATTERING_MODE_LINEAR_APPROX 2
-
-int GetFogInscatteringMode()
-{
-	return FOG_INSCATTERING_MODE_LINEAR_APPROX;
-}
-
-#include "HeterogeneousVolumesRayMarchingUtils.ush"
-
 RWTexture3D<float3> RWLightingCacheTexture;
 
 [numthreads(THREADGROUP_SIZE_3D, THREADGROUP_SIZE_3D, THREADGROUP_SIZE_3D)]
@@ -482,14 +113,13 @@ void RenderLightingCacheWithPreshadingCS(
 	{
 		return;
 	}
-
 	float3 VoxelJitter = 0.5;
 	if (bJitter)
 	{
 		uint3 Rand32Bits = Rand4DPCG32(uint4(VoxelIndex, View.StateFrameIndexMod8)).xyz;
 		VoxelJitter = float3(Rand32Bits) / float(uint(0xffffffff));
 	}
-	//float3 UVW = (VoxelIndex + VectorJitter) / float3(GetLightingCacheResolution());
+	//float3 UVW = (VoxelIndex + VoxelJitter) / float3(GetLightingCacheResolution());
 	float3 UVW = (VoxelIndex + 0.5) / float3(GetLightingCacheResolution());
 
 	float3 LocalBoundsMin = GetLocalBoundsOrigin() - GetLocalBoundsExtent();
@@ -514,11 +144,11 @@ void RenderLightingCacheWithPreshadingCS(
 	}
 
 	// Build shadow ray
-	const FDeferredLightData LightData = InitDeferredLightFromUniforms(LightType, VolumetricScatteringIntensity);
+	const FDeferredLightData LightData = InitDeferredLightFromUniforms(GetLightType(), GetVolumetricScatteringIntensity());
 #if DIM_LIGHTING_CACHE_MODE == LIGHTING_CACHE_TRANSMITTANCE
 	float3 L = LightData.Direction;
 	float3 ToLight = L * GetMaxShadowTraceDistance();
-	if (LightType != LIGHT_TYPE_DIRECTIONAL)
+	if (GetLightType() != LIGHT_TYPE_DIRECTIONAL)
 	{
 		// Note: this function also permutes ToLight and L
 		float3 TranslatedWorldPosition = DFFastToTranslatedWorld(WorldRayOrigin, PrimaryView.PreViewTranslation);
@@ -534,51 +164,50 @@ void RenderLightingCacheWithPreshadingCS(
 		}
 	}
 
-	float3 Transmittance = ComputeTransmittance(WorldRayOrigin, ToLight, MaxStepCount);
+	float3 Transmittance = ComputeTransmittance(WorldRayOrigin, ToLight, GetMaxStepCount());
 	RWLightingCacheTexture[VoxelIndex] = Transmittance;
 
 #elif DIM_LIGHTING_CACHE_MODE == LIGHTING_CACHE_INSCATTERING
-	float3 Inscattering = ComputeInscattering(WorldRayOrigin, LightData, LightType, MaxStepCount, CalcShadowBias(), bApplyShadowTransmittance);
+	float3 Inscattering = ComputeInscattering(WorldRayOrigin, LightData, GetLightType(), GetMaxStepCount(), CalcShadowBias(), ShouldApplyShadowTransmittance());
 	RWLightingCacheTexture[VoxelIndex] += Inscattering * View.PreExposure;
 #endif // DIM_LIGHTING_CACHE_MODE
 }
 
-RWStructuredBuffer<FVoxelDataPacked> RWVoxelOutputBuffer;
+RWTexture2D<float4> RWLightingTexture;
+RWTexture2D<float4> RWVelocityTexture;
 
-StructuredBuffer<FRayMarchingTile> RayMarchingTilesBuffer;
-Buffer<uint> NumRayMarchingTilesBuffer;
+int2 GetDownsampledResolution(int2 Resolution, int Factor)
+{
+	return (Resolution + Factor - 1) / Factor;
+}
 
-float2 GetDispatchSize()
+int2 GetScaledViewRect()
 {
-	return View.ViewSizeAndInvSize.xy / DownsampleFactor;
+	return GetDownsampledResolution(View.ViewSizeAndInvSize.xy, GetDownsampleFactor());
 }
 
-[numthreads(THREADGROUP_SIZE_1D, 1, 1)]
+[numthreads(THREADGROUP_SIZE_2D, THREADGROUP_SIZE_2D, 1)]
 void RenderSingleScatteringWithPreshadingCS(
 	uint GroupId : SV_GroupID,
-	uint GroupThreadId : SV_GroupThreadID
+	uint GroupThreadId : SV_GroupThreadID,
+	uint2 DispatchThreadId : SV_DispatchThreadID
 )
 {
-	// Convert 1D index to pixel coordinate
-	FRayMarchingTile RayMarchingTile = RayMarchingTilesBuffer[GroupId];
-	// TODO: Morton..
-	uint2 TileCoord = uint2(GroupThreadId % THREADGROUP_SIZE_2D, GroupThreadId / THREADGROUP_SIZE_2D);
-	//float2 PixelCoord = View.ViewRectMin.xy + RayMarchingTile.PixelOffset + TileCoord;
-	float2 PixelCoord = RayMarchingTile.PixelOffset + TileCoord;
-	float2 ViewPixelCoord = PixelCoord * DownsampleFactor;
+	// Create screen ray
+	if (any(DispatchThreadId.xy >= GetScaledViewRect()))
+	{
+		return;
+	}
+
+	uint LinearIndex = DispatchThreadId.y * GetGroupCount().x * THREADGROUP_SIZE_2D + DispatchThreadId.x;
+	uint2 ViewPixelCoord = DispatchThreadId.xy * DownsampleFactor + View.ViewRectMin.xy;
+	uint2 PixelCoord = DispatchThreadId.xy + View.ViewRectMin.xy / DownsampleFactor;
 
 	float3 Radiance = 0.0;
 	float4 LightingTexture = RWLightingTexture[PixelCoord];
 	float3 Transmittance = LightingTexture.aaa;
 	float PrevOpacity = Luminance(1.0 - Transmittance);
 
-	// Create screen ray
-	if (any(PixelCoord.xy >= GetDispatchSize()))
-	{
-		return;
-	}
-	PixelCoord += View.ViewRectMin.xy;
-
 	// Extract depth
 	float DeviceZ = SceneDepthTexture.Load(int3(ViewPixelCoord, 0)).r;
 #if HAS_INVERTED_Z_BUFFER
@@ -586,10 +215,10 @@ void RenderSingleScatteringWithPreshadingCS(
 #endif // HAS_INVERTED_Z_BUFFER
 
 	// Clip trace distance
-	float SceneDepth = min(ConvertFromDeviceZ(DeviceZ), MaxTraceDistance);
+	float SceneDepth = min(ConvertFromDeviceZ(DeviceZ), GetMaxTraceDistance());
 	DeviceZ = ConvertToDeviceZ(SceneDepth);
 
-	float Jitter = bJitter ? InterleavedGradientNoise(ViewPixelCoord, View.StateFrameIndexMod8) : 0.0;
+	float Jitter = ShouldJitter() ? InterleavedGradientNoise(ViewPixelCoord, View.StateFrameIndexMod8) : 0.0;
 
 	// Intersect ray with bounding volume
 	// TODO: LWC integration..
@@ -611,27 +240,7 @@ void RenderSingleScatteringWithPreshadingCS(
 	// Test bounding volume
 	float2 HitT = IntersectAABB(LocalRayOrigin, LocalRayDirection, 0.0, LocalRayLength, LocalBoundsMin, LocalBoundsMax);
 	float HitSpan = HitT.y - HitT.x;
-	if (HitSpan <= 0.0)
-	{
-		return;
-	}
-
-	float WorldHitTMin = WorldRayLength;
-#if DIM_SPARSE_VOXEL_TRACING
-	float RayTMax = HitT.y;
-
-	#if DIM_VOXEL_CULLING
-		HitT = TraceVoxels(LocalRayOrigin, LocalRayDirection, 0.0, RayTMax,
-			LocalBoundsMin, LocalBoundsMax,
-			RayMarchingTile.Voxels[0], RayMarchingTile.Voxels[1]).HitT;
-	#else
-		HitT = TraceVoxels(LocalRayOrigin, LocalRayDirection, 0.0, RayTMax, LocalBoundsMin, LocalBoundsMax).HitT;
-	#endif // DIM_VOXEL_CULLING
-
-	HitSpan = HitT.y - HitT.x;
-
-	while (HitSpan > 0.0)
-#endif // DIM_SPARSE_VOXEL_TRACING
+	if (HitSpan > 0.0)
 	{
 		FRayMarchingContext RayMarchingContext = CreateRayMarchingContext(
 			// Local-space
@@ -645,73 +254,49 @@ void RenderSingleScatteringWithPreshadingCS(
 			// Ray-step attributes
 			Jitter,
 			CalcStepSize(LocalRayDirection),
-			MaxStepCount,
-			bApplyEmissionAndTransmittance,
-			bApplyDirectLighting,
+			GetMaxStepCount(),
+			ShouldApplyEmissionAndTransmittance(),
+			ShouldApplyDirectLighting(),
 			DIM_APPLY_SHADOW_TRANSMITTANCE
 		);
 
-		const FDeferredLightData LightData = InitDeferredLightFromUniforms(LightType, VolumetricScatteringIntensity);
+		const FDeferredLightData LightData = InitDeferredLightFromUniforms(GetLightType(), GetVolumetricScatteringIntensity());
 		uint StepCount = CalcStepCount(RayMarchingContext);
 
 		float ScatterHitT = HitT.y;
 		RayMarchSingleScattering(
 			RayMarchingContext,
 			LightData,
-			LightType,
+			GetLightType(),
 			StepCount,
 			ScatterHitT,
 			Radiance,
 			Transmittance
 		);
-		WorldHitTMin = min(WorldHitTMin, ScatterHitT * RayMarchingContext.LocalToWorldScale);
-
-#if DIM_SPARSE_VOXEL_TRACING
-	#if DIM_VOXEL_CULLING
-		HitT = TraceVoxels(LocalRayOrigin, LocalRayDirection, HitT.y + 0.01, RayTMax,
-			LocalBoundsMin, LocalBoundsMax,
-			RayMarchingTile.Voxels[0], RayMarchingTile.Voxels[1]).HitT;
-	#else
-		HitT = TraceVoxels(LocalRayOrigin, LocalRayDirection, HitT.y + 0.0001, RayTMax, LocalBoundsMin, LocalBoundsMax).HitT;
-	#endif // DIM_VOXEL_CULLING
-
-		HitSpan = HitT.y - HitT.x;
-#endif // DIM_SPARSE_VOXEL_TRACING
-	}
-
-	// Radiance
-	LightingTexture.rgb += Radiance * View.PreExposure;
-	if (bApplyEmissionAndTransmittance)
-	{
-		float Opacity = Luminance(1.0 - Transmittance);
-		LightingTexture.a = Luminance(Transmittance);
 
-		if (ShouldWriteVelocity())
+		// Radiance
+		LightingTexture.rgb += Radiance * View.PreExposure;
+		if (ShouldApplyEmissionAndTransmittance())
 		{
-			bool bWriteVelocity =
-				(WorldHitTMin < WorldRayLength) &&
-				(PrevOpacity < 0.5) &&
-				(Opacity > 0.5);
+			float Opacity = Luminance(1.0 - Transmittance);
+			LightingTexture.a = Luminance(Transmittance);
 
-			if (bWriteVelocity)
+			if (ShouldWriteVelocity())
 			{
-				float3 Velocity = Calculate3DVelocityAtWorldPos(WorldRayOrigin, WorldRayDirection, WorldHitTMin);
-				RWVelocityTexture[PixelCoord] = EncodeVelocityToTexture(Velocity);
+				bool bWriteVelocity =
+					(ScatterHitT < HitT.y) &&
+					(PrevOpacity < 0.5) &&
+					(Opacity > 0.5);
+
+				if (bWriteVelocity)
+				{
+					float WorldScatterHitT = ScatterHitT * RayMarchingContext.LocalToWorldScale;
+					float3 Velocity = Calculate3DVelocityAtWorldPos(WorldRayOrigin, WorldRayDirection, WorldScatterHitT);
+					RWVelocityTexture[PixelCoord] = EncodeVelocityToTexture(Velocity);
+				}
 			}
 		}
 
-#if DIM_DEBUG
-		//if (!all(DispatchThreadId))
-		if (GroupId == 0 && GroupThreadId == 0)
-		{
-			RWVoxelOutputBuffer[0].MipLevel = GetNumVoxels();
-			for (uint VoxelIndex = 0; VoxelIndex < GetNumVoxels(); ++VoxelIndex)
-			{
-				RWVoxelOutputBuffer[VoxelIndex + 1] = GetVoxelDataPacked(VoxelIndex);
-			}
-		}
-#endif
+		RWLightingTexture[PixelCoord] = LightingTexture;
 	}
-
-	RWLightingTexture[PixelCoord] = LightingTexture;
 }
diff --git a/Engine/Shaders/Private/HeterogeneousVolumes/HeterogeneousVolumesPreshadingUniformBufferUtils.ush b/Engine/Shaders/Private/HeterogeneousVolumes/HeterogeneousVolumesPreshadingUniformBufferUtils.ush
new file mode 100644
--- /dev/null
+++ b/Engine/Shaders/Private/HeterogeneousVolumes/HeterogeneousVolumesPreshadingUniformBufferUtils.ush
@@ -0,0 +1,118 @@
+// Copyright Epic Games, Inc. All Rights Reserved.
+
+#include "HeterogeneousVolumesStochasticFiltering.ush"
+
+struct FVolumeSampleContext
+{
+	float3 UVW;
+	float MipLevel;
+};
+
+float4x4 GetLocalToWorld()
+{
+	return PreshadingUniformBuffer.LocalToWorld;
+}
+
+float4x4 GetWorldToLocal()
+{
+	return PreshadingUniformBuffer.WorldToLocal;
+}
+
+float3 GetLocalBoundsOrigin()
+{
+	return PreshadingUniformBuffer.LocalBoundsOrigin;
+}
+
+float3 GetLocalBoundsExtent()
+{
+	return PreshadingUniformBuffer.LocalBoundsExtent;
+}
+
+float3 SampleExtinction(inout FVolumeSampleContext Context)
+{
+	float3 Extinction = Texture3DSampleLevel(
+		PreshadingUniformBuffer.ExtinctionTexture,
+		PreshadingUniformBuffer.TextureSampler,
+		Context.UVW,
+		Context.MipLevel).rgb;
+
+	return Extinction;
+}
+
+float3 SampleEmission(inout FVolumeSampleContext Context)
+{
+	float3 Emission = Texture3DSampleLevel(
+		PreshadingUniformBuffer.EmissionTexture,
+		PreshadingUniformBuffer.TextureSampler,
+		Context.UVW,
+		Context.MipLevel).xyz;
+
+	return Emission;
+}
+
+float3 SampleAlbedo(inout FVolumeSampleContext Context)
+{
+	float3 Albedo = Texture3DSampleLevel(
+		PreshadingUniformBuffer.AlbedoTexture,
+		PreshadingUniformBuffer.TextureSampler,
+		Context.UVW,
+		Context.MipLevel).xyz;
+
+	return Albedo;
+}
+
+FVolumeSampleContext CreateVolumeSampleContext(
+	float3 LocalPosition,
+	float3 WorldPosition,
+	float3 WorldPosition_DDX,
+	float3 WorldPosition_DDY,
+	float MipLevel
+)
+{
+	float3 LocalBoundsMin = GetLocalBoundsOrigin() - GetLocalBoundsExtent();
+	float3 LocalBoundsMax = GetLocalBoundsOrigin() + GetLocalBoundsExtent();
+
+	FVolumeSampleContext VolumeSampleContext;
+	VolumeSampleContext.UVW = (LocalPosition - LocalBoundsMin) / (LocalBoundsMax - LocalBoundsMin);
+	VolumeSampleContext.MipLevel = MipLevel;
+
+	return VolumeSampleContext;
+}
+
+FVolumeSampleContext CreateVolumeSampleContext(float3 LocalPosition, float3 WorldPosition, float MipLevel)
+{
+	float3 WorldPosition_DDX = 0.0;
+	float3 WorldPosition_DDY = 0.0;
+
+	return CreateVolumeSampleContext(
+		LocalPosition,
+		WorldPosition,
+		WorldPosition_DDX,
+		WorldPosition_DDY,
+		MipLevel
+	);
+}
+
+FVolumeSampleContext CreateVolumeSampleContext(float3 LocalPosition, float3 WorldPosition, float3 FilterWidth, float MipLevel, float Rand, int SamplingMethod)
+{
+	// TODO: Use FilterWidth determine MipLevel
+	//float3 WorldPosition_DDX = mul(float4(LocalPosition + float3(1.0, 0.0, 0.0), 1.0), GetLocalToWorld()).xyz;
+	//float3 WorldPosition_DDY = mul(float4(LocalPosition + float3(0.0, 1.0, 0.0), 1.0), GetLocalToWorld()).xyz;
+	//float3 WorldPosition_DDZ = mul(float4(LocalPosition + float3(0.0, 0.0, 1.0), 1.0), GetLocalToWorld()).xyz;
+
+	if (SamplingMethod == STOCHASTIC_FILTERING_CONSTANT)
+	{
+		LocalPosition = StochasticFilteringConstant(LocalPosition, Rand);
+	}
+	else if (SamplingMethod == STOCHASTIC_FILTERING_TRILINEAR)
+	{
+		LocalPosition = StochasticFilteringTrilinear(LocalPosition, Rand);
+	}
+	else if (SamplingMethod == STOCHASTIC_FILTERING_TRICUBIC)
+	{
+		LocalPosition = StochasticFilteringTricubic(LocalPosition, Rand);
+	}
+
+	WorldPosition = mul(float4(LocalPosition, 1.0), GetLocalToWorld()).xyz;
+	return CreateVolumeSampleContext(LocalPosition, WorldPosition, MipLevel);
+}
\ No newline at end of file
diff --git a/Engine/Shaders/Private/HeterogeneousVolumes/HeterogeneousVolumesPreshadingUtils.ush b/Engine/Shaders/Private/HeterogeneousVolumes/HeterogeneousVolumesPreshadingUtils.ush
new file mode 100644
--- /dev/null
+++ b/Engine/Shaders/Private/HeterogeneousVolumes/HeterogeneousVolumesPreshadingUtils.ush
@@ -0,0 +1,190 @@
+// Copyright Epic Games, Inc. All Rights Reserved.
+
+#include "HeterogeneousVolumesStochasticFiltering.ush"
+
+int bApplyEmissionAndTransmittance;
+bool ShouldApplyEmissionAndTransmittance()
+{
+	return bApplyEmissionAndTransmittance != 0;
+}
+
+int bApplyDirectLighting;
+bool ShouldApplyDirectLighting()
+{
+	return bApplyDirectLighting != 0;
+}
+
+int LightType;
+int GetLightType()
+{
+	return LightType;
+}
+
+uint VirtualShadowMapId;
+uint GetVirtualShadowMapId()
+{
+	return VirtualShadowMapId;
+}
+
+int MipLevel;
+int GetMipLevel()
+{
+	return MipLevel;
+}
+
+int MaxStepCount;
+int GetMaxStepCount()
+{
+	return MaxStepCount;
+}
+
+int bJitter;
+bool ShouldJitter()
+{
+	return bJitter;
+}
+
+float VolumetricScatteringIntensity;
+float GetVolumetricScatteringIntensity()
+{
+	return VolumetricScatteringIntensity;
+}
+
+float MaxStepSize;
+float GetMaxStepSize()
+{
+	return MaxStepSize;
+}
+
+int3 GroupCount;
+int3 GetGroupCount()
+{
+	return GroupCount;
+}
+
+int DownsampleFactor;
+int GetDownsampleFactor()
+{
+	return DownsampleFactor;
+}
+
+float MaxTraceDistance;
+float GetMaxTraceDistance()
+{
+	return MaxTraceDistance;
+}
+
+float MaxShadowTraceDistance;
+float GetMaxShadowTraceDistance()
+{
+	return MaxShadowTraceDistance;
+}
+
+float StepSize;
+float GetStepSize()
+{
+	return StepSize;
+}
+
+// TEST..
+float StepFactor;
+float GetStepFactor()
+{
+	//float StepFactor = 1.0;
+	return StepFactor;
+}
+
+float ShadowStepSize;
+float GetShadowStepSize()
+{
+	//float ShadowStepSize = 1.0;
+	return ShadowStepSize;
+}
+
+float ShadowStepFactor;
+float GetShadowStepFactor()
+{
+	//float ShadowStepFactor = 1.0;
+	return ShadowStepFactor;
+}
+// END TEST..
+
+int GetStochasticFilteringMode()
+{
+	return STOCHASTIC_FILTERING_TRILINEAR;
+}
+
+int ShouldApplyShadowTransmittance()
+{
+	int bShouldApplyShadowTransmittance = 1;
+	return bShouldApplyShadowTransmittance;
+}
+
+//int bApplyHeightFog;
+int ShouldApplyHeightFog()
+{
+	int bApplyHeightFog = 1;
+	return bApplyHeightFog;
+}
+
+//int bApplyVolumetricFog;
+int ShouldApplyVolumetricFog()
+{
+	int bApplyVolumetricFog = 1;
+	return bApplyVolumetricFog;
+}
+
+//float IndirectInscatteringFactor;
+float GetIndirectInscatteringFactor()
+{
+	float IndirectInscatteringFactor = 1.0;
+	return IndirectInscatteringFactor;
+}
+
+// TODO: Implement ambient occlusion for preshading pipeline
+float EvalAmbientOcclusion(float3 WorldPosition)
+{
+	return 1.0;
+}
+
+// Cinematic feature options
+bool UseInscatteringVolume()
+{
+	bool bUseInscatteringVolume = true;
+	return bUseInscatteringVolume;
+}
+
+#define INDIRECT_LIGHTING_MODE_OFF 0
+#define INDIRECT_LIGHTING_MODE_LIGHTING_CACHE 1
+#define INDIRECT_LIGHTING_MODE_SINGLE_SCATTERING 2
+
+int GetIndirectLightingMode()
+{
+	return INDIRECT_LIGHTING_MODE_OFF;
+}
+
+bool ShouldWriteVelocity()
+{
+	bool bShouldWriteVelocity = false;
+	return bShouldWriteVelocity;
+}
+
+bool IsOfflineRender()
+{
+	bool bIsOfflineRender = false;
+	return bIsOfflineRender;
+}
+
+#define FOG_INSCATTERING_MODE_OFF 0
+#define FOG_INSCATTERING_MODE_REFERENCE 1
+#define FOG_INSCATTERING_MODE_LINEAR_APPROX 2
+
+int GetFogInscatteringMode()
+{
+	return FOG_INSCATTERING_MODE_LINEAR_APPROX;
+}
+
+uint3 GetVolumeResolution()
+{
+	return PreshadingUniformBuffer.VolumeResolution;
+}
diff --git a/Engine/Shaders/Private/HeterogeneousVolumes/HeterogeneousVolumesRayMarchingUtils.ush b/Engine/Shaders/Private/HeterogeneousVolumes/HeterogeneousVolumesRayMarchingUtils.ush
--- a/Engine/Shaders/Private/HeterogeneousVolumes/HeterogeneousVolumesRayMarchingUtils.ush
+++ b/Engine/Shaders/Private/HeterogeneousVolumes/HeterogeneousVolumesRayMarchingUtils.ush
@@ -102,9 +102,9 @@ float ComputeHardSurfaceShadowFactor(
 	}
 
 #if VIRTUAL_SHADOW_MAP
-	if (VirtualShadowMapId != INDEX_NONE)
+	if (GetVirtualShadowMapId() != INDEX_NONE)
 	{
-		FVirtualShadowMapSampleResult VirtualShadowMapSample = SampleVirtualShadowMapTranslatedWorld(VirtualShadowMapId, TranslatedWorldPosition);
+		FVirtualShadowMapSampleResult VirtualShadowMapSample = SampleVirtualShadowMapTranslatedWorld(GetVirtualShadowMapId(), TranslatedWorldPosition);
 		HardSurfaceShadowFactor *= VirtualShadowMapSample.ShadowFactor;
 	}
 #endif // VIRTUALSHADOW_MAP
diff --git a/Engine/Shaders/Private/HeterogeneousVolumes/HeterogeneousVolumesStochasticFiltering.ush b/Engine/Shaders/Private/HeterogeneousVolumes/HeterogeneousVolumesStochasticFiltering.ush
--- a/Engine/Shaders/Private/HeterogeneousVolumes/HeterogeneousVolumesStochasticFiltering.ush
+++ b/Engine/Shaders/Private/HeterogeneousVolumes/HeterogeneousVolumesStochasticFiltering.ush
@@ -1,5 +1,7 @@
 // Copyright Epic Games, Inc. All Rights Reserved.
 
+#pragma once
+
 int3 WeightedReservoir(int i, int j, float Prob, int3 Index, inout float Rand)
 {
 	if (Rand < Prob)
diff --git a/Engine/Source/Runtime/Renderer/Private/HeterogeneousVolumes/HeterogeneousVolumes.h b/Engine/Source/Runtime/Renderer/Private/HeterogeneousVolumes/HeterogeneousVolumes.h
--- a/Engine/Source/Runtime/Renderer/Private/HeterogeneousVolumes/HeterogeneousVolumes.h
+++ b/Engine/Source/Runtime/Renderer/Private/HeterogeneousVolumes/HeterogeneousVolumes.h
@@ -267,6 +267,21 @@ BEGIN_UNIFORM_BUFFER_STRUCT(FSparseVoxelUniformBufferParameters, )
 	SHADER_PARAMETER(int, bApplyVolumetricFog)
 END_UNIFORM_BUFFER_STRUCT()
 
+BEGIN_UNIFORM_BUFFER_STRUCT(FPreshadingUniformBufferParameters, )
+	// Object data
+	SHADER_PARAMETER(FMatrix44f, LocalToWorld)
+	SHADER_PARAMETER(FMatrix44f, WorldToLocal)
+	SHADER_PARAMETER(FVector3f, LocalBoundsOrigin)
+	SHADER_PARAMETER(FVector3f, LocalBoundsExtent)
+
+	// Volume data
+	SHADER_PARAMETER(FIntVector, VolumeResolution)
+	SHADER_PARAMETER_RDG_TEXTURE(Texture3D, ExtinctionTexture)
+	SHADER_PARAMETER_RDG_TEXTURE(Texture3D, EmissionTexture)
+	SHADER_PARAMETER_RDG_TEXTURE(Texture3D, AlbedoTexture)
+	SHADER_PARAMETER_SAMPLER(SamplerState, TextureSampler)
+END_UNIFORM_BUFFER_STRUCT()
+
 BEGIN_SHADER_PARAMETER_STRUCT(FLightingCacheParameters, )
 	SHADER_PARAMETER(FIntVector, LightingCacheResolution)
 	SHADER_PARAMETER(float, LightingCacheVoxelBias)
diff --git a/Engine/Source/Runtime/Renderer/Private/HeterogeneousVolumes/HeterogeneousVolumesPreshadingPipeline.cpp b/Engine/Source/Runtime/Renderer/Private/HeterogeneousVolumes/HeterogeneousVolumesPreshadingPipeline.cpp
--- a/Engine/Source/Runtime/Renderer/Private/HeterogeneousVolumes/HeterogeneousVolumesPreshadingPipeline.cpp
+++ b/Engine/Source/Runtime/Renderer/Private/HeterogeneousVolumes/HeterogeneousVolumesPreshadingPipeline.cpp
@@ -14,194 +14,35 @@
 #include "VolumeLighting.h"
 #include "VolumetricFog.h"
 
-class FGenerateRayMarchingTiles : public FGlobalShader
-{
-	DECLARE_GLOBAL_SHADER(FGenerateRayMarchingTiles);
-	SHADER_USE_PARAMETER_STRUCT(FGenerateRayMarchingTiles, FGlobalShader);
-
-	class FDebugDim : SHADER_PERMUTATION_BOOL("DIM_DEBUG");
-	class FVoxelCullingDim : SHADER_PERMUTATION_BOOL("DIM_VOXEL_CULLING");
-	using FPermutationDomain = TShaderPermutationDomain<FDebugDim, FVoxelCullingDim>;
-
-	BEGIN_SHADER_PARAMETER_STRUCT(FParameters, )
-		// Scene data
-		SHADER_PARAMETER_STRUCT_REF(FViewUniformShaderParameters, View)
-		SHADER_PARAMETER_STRUCT_INCLUDE(FSceneTextureParameters, SceneTextures)
-
-		// Object data
-		SHADER_PARAMETER(FMatrix44f, LocalToWorld)
-		SHADER_PARAMETER(FMatrix44f, WorldToLocal)
-
-		// Sparse voxel data
-		SHADER_PARAMETER_RDG_UNIFORM_BUFFER(FSparseVoxelUniformBufferParameters, SparseVoxelUniformBuffer)
-
-		// Ray data
-		SHADER_PARAMETER(float, MaxTraceDistance)
-		SHADER_PARAMETER(float, StepSize)
-		SHADER_PARAMETER(int, MaxStepCount)
-		SHADER_PARAMETER(int, bJitter)
-
-		// Dispatch data
-		SHADER_PARAMETER(FIntVector, GroupCount)
-		SHADER_PARAMETER(int32, DownsampleFactor)
-
-		// Debug Output
-		SHADER_PARAMETER_RDG_BUFFER_UAV(RWStructuredBuffer<Volumes::FRayMarchingDebug>, RWRayMarchingDebugBuffer)
-
-		// Output
-		SHADER_PARAMETER_RDG_BUFFER_UAV(RWBuffer<uint>, RWNumRayMarchingTilesBuffer)
-		SHADER_PARAMETER_RDG_BUFFER_UAV(RWStructuredBuffer<Volumes::FRayMarchingTile>, RWRayMarchingTilesBuffer)
-
-		SHADER_PARAMETER_RDG_BUFFER_UAV(RWBuffer<uint>, RWNumVoxelsPerTileBuffer)
-		SHADER_PARAMETER_RDG_BUFFER_UAV(RWStructuredBuffer<FVoxelDataPacked>, RWVoxelsPerTileBuffer)
-	END_SHADER_PARAMETER_STRUCT()
-
-	static bool ShouldCompilePermutation(
-		const FGlobalShaderPermutationParameters& Parameters
-	)
-	{
-		return DoesPlatformSupportHeterogeneousVolumes(Parameters.Platform);
-	}
-
-	static void ModifyCompilationEnvironment(
-		const FGlobalShaderPermutationParameters& Parameters,
-		FShaderCompilerEnvironment& OutEnvironment
-	)
-	{
-		FMaterialShader::ModifyCompilationEnvironment(Parameters, OutEnvironment);
-		OutEnvironment.SetDefine(TEXT("THREADGROUP_SIZE_1D"), GetThreadGroupSize1D());
-		OutEnvironment.SetDefine(TEXT("THREADGROUP_SIZE_2D"), GetThreadGroupSize2D());
-
-		// This shader takes a very long time to compile with FXC, so we pre-compile it with DXC first and then forward the optimized HLSL to FXC.
-		OutEnvironment.CompilerFlags.Add(CFLAG_AllowTypedUAVLoads);
-	}
-
-	static int32 GetThreadGroupSize1D() { return GetThreadGroupSize2D() * GetThreadGroupSize2D(); }
-	static int32 GetThreadGroupSize2D() { return 8; }
-};
-
-IMPLEMENT_GLOBAL_SHADER(FGenerateRayMarchingTiles, "/Engine/Private/HeterogeneousVolumes/HeterogeneousVolumesPreshadingPipeline.usf", "GenerateRayMarchingTiles", SF_Compute);
-
-namespace HeterogeneousVolumes
-{
-	struct FRayMarchingTile
+IMPLEMENT_UNIFORM_BUFFER_STRUCT(FPreshadingUniformBufferParameters, "PreshadingUniformBuffer");
+
+static TAutoConsoleVariable<int32> CVarHeterogeneousVolumesPreshadingPixelFormat(
+	TEXT("r.HeterogeneousVolumes.Preshading.PixelFormat"),
+	2,
+	TEXT("Determines the pixel format for preshaded results (Default = 2)\n")
+	TEXT("0: PF_R16F\n")
+	TEXT("1: PF_FloatR11G11B10\n")
+	TEXT("2: PF_FloatRGB\n"),
+	ECVF_RenderThreadSafe
+);
+
+namespace HeterogeneousVolumes {
+	EPixelFormat GetPreshadingPixelFormat()
 	{
-		FIntPoint PixelOffset;
-		uint32 Voxels[2];
-
-		uint32 Id;
-		uint32 Padding[3];
-	};
-
-	struct FRayMarchingDebug
-	{
-		FVector4f Planes[5];
-		FVector4f BBox[2];
-
-		float Padding[4];
-	};
-}
-
-void GenerateRayMarchingTiles(
-	FRDGBuilder& GraphBuilder,
-	// Scene data
-	const FScene* Scene,
-	const FViewInfo& View,
-	const FSceneTextures& SceneTextures,
-	// Object data
-	const IHeterogeneousVolumeInterface* HeterogeneousVolumeInterface,
-	// Sparse voxel data
-	FRDGBufferRef NumVoxelsBuffer,
-	const TRDGUniformBufferRef<FSparseVoxelUniformBufferParameters>& SparseVoxelUniformBuffer,
-	// Output
-	FRDGBufferRef& NumRayMarchingTilesBuffer,
-	FRDGBufferRef& RayMarchingTilesBuffer,
-	FRDGBufferRef& VoxelsPerTileBuffer
-)
-{
-	FIntVector GroupCount = FComputeShaderUtils::GetGroupCount(HeterogeneousVolumes::GetScaledViewRect(View.ViewRect), FGenerateRayMarchingTiles::GetThreadGroupSize2D());
-	uint32 NumTiles = GroupCount.X * GroupCount.Y;
-
-	NumRayMarchingTilesBuffer = GraphBuilder.CreateBuffer(
-		FRDGBufferDesc::CreateIndirectDesc<FRHIDispatchIndirectParameters>(1),
-		TEXT("HeterogeneousVolume.NumRayMarchingTilesBuffer")
-	);
-	AddClearUAVPass(GraphBuilder, GraphBuilder.CreateUAV(NumRayMarchingTilesBuffer, PF_R32_UINT), 0);
-
-	RayMarchingTilesBuffer = GraphBuilder.CreateBuffer(
-		FRDGBufferDesc::CreateStructuredDesc(sizeof(HeterogeneousVolumes::FRayMarchingTile), NumTiles),
-		TEXT("HeterogeneousVolumes.RayMarchingTileBuffer")
-	);
-
-	FRDGBufferRef RayMarchingDebugBuffer = GraphBuilder.CreateBuffer(
-		FRDGBufferDesc::CreateStructuredDesc(sizeof(HeterogeneousVolumes::FRayMarchingDebug), NumTiles),
-		TEXT("HeterogeneousVolume.RayMarchingDebugBuffer")
-	);
-
-	FRDGBufferRef NumVoxelsPerTileBuffer = GraphBuilder.CreateBuffer(
-		FRDGBufferDesc::CreateBufferDesc(sizeof(uint32), NumTiles),
-		TEXT("HeterogeneousVolume.NumVoxelsPerTileBuffer")
-	);
-	AddClearUAVPass(GraphBuilder, GraphBuilder.CreateUAV(NumVoxelsPerTileBuffer, PF_R32_UINT), 0);
-
-	FIntVector VolumeResolution = SparseVoxelUniformBuffer->GetParameters()->VolumeResolution;
-	uint32 SparseMipLevel = SparseVoxelUniformBuffer->GetParameters()->MipLevel;
-	FIntVector SparseVolumeResolution = FIntVector(VolumeResolution.X >> SparseMipLevel,
-		VolumeResolution.Y >> SparseMipLevel,
-		VolumeResolution.Z >> SparseMipLevel
-	);
-
-	// TODO: Tight frustum culling guarantees no more than Length(SparseVolumeResolution) but approximate intersection cannot guarantee even L1 distance..
-	//uint32 DiagonalLength = FMath::CeilToInt(FMath::Sqrt(SparseVolumeResolution.X * SparseVolumeResolution.Y * SparseVolumeResolution.Z));
-	uint32 DiagonalLength = SparseVolumeResolution.X * SparseVolumeResolution.Y * SparseVolumeResolution.Z;
-
-	VoxelsPerTileBuffer = GraphBuilder.CreateBuffer(
-		FRDGBufferDesc::CreateStructuredDesc(sizeof(FVoxelDataPacked), NumTiles * DiagonalLength),
-		TEXT("HeterogeneousVolumes.VoxelsPerTileBuffer")
-	);
-
-	FGenerateRayMarchingTiles::FParameters* PassParameters = GraphBuilder.AllocParameters<FGenerateRayMarchingTiles::FParameters>();
-	{
-		// Scene data
-		PassParameters->View = View.ViewUniformBuffer;
-
-		// Sparse voxel data
-		PassParameters->SparseVoxelUniformBuffer = SparseVoxelUniformBuffer;
-
-		// Ray data
-		PassParameters->MaxTraceDistance = HeterogeneousVolumes::GetMaxTraceDistance();
-		PassParameters->StepSize = HeterogeneousVolumes::GetStepSize();
-		PassParameters->MaxStepCount = HeterogeneousVolumes::GetMaxStepCount();
-		PassParameters->bJitter = HeterogeneousVolumes::ShouldJitter();
-
-		// Dispatch data
-		PassParameters->GroupCount = GroupCount;
-		PassParameters->DownsampleFactor = HeterogeneousVolumes::GetDownsampleFactor();
-
-		// Debug
-		PassParameters->RWRayMarchingDebugBuffer = GraphBuilder.CreateUAV(RayMarchingDebugBuffer);
-
-		// Output
-		PassParameters->RWNumRayMarchingTilesBuffer = GraphBuilder.CreateUAV(NumRayMarchingTilesBuffer, PF_R32_UINT);
-		PassParameters->RWRayMarchingTilesBuffer = GraphBuilder.CreateUAV(RayMarchingTilesBuffer);
-
-		PassParameters->RWNumVoxelsPerTileBuffer = GraphBuilder.CreateUAV(NumVoxelsPerTileBuffer, PF_R32_UINT);
-		PassParameters->RWVoxelsPerTileBuffer = GraphBuilder.CreateUAV(VoxelsPerTileBuffer);
+		int Value = CVarHeterogeneousVolumesPreshadingPixelFormat.GetValueOnAnyThread();
+		switch (Value)
+		{
+		case 0:
+			return PF_R16F;
+		case 1:
+			return PF_FloatR11G11B10;
+		case 2:
+		default:
+			return PF_FloatRGB;
+		}
 	}
 
-	FGenerateRayMarchingTiles::FPermutationDomain PermutationVector;
-	PermutationVector.Set<FGenerateRayMarchingTiles::FDebugDim>(HeterogeneousVolumes::GetDebugMode() != 0);
-	PermutationVector.Set<FGenerateRayMarchingTiles::FVoxelCullingDim>(HeterogeneousVolumes::UseSparseVoxelPerTileCulling());
-
-	TShaderRef<FGenerateRayMarchingTiles> ComputeShader = View.ShaderMap->GetShader<FGenerateRayMarchingTiles>(PermutationVector);
-	FComputeShaderUtils::AddPass(
-		GraphBuilder,
-		RDG_EVENT_NAME("FGenerateRayMarchingTiles"),
-		ComputeShader,
-		PassParameters,
-		GroupCount
-	);
-}
+} // namespace HeterogeneousVolumes
 
 class FRenderLightingCacheWithPreshadingCS : public FGlobalShader
 {
@@ -242,15 +83,21 @@ class FRenderLightingCacheWithPreshadingCS : public FGlobalShader
 		SHADER_PARAMETER(float, IndirectInscatteringFactor)
 
 		// Volume structures
-		SHADER_PARAMETER_RDG_UNIFORM_BUFFER(FSparseVoxelUniformBufferParameters, SparseVoxelUniformBuffer)
+		SHADER_PARAMETER_RDG_UNIFORM_BUFFER(FPreshadingUniformBufferParameters, PreshadingUniformBuffer)
 		SHADER_PARAMETER_STRUCT_INCLUDE(FLightingCacheParameters, LightingCache)
 
 		// Ray data
-		SHADER_PARAMETER(float, MaxShadowTraceDistance)
 		SHADER_PARAMETER(float, StepSize)
-		SHADER_PARAMETER(int, MipLevel)
+		SHADER_PARAMETER(float, StepFactor)
+		SHADER_PARAMETER(float, ShadowStepSize)
+		SHADER_PARAMETER(float, ShadowStepFactor)
+
+		SHADER_PARAMETER(float, MaxTraceDistance)
+		SHADER_PARAMETER(float, MaxShadowTraceDistance)
 		SHADER_PARAMETER(int, MaxStepCount)
+		SHADER_PARAMETER(int, MipLevel)
 		SHADER_PARAMETER(int, bJitter)
+		SHADER_PARAMETER(int, StochasticFilteringMode)
 
 		// Output
 		SHADER_PARAMETER_RDG_TEXTURE_UAV(RWTexture3D<float>, RWLightingCacheTexture)
@@ -352,25 +199,25 @@ class FRenderSingleScatteringWithPreshadingCS : public FGlobalShader
 		// Indirect Lighting
 		SHADER_PARAMETER_RDG_UNIFORM_BUFFER(FLumenTranslucencyLightingUniforms, LumenGIVolumeStruct)
 
-		// Volume data
-		SHADER_PARAMETER(int, MipLevel)
-
 		// Volume structures
-		SHADER_PARAMETER_RDG_UNIFORM_BUFFER(FSparseVoxelUniformBufferParameters, SparseVoxelUniformBuffer)
+		SHADER_PARAMETER_RDG_UNIFORM_BUFFER(FPreshadingUniformBufferParameters, PreshadingUniformBuffer)
 		SHADER_PARAMETER_STRUCT_INCLUDE(FLightingCacheParameters, LightingCache)
 
 		// Ray data
-		SHADER_PARAMETER(float, MaxTraceDistance)
 		SHADER_PARAMETER(float, StepSize)
+		SHADER_PARAMETER(float, StepFactor)
+		SHADER_PARAMETER(float, ShadowStepSize)
+		SHADER_PARAMETER(float, ShadowStepFactor)
+
+		SHADER_PARAMETER(float, MaxTraceDistance)
+		SHADER_PARAMETER(float, MaxShadowTraceDistance)
 		SHADER_PARAMETER(int, MaxStepCount)
+		SHADER_PARAMETER(int, MipLevel)
 		SHADER_PARAMETER(int, bJitter)
+		SHADER_PARAMETER(int, StochasticFilteringMode)
 
-		// Ray marching data
-		SHADER_PARAMETER_RDG_BUFFER_SRV(StructuredBuffer<Volumes::FRayMarchingTile>, RayMarchingTilesBuffer)
-		SHADER_PARAMETER_RDG_BUFFER_SRV(StructuredBuffer<FVoxelDataPacked>, VoxelsPerTileBuffer)
-
-		// Indirect args
-		RDG_BUFFER_ACCESS(IndirectArgs, ERHIAccess::IndirectArgs)
+		// Dispatch data
+		SHADER_PARAMETER(FIntVector, GroupCount)
 		SHADER_PARAMETER(int32, DownsampleFactor)
 
 		// Output
@@ -456,6 +303,7 @@ class FRenderSingleScatteringWithPreshadingCS : public FGlobalShader
 
 IMPLEMENT_GLOBAL_SHADER(FRenderSingleScatteringWithPreshadingCS, "/Engine/Private/HeterogeneousVolumes/HeterogeneousVolumesPreshadingPipeline.usf", "RenderSingleScatteringWithPreshadingCS", SF_Compute);
 
+
 void RenderLightingCacheWithPreshadingCompute(
 	FRDGBuilder& GraphBuilder,
 	// Scene data
@@ -473,13 +321,8 @@ void RenderLightingCacheWithPreshadingCompute(
 	const FVirtualShadowMapArray& VirtualShadowMapArray,
 	// Object data
 	const IHeterogeneousVolumeInterface* HeterogeneousVolumeInterface,
-	// Sparse voxel data
-	FRDGBufferRef NumVoxelsBuffer,
-	const TRDGUniformBufferRef<FSparseVoxelUniformBufferParameters>& SparseVoxelUniformBuffer,
-	// Ray marching tiles
-	FRDGBufferRef NumRayMarchingTilesBuffer,
-	FRDGBufferRef RayMarchingTilesBuffer,
-	FRDGBufferRef VoxelsPerTileBuffer,
+	// Preshading data
+	const TRDGUniformBufferRef<FPreshadingUniformBufferParameters>& PreshadingUniformBuffer,
 	// Output
 	FRDGTextureRef& LightingCacheTexture
 )
@@ -489,6 +332,7 @@ void RenderLightingCacheWithPreshadingCompute(
 	TUniformBufferRef<FDeferredLightUniformStruct> DeferredLightUB = CreateUniformBufferImmediate(DeferredLightUniform, UniformBuffer_SingleDraw);
 
 	HeterogeneousVolumes::FLODValue LODValue = HeterogeneousVolumes::CalcLOD(View, HeterogeneousVolumeInterface);
+	float LODFactor = HeterogeneousVolumes::CalcLODFactor(LODValue.LOD, LODValue.Bias);
 
 	FRenderLightingCacheWithPreshadingCS::FParameters* PassParameters = GraphBuilder.AllocParameters<FRenderLightingCacheWithPreshadingCS::FParameters>();
 	{
@@ -506,7 +350,7 @@ void RenderLightingCacheWithPreshadingCompute(
 		PassParameters->VolumetricScatteringIntensity = LightSceneInfo->Proxy->GetVolumetricScatteringIntensity();
 
 		// Sparse voxel data
-		PassParameters->SparseVoxelUniformBuffer = SparseVoxelUniformBuffer;
+		PassParameters->PreshadingUniformBuffer = PreshadingUniformBuffer;
 
 		// Transmittance volume
 		PassParameters->LightingCache.LightingCacheResolution = HeterogeneousVolumes::GetLightingCacheResolution(HeterogeneousVolumeInterface, LODValue);
@@ -514,10 +358,17 @@ void RenderLightingCacheWithPreshadingCompute(
 		PassParameters->LightingCache.LightingCacheTexture = LightingCacheTexture;
 
 		// Ray data
-		//PassParameters->StepSize = HeterogeneousVolumes::GetStepSize();
+		PassParameters->StepSize = HeterogeneousVolumes::GetStepSize();
+		PassParameters->StepFactor = HeterogeneousVolumeInterface->GetStepFactor() * LODFactor;
+		PassParameters->ShadowStepSize = HeterogeneousVolumes::GetShadowStepSize();
+		PassParameters->ShadowStepFactor = HeterogeneousVolumeInterface->GetShadowStepFactor() * LODFactor;
+
+		PassParameters->MaxTraceDistance = HeterogeneousVolumes::GetMaxTraceDistance();
+		PassParameters->MaxShadowTraceDistance = HeterogeneousVolumes::GetMaxShadowTraceDistance();
 		PassParameters->MaxStepCount = HeterogeneousVolumes::GetMaxStepCount();
-		PassParameters->bJitter = HeterogeneousVolumes::ShouldJitter();
 		PassParameters->MipLevel = HeterogeneousVolumes::GetMipLevel();
+		PassParameters->bJitter = HeterogeneousVolumes::ShouldJitter();
+		PassParameters->StochasticFilteringMode = static_cast<int32>(HeterogeneousVolumes::GetStochasticFilteringMode());
 
 		// Shadow data
 		PassParameters->ForwardLightStruct = View.ForwardLightingResources.ForwardLightUniformBuffer;
@@ -611,22 +462,14 @@ void RenderSingleScatteringWithPreshadingCompute(
 	const FVirtualShadowMapArray& VirtualShadowMapArray,
 	// Object data
 	const IHeterogeneousVolumeInterface* HeterogeneousVolumeInterface,
-	// Sparse voxel data
-	FRDGBufferRef NumVoxelsBuffer,
-	const TRDGUniformBufferRef<FSparseVoxelUniformBufferParameters>& SparseVoxelUniformBuffer,
+	// Preshading data
+	const TRDGUniformBufferRef<FPreshadingUniformBufferParameters>& PreshadingUniformBuffer,
 	FRDGTextureRef LightingCacheTexture,
-	// Ray marching tiles
-	FRDGBufferRef NumRayMarchingTilesBuffer,
-	FRDGBufferRef RayMarchingTilesBuffer,
-	FRDGBufferRef VoxelsPerTileBuffer,
 	// Output
 	FRDGTextureRef& HeterogeneousVolumeTexture
 )
 {
-	FRDGBufferRef VoxelOutputBuffer = GraphBuilder.CreateBuffer(
-		FRDGBufferDesc::CreateStructuredDesc(sizeof(FVoxelDataPacked), HeterogeneousVolumes::GetVoxelCount(SparseVoxelUniformBuffer->GetParameters()->VolumeResolution)),
-		TEXT("HeterogeneousVolumes.VoxelOutputBuffer")
-	);
+	FIntVector GroupCount = FComputeShaderUtils::GetGroupCount(HeterogeneousVolumes::GetScaledViewRect(View.ViewRect), FRenderSingleScatteringWithPreshadingCS::GetThreadGroupSize2D());
 
 	// Note must be done in the same scope as we add the pass otherwise the UB lifetime will not be guaranteed
 	FDeferredLightUniformStruct DeferredLightUniform;
@@ -636,6 +479,9 @@ void RenderSingleScatteringWithPreshadingCompute(
 	}
 	TUniformBufferRef<FDeferredLightUniformStruct> DeferredLightUB = CreateUniformBufferImmediate(DeferredLightUniform, UniformBuffer_SingleDraw);
 
+	HeterogeneousVolumes::FLODValue LODValue = HeterogeneousVolumes::CalcLOD(View, HeterogeneousVolumeInterface);
+	//FIntVector LightingCacheResolution = HeterogeneousVolumes::GetLightingCacheResolution(HeterogeneousVolumeInterface, LODValue);
+
 	bool bWriteVelocity = HeterogeneousVolumes::ShouldWriteVelocity() && HasBeenProduced(SceneTextures.Velocity);
 	FRenderSingleScatteringWithPreshadingCS::FParameters* PassParameters = GraphBuilder.AllocParameters<FRenderSingleScatteringWithPreshadingCS::FParameters>();
 	{
@@ -690,12 +536,11 @@ void RenderSingleScatteringWithPreshadingCompute(
 		PassParameters->MipLevel = HeterogeneousVolumes::GetMipLevel();
 
 		// Sparse voxel data
-		PassParameters->SparseVoxelUniformBuffer = SparseVoxelUniformBuffer;
+		PassParameters->PreshadingUniformBuffer = PreshadingUniformBuffer;
 
 		// Transmittance volume
 		if ((HeterogeneousVolumes::UseLightingCacheForTransmittance() && bApplyShadowTransmittance) || HeterogeneousVolumes::UseLightingCacheForInscattering())
 		{
-			HeterogeneousVolumes::FLODValue LODValue = HeterogeneousVolumes::CalcLOD(View, HeterogeneousVolumeInterface);
 			PassParameters->LightingCache.LightingCacheResolution = HeterogeneousVolumes::GetLightingCacheResolution(HeterogeneousVolumeInterface, LODValue);
 			PassParameters->LightingCache.LightingCacheVoxelBias = HeterogeneousVolumeInterface->GetShadowBiasFactor();
 			PassParameters->LightingCache.LightingCacheTexture = LightingCacheTexture;
@@ -708,17 +553,19 @@ void RenderSingleScatteringWithPreshadingCompute(
 		}
 
 		// Ray data
+		float LODFactor = HeterogeneousVolumes::CalcLODFactor(LODValue.LOD, LODValue.Bias);
 		PassParameters->MaxTraceDistance = HeterogeneousVolumes::GetMaxTraceDistance();
+		PassParameters->MaxShadowTraceDistance = HeterogeneousVolumes::GetMaxShadowTraceDistance();
 		PassParameters->StepSize = HeterogeneousVolumes::GetStepSize();
+		PassParameters->StepFactor = HeterogeneousVolumeInterface->GetStepFactor() * LODFactor;
+		PassParameters->ShadowStepSize = HeterogeneousVolumes::GetShadowStepSize();
+		PassParameters->ShadowStepFactor = HeterogeneousVolumeInterface->GetShadowStepFactor() * LODFactor;
 		PassParameters->MaxStepCount = HeterogeneousVolumes::GetMaxStepCount();
 		PassParameters->bJitter = HeterogeneousVolumes::ShouldJitter();
-
-		// Ray marching data
-		PassParameters->RayMarchingTilesBuffer = GraphBuilder.CreateSRV(RayMarchingTilesBuffer);
-		PassParameters->VoxelsPerTileBuffer = GraphBuilder.CreateSRV(VoxelsPerTileBuffer);
+		PassParameters->StochasticFilteringMode = static_cast<int32>(HeterogeneousVolumes::GetStochasticFilteringMode());
 
 		// Dispatch data
-		PassParameters->IndirectArgs = NumRayMarchingTilesBuffer;
+		PassParameters->GroupCount = GroupCount;
 		PassParameters->DownsampleFactor = HeterogeneousVolumes::GetDownsampleFactor();
 
 		// Output
@@ -727,7 +574,6 @@ void RenderSingleScatteringWithPreshadingCompute(
 		{
 			PassParameters->RWVelocityTexture = GraphBuilder.CreateUAV(SceneTextures.Velocity);
 		}
-		PassParameters->RWVoxelOutputBuffer = GraphBuilder.CreateUAV(VoxelOutputBuffer);
 	}
 
 	FString LightName = "none";
@@ -758,8 +604,8 @@ void RenderSingleScatteringWithPreshadingCompute(
 		RDG_EVENT_NAME("RenderSingleScatteringWithPreshadingCS (Light = %s)", *LightName),
 		ComputeShader,
 		PassParameters,
-		PassParameters->IndirectArgs,
-		0);
+		GroupCount
+	);
 }
 
 void RenderWithInscatteringVolumePipelineWithPreshadingCompute(
@@ -774,13 +620,8 @@ void RenderWithInscatteringVolumePipelineWithPreshadingCompute(
 	// Object data
 	const IHeterogeneousVolumeInterface* HeterogeneousVolumeInterface,
 	const FMaterialRenderProxy* MaterialRenderProxy,
-	// Sparse voxel data
-	FRDGBufferRef NumVoxelsBuffer,
-	TRDGUniformBufferRef<FSparseVoxelUniformBufferParameters>& SparseVoxelUniformBuffer,
-	// Render tile data
-	FRDGBufferRef NumRayMarchingTilesBuffer,
-	FRDGBufferRef RayMarchingTilesBuffer,
-	FRDGBufferRef VoxelsPerTileBuffer,
+	// Prreshading data
+	TRDGUniformBufferRef<FPreshadingUniformBufferParameters>& PreshadingUniformBuffer,
 	// Output
 	FRDGTextureRef& LightingCacheTexture,
 	FRDGTextureRef& HeterogeneousVolumeRadiance
@@ -847,14 +688,8 @@ void RenderWithInscatteringVolumePipelineWithPreshadingCompute(
 				VirtualShadowMapArray,
 				// Object data
 				HeterogeneousVolumeInterface,
-				// Volume data
-				NumVoxelsBuffer,
 				// Sparse voxel data
-				SparseVoxelUniformBuffer,
-				// Ray marching tile
-				NumRayMarchingTilesBuffer,
-				RayMarchingTilesBuffer,
-				VoxelsPerTileBuffer,
+				PreshadingUniformBuffer,
 				// Output
 				LightingCacheTexture
 			);
@@ -890,142 +725,9 @@ void RenderWithInscatteringVolumePipelineWithPreshadingCompute(
 			VirtualShadowMapArray,
 			// Object
 			HeterogeneousVolumeInterface,
-			// Volume data
-			NumVoxelsBuffer,
 			// Sparse voxel data
-			SparseVoxelUniformBuffer,
+			PreshadingUniformBuffer,
 			LightingCacheTexture,
-			// Ray marching tile
-			NumRayMarchingTilesBuffer,
-			RayMarchingTilesBuffer,
-			VoxelsPerTileBuffer,
-			// Output
-			HeterogeneousVolumeRadiance
-		);
-	}
-}
-
-void RenderWithTransmittanceVolumePipelineWithPreshadingCompute(
-	FRDGBuilder& GraphBuilder,
-	// Scene data
-	const FSceneTextures& SceneTextures,
-	const FScene* Scene,
-	FViewInfo& View, int32 ViewIndex,
-	// Shadow data
-	TArray<FVisibleLightInfo, SceneRenderingAllocator>& VisibleLightInfos,
-	const FVirtualShadowMapArray& VirtualShadowMapArray,
-	// Object data
-	const IHeterogeneousVolumeInterface* HeterogeneousVolumeInterface,
-	const FMaterialRenderProxy* MaterialRenderProxy,
-	// Sparse voxel data
-	FRDGBufferRef NumVoxelsBuffer,
-	TRDGUniformBufferRef<FSparseVoxelUniformBufferParameters>& SparseVoxelUniformBuffer,
-	// Render tile data
-	FRDGBufferRef NumRayMarchingTilesBuffer,
-	FRDGBufferRef RayMarchingTilesBuffer,
-	FRDGBufferRef VoxelsPerTileBuffer,
-	// Output
-	FRDGTextureRef& LightingCacheTexture,
-	FRDGTextureRef& HeterogeneousVolumeRadiance
-)
-{
-	RDG_EVENT_SCOPE(GraphBuilder, "Direct Volume Rendering");
-
-	// Light culling
-	TArray<FLightSceneInfoCompact, TInlineAllocator<64>> LightSceneInfoCompact;
-	for (auto LightIt = Scene->Lights.CreateConstIterator(); LightIt; ++LightIt)
-	{
-		if (HeterogeneousVolumes::SupportsLightType(LightIt->LightType) &&
-			(LightIt->LightSceneInfo->Proxy->GetViewLightingChannelMask() & View.ViewLightingChannelMask) &&
-			LightIt->AffectsPrimitive(HeterogeneousVolumeInterface->GetBounds(), HeterogeneousVolumeInterface->GetPrimitiveSceneProxy()))
-		{
-			LightSceneInfoCompact.Add(*LightIt);
-		}
-	}
-
-	// Single-scattering
-	int32 NumPasses = FMath::Max(LightSceneInfoCompact.Num(), 1);
-	for (int32 PassIndex = 0; PassIndex < NumPasses; ++PassIndex)
-	{
-		bool bApplyEmissionAndTransmittance = PassIndex == 0;
-		bool bApplyDirectLighting = !LightSceneInfoCompact.IsEmpty();
-		bool bApplyShadowTransmittance = false;
-
-		uint32 LightType = 0;
-		FLightSceneInfo* LightSceneInfo = nullptr;
-		const FVisibleLightInfo* VisibleLightInfo = nullptr;
-		if (bApplyDirectLighting)
-		{
-			LightType = LightSceneInfoCompact[PassIndex].LightType;
-			LightSceneInfo = LightSceneInfoCompact[PassIndex].LightSceneInfo;
-			check(LightSceneInfo != nullptr);
-
-			bApplyDirectLighting = (LightSceneInfo != nullptr);
-			if (LightSceneInfo)
-			{
-				VisibleLightInfo = &VisibleLightInfos[LightSceneInfo->Id];
-				bApplyShadowTransmittance = LightSceneInfo->Proxy->CastsVolumetricShadow();
-			}
-		}
-
-		if (HeterogeneousVolumes::UseLightingCacheForTransmittance() && bApplyShadowTransmittance)
-		{
-			RenderLightingCacheWithPreshadingCompute(
-				GraphBuilder,
-				// Scene data
-				Scene,
-				View, ViewIndex,
-				SceneTextures,
-				// Light data
-				bApplyEmissionAndTransmittance,
-				bApplyDirectLighting,
-				bApplyShadowTransmittance,
-				LightType,
-				LightSceneInfo,
-				// Shadow data
-				VisibleLightInfo,
-				VirtualShadowMapArray,
-				// Object data
-				HeterogeneousVolumeInterface,
-				// Volume data
-				NumVoxelsBuffer,
-				// Sparse voxel data
-				SparseVoxelUniformBuffer,
-				// Ray marching tile
-				NumRayMarchingTilesBuffer,
-				RayMarchingTilesBuffer,
-				VoxelsPerTileBuffer,
-				// Output
-				LightingCacheTexture
-			);
-		}
-
-		RenderSingleScatteringWithPreshadingCompute(
-			GraphBuilder,
-			// Scene
-			Scene,
-			View, ViewIndex,
-			SceneTextures,
-			// Light
-			bApplyEmissionAndTransmittance,
-			bApplyDirectLighting,
-			bApplyShadowTransmittance,
-			LightType,
-			LightSceneInfo,
-			// Shadow
-			VisibleLightInfo,
-			VirtualShadowMapArray,
-			// Object
-			HeterogeneousVolumeInterface,
-			// Volume data
-			NumVoxelsBuffer,
-			// Sparse voxel data
-			SparseVoxelUniformBuffer,
-			LightingCacheTexture,
-			// Ray marching tile
-			NumRayMarchingTilesBuffer,
-			RayMarchingTilesBuffer,
-			VoxelsPerTileBuffer,
 			// Output
 			HeterogeneousVolumeRadiance
 		);
@@ -1044,9 +746,8 @@ void RenderWithPreshadingCompute(
 	// Object data
 	const IHeterogeneousVolumeInterface* HeterogeneousVolumeInterface,
 	const FMaterialRenderProxy* MaterialRenderProxy,
-	// Sparse voxel data
-	FRDGBufferRef NumVoxelsBuffer,
-	TRDGUniformBufferRef<FSparseVoxelUniformBufferParameters>& SparseVoxelUniformBuffer,
+	// Preshading data
+	TRDGUniformBufferRef<FPreshadingUniformBufferParameters>& PreshadingUniformBuffer,
 	// Output
 	FRDGTextureRef& LightingCacheTexture,
 	FRDGTextureRef& HeterogeneousVolumeRadiance
@@ -1054,82 +755,24 @@ void RenderWithPreshadingCompute(
 {
 	RDG_EVENT_SCOPE(GraphBuilder, "Software Ray Tracing");
 
-	FRDGBufferRef NumRayMarchingTilesBuffer;
-	FRDGBufferRef RayMarchingTilesBuffer;
-	FRDGBufferRef VoxelsPerTileBuffer;
-	{
-		RDG_EVENT_SCOPE(GraphBuilder, "Ray Tile Generation");
-		GenerateRayMarchingTiles(
-			GraphBuilder,
-			// Scene
-			Scene,
-			View,
-			SceneTextures,
-			// Object
-			HeterogeneousVolumeInterface,
-			// Volume data
-			NumVoxelsBuffer,
-			// Sparse voxel data
-			SparseVoxelUniformBuffer,
-			// Output
-			NumRayMarchingTilesBuffer,
-			RayMarchingTilesBuffer,
-			VoxelsPerTileBuffer
-		);
-	}
-
-	if (HeterogeneousVolumes::UseLightingCacheForInscattering())
-	{
-		RenderWithInscatteringVolumePipelineWithPreshadingCompute(
-			GraphBuilder,
-			// Scene data
-			SceneTextures,
-			Scene,
-			View, ViewIndex,
-			// Shadow data
-			VisibleLightInfos,
-			VirtualShadowMapArray,
-			// Object data
-			HeterogeneousVolumeInterface,
-			MaterialRenderProxy,
-			// Sparse voxel data
-			NumVoxelsBuffer,
-			SparseVoxelUniformBuffer,
-			// Render tile data
-			NumRayMarchingTilesBuffer,
-			RayMarchingTilesBuffer,
-			VoxelsPerTileBuffer,
-			// Output
-			LightingCacheTexture,
-			HeterogeneousVolumeRadiance
-		);
-	}
-	else
-	{
-		RenderWithTransmittanceVolumePipelineWithPreshadingCompute(
-			GraphBuilder,
-			// Scene data
-			SceneTextures,
-			Scene,
-			View, ViewIndex,
-			// Shadow data
-			VisibleLightInfos,
-			VirtualShadowMapArray,
-			// Object data
-			HeterogeneousVolumeInterface,
-			MaterialRenderProxy,
-			// Sparse voxel data
-			NumVoxelsBuffer,
-			SparseVoxelUniformBuffer,
-			// Render tile data
-			NumRayMarchingTilesBuffer,
-			RayMarchingTilesBuffer,
-			VoxelsPerTileBuffer,
-			// Output
-			LightingCacheTexture,
-			HeterogeneousVolumeRadiance
-		);
-	}
+	RenderWithInscatteringVolumePipelineWithPreshadingCompute(
+		GraphBuilder,
+		// Scene data
+		SceneTextures,
+		Scene,
+		View, ViewIndex,
+		// Shadow data
+		VisibleLightInfos,
+		VirtualShadowMapArray,
+		// Object data
+		HeterogeneousVolumeInterface,
+		MaterialRenderProxy,
+		// Preshading data
+		PreshadingUniformBuffer,
+		// Output
+		LightingCacheTexture,
+		HeterogeneousVolumeRadiance
+	);
 }
 
 void RenderWithInscatteringVolumePipelineWithPreshadingHardwareRayTracing(
@@ -1600,6 +1243,13 @@ void GenerateMips3D(
 		GroupCount);
 }
 
+static TAutoConsoleVariable<int32> CVarHeterogeneousVolumesPreshadingCLOD(
+	TEXT("r.HeterogeneousVolumes.Preshading.CLOD"),
+	1,
+	TEXT("Enables continuous level-of-detail for preshading resolution (Default = 1)"),
+	ECVF_RenderThreadSafe
+);
+
 void RenderWithPreshading(
 	FRDGBuilder& GraphBuilder,
 	// Scene data
@@ -1612,7 +1262,7 @@ void RenderWithPreshading(
 	// Object data
 	const IHeterogeneousVolumeInterface* HeterogeneousVolumeInterface,
 	const FMaterialRenderProxy* MaterialRenderProxy,
-	const FPersistentPrimitiveIndex &PersistentPrimitiveIndex,
+	const FPersistentPrimitiveIndex& PersistentPrimitiveIndex,
 	const FBoxSphereBounds LocalBoxSphereBounds,
 	// Transmittance acceleration
 	FRDGTextureRef LightingCacheTexture,
@@ -1622,16 +1272,31 @@ void RenderWithPreshading(
 {
 	CSV_SCOPED_TIMING_STAT_EXCLUSIVE(Effects);
 
-	// Determine baking voxel resolution
-	float LODFactor = HeterogeneousVolumes::CalcLODFactor(View, HeterogeneousVolumeInterface);
 	FIntVector VolumeResolution = HeterogeneousVolumes::GetVolumeResolution(HeterogeneousVolumeInterface);
-	// TODO: Modify volume resolution by LODFactor??
+	FIntVector BakingResolution = VolumeResolution;
+	if (CVarHeterogeneousVolumesPreshadingCLOD.GetValueOnAnyThread() == 1)
+	{
+		// Determine baking voxel resolution
+		float LODFactor = HeterogeneousVolumes::CalcLODFactor(View, HeterogeneousVolumeInterface);
+		BakingResolution = FIntVector(BakingResolution / LODFactor);
+		BakingResolution.X = FMath::Clamp(BakingResolution.X, 1, 1024);
+		BakingResolution.Y = FMath::Clamp(BakingResolution.Y, 1, 1024);
+		BakingResolution.Z = FMath::Clamp(BakingResolution.Z, 1, 1024);
+
+	}
+	else if (CVarHeterogeneousVolumesPreshadingCLOD.GetValueOnAnyThread() == 2)
+	{
+		HeterogeneousVolumes::FLODValue LODValue = HeterogeneousVolumes::CalcLOD(View, HeterogeneousVolumeInterface);
+		BakingResolution = HeterogeneousVolumes::GetLightingCacheResolution(HeterogeneousVolumeInterface, LODValue);
+	}
 
 	// Create baked material grids
-	uint32 NumMips = FMath::Log2(float(FMath::Min(FMath::Min(VolumeResolution.X, VolumeResolution.Y), VolumeResolution.Z))) + 1;
+	const uint32 NumMips = 1;
+	const EPixelFormat PixelFormat = HeterogeneousVolumes::GetPreshadingPixelFormat();
+
 	FRDGTextureDesc BakedMaterialDesc = FRDGTextureDesc::Create3D(
-		VolumeResolution,
-		PF_FloatR11G11B10,
+		BakingResolution,
+		PixelFormat,
 		FClearValueBinding::Black,
 		TexCreate_ShaderResource | TexCreate_UAV | TexCreate_3DTiling,
 		NumMips
@@ -1641,82 +1306,32 @@ void RenderWithPreshading(
 	FRDGTextureRef EmissionTexture = GraphBuilder.CreateTexture(BakedMaterialDesc, TEXT("HeterogeneousVolumes.EmissionTexture"));
 	FRDGTextureRef AlbedoTexture = GraphBuilder.CreateTexture(BakedMaterialDesc, TEXT("HeterogeneousVolumes.AlbedoTexture"));
 
-	// Preshading pipeline
-	{
-		RDG_EVENT_SCOPE(GraphBuilder, "Preshading Pipeline");
-		SCOPE_CYCLE_COUNTER(STATGROUP_HeterogeneousVolumesMaterialBaking);
-
-		{
-			RDG_EVENT_SCOPE(GraphBuilder, "Material Baking");
-
-			ComputeHeterogeneousVolumeBakeMaterial(
-				GraphBuilder,
-				// Scene data
-				Scene,
-				View,
-				// Object data
-				HeterogeneousVolumeInterface,
-				MaterialRenderProxy,
-				PersistentPrimitiveIndex,
-				LocalBoxSphereBounds,
-				// Volume data
-				VolumeResolution,
-				// Output
-				ExtinctionTexture,
-				EmissionTexture,
-				AlbedoTexture
-			);
-		}
-
-		// MIP Generation
-		{
-			RDG_EVENT_SCOPE(GraphBuilder, "MIP Generation");
-			for (uint32 MipLevel = 1; MipLevel < NumMips; ++MipLevel)
-			{
-				GenerateMips3D<FGenerateMips3D>(GraphBuilder, View, ExtinctionTexture, MipLevel);
-				// TODO: Reinstate once ray-marching determines appropriate MIP level to sample
-				//GenerateMips3D<FGenerateMips3D>(GraphBuilder, View, EmissionTexture, MipLevel);
-				//GenerateMips3D<FGenerateMips3D>(GraphBuilder, View, AlbedoTexture, MipLevel);
-			}
-		}
-	}
-
-	// Sparse Voxel Pipeline
-	int32 MipBias = HeterogeneousVolumes::GetSparseVoxelMipBias();
-	uint32 SparseMipLevel = FMath::Clamp(int32(NumMips) - MipBias, 0, int32(NumMips) - 1);
-
-	FRDGTextureRef MinTexture;
-	FRDGBufferRef NumVoxelsBuffer;
-	FRDGBufferRef VoxelBuffer;
+	// Material baking
 	{
-		RDG_EVENT_SCOPE(GraphBuilder, "Sparse Voxel Pipeline");
+		RDG_EVENT_SCOPE(GraphBuilder, "Material Baking");
 		SCOPE_CYCLE_COUNTER(STATGROUP_HeterogeneousVolumesMaterialBaking);
 
-		{
-			RDG_EVENT_SCOPE(GraphBuilder, "Min Generation");
-
-			FRDGTextureDesc MinTextureDesc = ExtinctionTexture->Desc;
-			MinTextureDesc.Extent.X = FMath::Max(MinTextureDesc.Extent.X >> SparseMipLevel, 1);
-			MinTextureDesc.Extent.Y = FMath::Max(MinTextureDesc.Extent.Y >> SparseMipLevel, 1);
-			MinTextureDesc.Depth = FMath::Max(MinTextureDesc.Depth >> SparseMipLevel, 1);
-			MinTextureDesc.NumMips = FMath::Log2(float(FMath::Min(FMath::Min(MinTextureDesc.Extent.X, MinTextureDesc.Extent.Y), (int32)MinTextureDesc.Depth))) + 1;
-
-			MinTexture = GraphBuilder.CreateTexture(MinTextureDesc, TEXT("HeterogeneousVolumes.MinTexture"));
-			CopyTexture3D(GraphBuilder, View, ExtinctionTexture, SparseMipLevel, MinTexture);
-			for (uint32 MipLevel = 1; MipLevel < MinTextureDesc.NumMips; ++MipLevel)
-			{
-				GenerateMips3D<FGenerateMin3D>(GraphBuilder, View, MinTexture, MipLevel);
-			}
-		}
-
-		{
-			RDG_EVENT_SCOPE(GraphBuilder, "Sparse Voxel Generation");
-			GenerateSparseVoxels(GraphBuilder, View, MinTexture, VolumeResolution, SparseMipLevel, NumVoxelsBuffer, VoxelBuffer);
-		}
+		ComputeHeterogeneousVolumeBakeMaterial(
+			GraphBuilder,
+			// Scene data
+			Scene,
+			View,
+			// Object data
+			HeterogeneousVolumeInterface,
+			MaterialRenderProxy,
+			PersistentPrimitiveIndex,
+			LocalBoxSphereBounds,
+			// Volume data
+			BakingResolution,
+			// Output
+			ExtinctionTexture,
+			EmissionTexture,
+			AlbedoTexture
+		);
 	}
 
-	// Create Sparse Voxel UniformBuffer
-	FSparseVoxelUniformBufferParameters* SparseVoxelUniformBufferParameters = GraphBuilder.AllocParameters<FSparseVoxelUniformBufferParameters>();
+	// Create Preshading UniformBuffer
+	FPreshadingUniformBufferParameters* PreshadingUniformBufferParameters = GraphBuilder.AllocParameters<FPreshadingUniformBufferParameters>();
 	{
 		// Object data
 		// TODO: Convert to relative-local space
@@ -1724,85 +1339,40 @@ void RenderWithPreshading(
 		//FMatrix44f RelativeLocalToWorld = FDFMatrix::MakeToRelativeWorldMatrix(ViewOriginHigh, HeterogeneousVolumeInterface->GetLocalToWorld()).M;
 		FMatrix InstanceToLocal = HeterogeneousVolumeInterface->GetInstanceToLocal();
 		FMatrix LocalToWorld = HeterogeneousVolumeInterface->GetLocalToWorld();
-		SparseVoxelUniformBufferParameters->LocalToWorld = FMatrix44f(InstanceToLocal * LocalToWorld);
-		SparseVoxelUniformBufferParameters->WorldToLocal = SparseVoxelUniformBufferParameters->LocalToWorld.Inverse();
+		PreshadingUniformBufferParameters->LocalToWorld = FMatrix44f(InstanceToLocal * LocalToWorld);
+		PreshadingUniformBufferParameters->WorldToLocal = PreshadingUniformBufferParameters->LocalToWorld.Inverse();
 
 		FMatrix LocalToInstance = InstanceToLocal.Inverse();
 		FBoxSphereBounds InstanceBoxSphereBounds = LocalBoxSphereBounds.TransformBy(LocalToInstance);
-		SparseVoxelUniformBufferParameters->LocalBoundsOrigin = FVector3f(InstanceBoxSphereBounds.Origin);
-		SparseVoxelUniformBufferParameters->LocalBoundsExtent = FVector3f(InstanceBoxSphereBounds.BoxExtent);
+		PreshadingUniformBufferParameters->LocalBoundsOrigin = FVector3f(InstanceBoxSphereBounds.Origin);
+		PreshadingUniformBufferParameters->LocalBoundsExtent = FVector3f(InstanceBoxSphereBounds.BoxExtent);
 
 		// Volume data
-		SparseVoxelUniformBufferParameters->VolumeResolution = VolumeResolution;
-		SparseVoxelUniformBufferParameters->ExtinctionTexture = ExtinctionTexture;
-		SparseVoxelUniformBufferParameters->EmissionTexture = EmissionTexture;
-		SparseVoxelUniformBufferParameters->AlbedoTexture = AlbedoTexture;
-		SparseVoxelUniformBufferParameters->TextureSampler = TStaticSamplerState<SF_Trilinear, AM_Clamp, AM_Clamp, AM_Clamp>::GetRHI();
-
-		// Sparse voxel data
-		SparseVoxelUniformBufferParameters->NumVoxelsBuffer = GraphBuilder.CreateSRV(NumVoxelsBuffer, PF_R32_UINT);
-		SparseVoxelUniformBufferParameters->VoxelBuffer = GraphBuilder.CreateSRV(VoxelBuffer);
-		SparseVoxelUniformBufferParameters->MipLevel = SparseMipLevel;
-
-		// Traversal hints
-		SparseVoxelUniformBufferParameters->MaxTraceDistance = HeterogeneousVolumes::GetMaxTraceDistance();
-		SparseVoxelUniformBufferParameters->MaxShadowTraceDistance = HeterogeneousVolumes::GetMaxShadowTraceDistance();
-		SparseVoxelUniformBufferParameters->StepSize = HeterogeneousVolumes::GetStepSize();
-		SparseVoxelUniformBufferParameters->StepFactor = HeterogeneousVolumeInterface->GetStepFactor() * LODFactor;
-		SparseVoxelUniformBufferParameters->ShadowStepSize = HeterogeneousVolumes::GetShadowStepSize();
-		SparseVoxelUniformBufferParameters->ShadowStepFactor = HeterogeneousVolumeInterface->GetShadowStepFactor() * LODFactor;
-		SparseVoxelUniformBufferParameters->bApplyHeightFog = HeterogeneousVolumes::ShouldApplyHeightFog();
-		SparseVoxelUniformBufferParameters->bApplyVolumetricFog = HeterogeneousVolumes::ShouldApplyVolumetricFog();
-		SparseVoxelUniformBufferParameters->IndirectInscatteringFactor = HeterogeneousVolumes::GetIndirectLightingFactor();
+		PreshadingUniformBufferParameters->VolumeResolution = VolumeResolution;
+		PreshadingUniformBufferParameters->ExtinctionTexture = ExtinctionTexture;
+		PreshadingUniformBufferParameters->EmissionTexture = EmissionTexture;
+		PreshadingUniformBufferParameters->AlbedoTexture = AlbedoTexture;
+		PreshadingUniformBufferParameters->TextureSampler = TStaticSamplerState<SF_Trilinear, AM_Clamp, AM_Clamp, AM_Clamp>::GetRHI();
 	}
-	TRDGUniformBufferRef<FSparseVoxelUniformBufferParameters> SparseVoxelUniformBuffer = GraphBuilder.CreateUniformBuffer(SparseVoxelUniformBufferParameters);
+	TRDGUniformBufferRef<FPreshadingUniformBufferParameters> PreshadingUniformBuffer = GraphBuilder.CreateUniformBuffer(PreshadingUniformBufferParameters);
 
-	// Hardware ray tracing
-	if (HeterogeneousVolumes::UseHardwareRayTracing())
-	{
-		RenderWithPreshadingHardwareRayTracing(
-			GraphBuilder,
-			// Scene data
-			SceneTextures,
-			Scene,
-			View, ViewIndex,
-			// Shadow data
-			VisibleLightInfos,
-			VirtualShadowMapArray,
-			// Object data
-			HeterogeneousVolumeInterface,
-			MaterialRenderProxy,
-			// Sparse voxel data
-			NumVoxelsBuffer,
-			SparseVoxelUniformBuffer,
-			// Transmittance acceleration
-			LightingCacheTexture,
-			// Output
-			HeterogeneousVolumeRadiance
-		);
-	}
-	// Software ray tracing
-	else
-	{
-		RenderWithPreshadingCompute(
-			GraphBuilder,
-			// Scene data
-			SceneTextures,
-			Scene,
-			View, ViewIndex,
-			// Shadow data
-			VisibleLightInfos,
-			VirtualShadowMapArray,
-			// Object data
-			HeterogeneousVolumeInterface,
-			MaterialRenderProxy,
-			// Sparse voxel data
-			NumVoxelsBuffer,
-			SparseVoxelUniformBuffer,
-			// Transmittance acceleration
-			LightingCacheTexture,
-			// Output
-			HeterogeneousVolumeRadiance
-		);
-	}
-}
\ No newline at end of file
+	RenderWithPreshadingCompute(
+		GraphBuilder,
+		// Scene data
+		SceneTextures,
+		Scene,
+		View, ViewIndex,
+		// Shadow data
+		VisibleLightInfos,
+		VirtualShadowMapArray,
+		// Object data
+		HeterogeneousVolumeInterface,
+		MaterialRenderProxy,
+		// Preshading data
+		PreshadingUniformBuffer,
+		// Transmittance acceleration
+		LightingCacheTexture,
+		// Output
+		HeterogeneousVolumeRadiance
+	);
+}
