[HWRT] Fix handling of EngineShowFlags.StaticMeshes and EngineShowFlags.SkeletalMeshes.
- plus new cvar to control whether to include dynamic ray tracing instances.

diff --git a/Engine/Source/Runtime/Engine/Private/PrimitiveSceneProxy.cpp b/Engine/Source/Runtime/Engine/Private/PrimitiveSceneProxy.cpp
--- a/Engine/Source/Runtime/Engine/Private/PrimitiveSceneProxy.cpp
+++ b/Engine/Source/Runtime/Engine/Private/PrimitiveSceneProxy.cpp
@@ -1927,7 +1927,7 @@ ERayTracingPrimitiveFlags FPrimitiveSceneProxy::GetCachedRayTracingInstance(FRay
 ERayTracingProxyType FPrimitiveSceneProxy::GetRayTracingProxyType() const
 {
 	ERayTracingProxyType ProxyType = ERayTracingProxyType::None;
-	if (bIsStaticMesh)
+	if (bIsStaticMesh || bIsNaniteMesh)
 	{
 		EnumAddFlags(ProxyType, ERayTracingProxyType::StaticMesh);
 	}
diff --git a/Engine/Source/Runtime/Renderer/Private/RayTracing/RayTracing.cpp b/Engine/Source/Runtime/Renderer/Private/RayTracing/RayTracing.cpp
--- a/Engine/Source/Runtime/Renderer/Private/RayTracing/RayTracing.cpp
+++ b/Engine/Source/Runtime/Renderer/Private/RayTracing/RayTracing.cpp
@@ -21,6 +21,14 @@
 #include "Async/ParallelFor.h"
 #include <type_traits>
 
+bool GRayTracingDynamicInstances = true;
+static FAutoConsoleVariableRef CVarRayTracingDynamicInstances(
+	TEXT("r.RayTracing.DynamicInstances"),
+	GRayTracingDynamicInstances,
+	TEXT("Where to include dynamic instances in the ray tracing scene."),
+	ECVF_RenderThreadSafe
+);
+
 static int32 GRayTracingParallelMeshBatchSetup = 1;
 static FAutoConsoleVariableRef CRayTracingParallelMeshBatchSetup(
 	TEXT("r.RayTracing.ParallelMeshBatchSetup"),
@@ -102,7 +110,7 @@ static FAutoConsoleCommand UpdateCachedRayTracingStateCmd(
 	TEXT("Update cached ray tracing state (mesh commands and instances)."),
 	FConsoleCommandDelegate::CreateStatic([] { bUpdateCachedRayTracingState = true; }));
 
-static ERayTracingProxyType ActiveRayTracingProxyTypes = ERayTracingProxyType::All;
+static ERayTracingProxyType GActiveRayTracingProxyTypes = ERayTracingProxyType::All;
 
 static_assert(sizeof(FScene::FPrimitiveRayTracingData) == 8, "FScene::FPrimitiveRayTracingData is packed to 8 bytes to be cache efficient during GatherRelevantPrimitives");
 
@@ -166,7 +174,7 @@ static void RefreshRayTracingInstancesSinkFunction()
 		ENQUEUE_RENDER_COMMAND(RefreshRayTracingInstancesCmd)(
 			[bRequireUpdateCachedRayTracingState, NewActiveRayTracingProxyTypes](FRHICommandListImmediate&)
 			{
-				ActiveRayTracingProxyTypes = NewActiveRayTracingProxyTypes;
+				GActiveRayTracingProxyTypes = NewActiveRayTracingProxyTypes;
 				bUpdateCachedRayTracingState = bRequireUpdateCachedRayTracingState;
 			}
 		);
@@ -320,12 +328,23 @@ namespace RayTracing
 			, View(InView)
 			, SceneOptions(MoveTemp(InSceneOptions))
 		{
+			ActiveRayTracingProxyTypes = GActiveRayTracingProxyTypes;
 
+			if (!View.Family->EngineShowFlags.StaticMeshes)
+			{
+				EnumRemoveFlags(ActiveRayTracingProxyTypes, ERayTracingProxyType::StaticMesh);
+			}
+
+			if (!View.Family->EngineShowFlags.SkeletalMeshes)
+			{
+				EnumRemoveFlags(ActiveRayTracingProxyTypes, ERayTracingProxyType::SkeletalMesh);
+			}
 		}
 
 		FScene& Scene;
 		FViewInfo& View;
 		FSceneOptions SceneOptions;
+		ERayTracingProxyType ActiveRayTracingProxyTypes;
 
 		// Filtered lists of relevant primitives
 		TArray<int32> StaticPrimitivesIndices;
@@ -1157,6 +1176,7 @@ namespace RayTracing
 		FViewInfo& View = TaskData.View;
 
 		const bool bGameView = View.bIsGameView || View.Family->EngineShowFlags.Game;
+		const ERayTracingProxyType ActiveRayTracingProxyTypes = TaskData.ActiveRayTracingProxyTypes;
 
 		const bool bPerformRayTracing = View.State != nullptr && !View.bIsReflectionCapture && View.IsRayTracingAllowedForView();
 
@@ -1181,7 +1201,7 @@ namespace RayTracing
 				Contexts,
 				Scene.PrimitiveSceneProxies.Num(),
 				MinBatchSize,
-				[&Scene, &View, bGameView, bUsingReferenceBasedResidency](FGatherRelevantPrimitivesContext& Context, int32 PrimitiveIndex)
+				[&Scene, &View, bGameView, ActiveRayTracingProxyTypes, bUsingReferenceBasedResidency](FGatherRelevantPrimitivesContext& Context, int32 PrimitiveIndex)
 			{
 				// Get primitive visibility state from culling
 				if (!View.PrimitiveRayTracingVisibilityMap[PrimitiveIndex])
@@ -1249,12 +1269,12 @@ namespace RayTracing
 				{
 					checkf(!EnumHasAllFlags(Flags, ERayTracingPrimitiveFlags::CacheInstances), TEXT("Only static primitives are expected to use CacheInstances flag."));
 
-					if (View.Family->EngineShowFlags.SkeletalMeshes) // TODO: Fix this check
+					if (GRayTracingDynamicInstances)
 					{
 						Context.DynamicPrimitives.AddElement(PrimitiveIndex);
 					}
 				}
-				else if (View.Family->EngineShowFlags.StaticMeshes)
+				else
 				{
 					Context.StaticPrimitives.AddElement(PrimitiveIndex);
 				}
