Mobile: Fixed issues with a "Quad Overdraw" mode when deferred shading is active

diff --git a/Engine/Source/Runtime/Renderer/Private/MobileShadingRenderer.cpp b/Engine/Source/Runtime/Renderer/Private/MobileShadingRenderer.cpp
--- a/Engine/Source/Runtime/Renderer/Private/MobileShadingRenderer.cpp
+++ b/Engine/Source/Runtime/Renderer/Private/MobileShadingRenderer.cpp
@@ -2749,7 +2749,15 @@ void FMobileSceneRenderer::RenderMobileDebugView(FRHICommandList& RHICmdList, co
 
 		// Here we use the base pass depth result to get z culling for opaque and masque.
 		// The color needs to be cleared at this point since shader complexity renders in additive.
-		DrawClearQuad(RHICmdList, FLinearColor::Black);
+		const static FLinearColor ClearColors[] = {
+			FLinearColor::Black, FLinearColor::Black, FLinearColor::Black, FLinearColor::Black,
+			FLinearColor::Black, FLinearColor::Black, FLinearColor::Black, FLinearColor::Black
+		};
+
+		FGraphicsPipelineStateInitializer GraphicsPSOInit;
+		RHICmdList.ApplyCachedRenderTargets(GraphicsPSOInit);
+		check(GraphicsPSOInit.RenderTargetsEnabled <= UE_ARRAY_COUNT(ClearColors));
+		DrawClearQuadMRT(RHICmdList, true, GraphicsPSOInit.RenderTargetsEnabled, ClearColors, false, 0, false, 0);
 
 		RHICmdList.SetViewport(View.ViewRect.Min.X, View.ViewRect.Min.Y, 0, View.ViewRect.Max.X, View.ViewRect.Max.Y, 1);
 
diff --git a/Engine/Source/Runtime/Renderer/Private/MobileTranslucentRendering.cpp b/Engine/Source/Runtime/Renderer/Private/MobileTranslucentRendering.cpp
--- a/Engine/Source/Runtime/Renderer/Private/MobileTranslucentRendering.cpp
+++ b/Engine/Source/Runtime/Renderer/Private/MobileTranslucentRendering.cpp
@@ -8,7 +8,7 @@ void FMobileSceneRenderer::RenderTranslucency(FRHICommandList& RHICmdList, const
 {
 	const FSceneViewFamily& ViewFamily = *View.Family;
 
-	const bool bShouldRenderTranslucency = ShouldRenderTranslucency(InStandardTranslucencyPass, FamilyViewInfos) && ViewFamily.EngineShowFlags.Translucency;
+	const bool bShouldRenderTranslucency = !ViewFamily.UseDebugViewPS() && ShouldRenderTranslucency(InStandardTranslucencyPass, FamilyViewInfos) && ViewFamily.EngineShowFlags.Translucency;
 	if (bShouldRenderTranslucency)
 	{
 		CSV_SCOPED_TIMING_STAT_EXCLUSIVE(RenderTranslucency);
