Refactor and combine some ShadowDepth permutations to reduce permutation explosion
- Standardize and move the relevant shader logic under "SUPPORTS_XYZ" defines, and dynamic branches on ShadowDepthType
- Renaming some functions to make it clear which light/shadow depth type they are associated with
- Clean up the EShadowDepthShaderMode logic between the VS/PS and make more of the decisions about which permutation to use based on simple flags and platform support for consistency
- Add a "combined" permutation that will supercede any individual permutations (which get removed auotmatically), currently combining VSM with the OnePassPoint light permutations
- Testing across various platforms indicates this particular combination has a negligeable impact to performance and a generally compatible feature matrix

Add r.Shadow.Virtual.ProjectEnabled to allow culling VSM-related permutations if they are not needed.
Add r.Shadow.Classic.ProjectEnabled to allow culling classic shadow-map related permutations (CSMs, spot and one pass point light) if they are not needed.

diff --git a/Engine/Plugins/FX/Niagara/Shaders/Private/NiagaraRibbonVertexFactory.ush b/Engine/Plugins/FX/Niagara/Shaders/Private/NiagaraRibbonVertexFactory.ush
--- a/Engine/Plugins/FX/Niagara/Shaders/Private/NiagaraRibbonVertexFactory.ush
+++ b/Engine/Plugins/FX/Niagara/Shaders/Private/NiagaraRibbonVertexFactory.ush
@@ -109,7 +109,7 @@ struct FVertexFactoryInput
 	uint InterpVtxId					: SV_VertexID;
 
 	// Optional instance ID for vertex layered rendering
-#if FEATURE_LEVEL >= FEATURE_LEVEL_SM4 && ONEPASS_POINTLIGHT_SHADOW && USING_VERTEX_SHADER_LAYER
+#if FEATURE_LEVEL >= FEATURE_LEVEL_SM4 && SUPPORT_SHADOW_DEPTH_POINT && USING_VERTEX_SHADER_LAYER
 	#undef GetInstanceIdFromVF
 	#define GetInstanceIdFromVF(VFInput)				(VFInput.InstanceId)
 	uint InstanceId					: SV_InstanceID;
@@ -999,7 +999,7 @@ FVertexFactoryInput LoadVertexFactoryInputForDynamicUpdate(uint TriangleIndex, i
 		Input.InterpVtxId = NiagaraRibbonVFLooseParameters.IndexBuffer[Input.InterpVtxId];
 	}
 
-#if FEATURE_LEVEL >= FEATURE_LEVEL_SM4 && ONEPASS_POINTLIGHT_SHADOW && USING_VERTEX_SHADER_LAYER
+#if FEATURE_LEVEL >= FEATURE_LEVEL_SM4 && SUPPORT_SHADOW_DEPTH_POINT && USING_VERTEX_SHADER_LAYER
 	Input.InstanceId = 0;
 #endif
 	return Input;
@@ -1019,7 +1019,7 @@ FVertexFactoryInput LoadVertexFactoryInputForHGS(uint TriangleIndex, int VertexI
 	FTriangleBaseAttributes Tri = LoadTriangleBaseAttributes(TriangleIndex);
 	uint VertexId = Tri.Indices[VertexIndex];
 	Input.InterpVtxId = VertexId;
-#if FEATURE_LEVEL >= FEATURE_LEVEL_SM4 && ONEPASS_POINTLIGHT_SHADOW && USING_VERTEX_SHADER_LAYER
+#if FEATURE_LEVEL >= FEATURE_LEVEL_SM4 && SUPPORT_SHADOW_DEPTH_POINT && USING_VERTEX_SHADER_LAYER
 	Input.InstanceId = 0;
 #endif
 	return Input;
@@ -1128,7 +1128,7 @@ FVertexFactoryInput NiagaraVertexFactoryExport_InitVertexFactoryInput(uint Trian
 	Input.InterpVtxId = TriangleIndex * 3 + VertexIndex;
 	Input.InterpVtxId = NiagaraRibbonVFLooseParameters.IndexBuffer[Input.InterpVtxId];
 
-#if FEATURE_LEVEL >= FEATURE_LEVEL_SM4 && ONEPASS_POINTLIGHT_SHADOW && USING_VERTEX_SHADER_LAYER
+#if FEATURE_LEVEL >= FEATURE_LEVEL_SM4 && SUPPORT_SHADOW_DEPTH_POINT && USING_VERTEX_SHADER_LAYER
 	Input.InstanceId = 0;
 #endif
 	return Input;
diff --git a/Engine/Shaders/Private/ShadowDepthCommon.ush b/Engine/Shaders/Private/ShadowDepthCommon.ush
--- a/Engine/Shaders/Private/ShadowDepthCommon.ush
+++ b/Engine/Shaders/Private/ShadowDepthCommon.ush
@@ -4,10 +4,6 @@
 	ShadowDepthCommon.usf: Shared functionality for shadow depth rendering.
 =============================================================================*/
 
-#ifndef PERSPECTIVE_CORRECT_DEPTH
-#define PERSPECTIVE_CORRECT_DEPTH 0
-#endif
-
 #define INTERPOLATE_VF_ATTRIBUTES (!MATERIALBLENDING_SOLID || MATERIALBLENDING_MASKED || (REQUIRES_VF_ATTRIBUTES_FOR_CLIPPING && !POSITION_ONLY))
 #define INTERPOLATE_POSITION (MATERIALBLENDING_MASKED)
 
@@ -18,11 +14,13 @@ struct FShadowDepthVSToPS
 	FVertexFactoryInterpolantsVSToPS FactoryInterpolants;
 #endif
 
-#if PERSPECTIVE_CORRECT_DEPTH
+#if SUPPORT_SHADOW_DEPTH_SPOT
 	float ShadowDepth : TEXCOORD6;
 	float DepthBias : TEXCOORD8;
-#elif !COMPILER_SUPPORTS_EMPTY_STRUCTS && !ONEPASS_POINTLIGHT_SHADOW
-	float Dummy : TEXCOORD6;
+#endif
+
+#if SUPPORT_SHADOW_DEPTH_POINT
+	uint LayerIndex : SV_RenderTargetArrayIndex;
 #endif
 
 #if INTERPOLATE_POSITION
@@ -33,4 +31,9 @@ struct FShadowDepthVSToPS
 #if USE_WORLD_POSITION_EXCLUDING_SHADER_OFFSETS && !IS_NANITE_PASS
 	float3 PixelPositionExcludingWPO : TEXCOORD9;
 #endif
+
+#if SUPPORT_SHADOW_DEPTH_VSM
+	nointerpolation uint PackedPageInfo : TEXCOORD8;
+	float4 VirtualSmPageClip : SV_ClipDistance;
+#endif
 };
diff --git a/Engine/Shaders/Private/ShadowDepthPixelShader.usf b/Engine/Shaders/Private/ShadowDepthPixelShader.usf
--- a/Engine/Shaders/Private/ShadowDepthPixelShader.usf
+++ b/Engine/Shaders/Private/ShadowDepthPixelShader.usf
@@ -12,17 +12,6 @@
 #define SHADOW_DEPTH_SHADER 1
 #define USE_STENCIL_LOD_DITHER	0
 
-#ifndef ENABLE_NON_NANITE_VSM
-#error "ENABLE_NON_NANITE_VSM should be defined to either 0 or 1!"
-#endif
-
-
-#if ENABLE_NON_NANITE_VSM
-	#define VIRTUAL_SM_ENABLED (!(ONEPASS_POINTLIGHT_SHADOW || PERSPECTIVE_CORRECT_DEPTH))
-#else
-	#define VIRTUAL_SM_ENABLED 0
-#endif
-
 #include "Nanite/NanitePackedNaniteView.ush"
 #include "Common.ush"
 
@@ -39,23 +28,17 @@
 #include "/Engine/Generated/VertexFactory.ush"
 #include "ShadowDepthCommon.ush"
 
-
-#if VIRTUAL_SM_ENABLED
+#if SUPPORT_SHADOW_DEPTH_VSM
 #include "Nanite/NaniteDataDecode.ush"
 #include "SceneData.ush"
 #include "VirtualShadowMaps/VirtualShadowMapPageAccessCommon.ush"
 #endif
 
-#define SECONDARY_OCCLUSION			1
-
 void Main( 
 	FShadowDepthVSToPS Inputs,
-#if VIRTUAL_SM_ENABLED
-	nointerpolation uint PackedPageInfo : TEXCOORD8,
-#endif
 	in float4 SvPosition : SV_Position		// after all interpolators
-#if PERSPECTIVE_CORRECT_DEPTH
-	,out float OutDepth : SV_DEPTH
+#if SUPPORT_SHADOW_DEPTH_SPOT
+	, out float OutDepth : SV_DEPTH
 #endif
 	)
 {
@@ -94,7 +77,7 @@ void Main(
 	ClipLODTransition(SvPosition.xy);
 #endif
 
-#if PERSPECTIVE_CORRECT_DEPTH
+#if SUPPORT_SHADOW_DEPTH_SPOT
 	const float InvMaxSubjectDepth = PassStruct.ShadowParams.w;
 	Inputs.ShadowDepth = 1 - Inputs.ShadowDepth * InvMaxSubjectDepth;
 	Inputs.ShadowDepth += Inputs.DepthBias;
@@ -102,11 +85,11 @@ void Main(
 	OutDepth = saturate(Inputs.ShadowDepth);
 #endif
 
-#if ENABLE_NON_NANITE_VSM && VIRTUAL_TEXTURE_TARGET
+#if SUPPORT_SHADOW_DEPTH_VSM
 	uint2 vAddress = (uint2)SvPosition.xy;
 	float DeviceZ = SvPosition.z;
 
-	FPageInfo PageInfo = UnpackPageInfo( PackedPageInfo );
+	FPageInfo PageInfo = UnpackPageInfo( Inputs.PackedPageInfo );
 	FNaniteView NaniteView = UnpackNaniteView( PassStruct.PackedNaniteViews[ PageInfo.ViewId ] );
 
 	// Note: not the usual UB
diff --git a/Engine/Shaders/Private/ShadowDepthVertexShader.usf b/Engine/Shaders/Private/ShadowDepthVertexShader.usf
--- a/Engine/Shaders/Private/ShadowDepthVertexShader.usf
+++ b/Engine/Shaders/Private/ShadowDepthVertexShader.usf
@@ -8,17 +8,6 @@
 #define SHADOW_DEPTH_SHADER 1
 #define USE_STENCIL_LOD_DITHER	0
 
-#ifndef ENABLE_NON_NANITE_VSM
-#error "ENABLE_NON_NANITE_VSM should be defined to either 0 or 1!"
-#endif
-
-
-#if ENABLE_NON_NANITE_VSM
-	#define VIRTUAL_SM_ENABLED (!(ONEPASS_POINTLIGHT_SHADOW || PERSPECTIVE_CORRECT_DEPTH))
-#else
-	#define VIRTUAL_SM_ENABLED 0
-#endif
-
 // Need to get the definition of nanite packed view for generated unform buffer (which is pulled in by Common.ush)
 #include "Nanite/NanitePackedNaniteView.ush"
 #include "Common.ush"
@@ -41,87 +30,22 @@
 #define VF_USE_VSM_CONSTANT_BIAS								0
 #endif
 
-#if VIRTUAL_SM_ENABLED
+#if SUPPORT_SHADOW_DEPTH_VSM
 #include "Nanite/NaniteDataDecode.ush"
 #include "SceneData.ush"
 #include "VirtualShadowMaps/VirtualShadowMapPageAccessCommon.ush"
-#endif // VIRTUAL_SM_ENABLED
-
-
-void SetShadowDepthOutputs(
-	float4x4 WorldToClipMatrix, 
-	float4x4 WorldToShadowMatrix, 
-	float4 WorldPosition, 
-	float3 WorldVertexNormal, 
-	out float4 OutPosition, 
-	out float ShadowDepth
-#if PERSPECTIVE_CORRECT_DEPTH
-	, out float OutDepthBias
-#endif
-)
-{
-	OutPosition = mul(WorldPosition, WorldToClipMatrix);
-
-	// Clamp the vertex to the near plane if it is in front of the near plane
-	// This has problems if some vertices of a triangle get clamped and others do not, also causes artifacts with non-ortho projections
-	if (PassStruct.bClampToNearPlane > 0 && OutPosition.z > OutPosition.w)
-	{
-		OutPosition.z = 0.999999f;
-		OutPosition.w = 1.0f;
-	}
-
-	#if ONEPASS_POINTLIGHT_SHADOW
-	const float3 ViewDirection = -normalize(mul(WorldPosition, WorldToShadowMatrix).xyz);
-	const float3 ViewNormal = mul(float4(WorldVertexNormal,0), WorldToShadowMatrix).xyz;
-	const float NoL = abs(dot(ViewDirection, ViewNormal));
-	#else
-	const float NoL = abs(dot(
-		float3(WorldToShadowMatrix[0].z, WorldToShadowMatrix[1].z, WorldToShadowMatrix[2].z),
-		WorldVertexNormal));
-	#endif
-
-	const float MaxSlopeDepthBias = PassStruct.ShadowParams.z;
-	const float Slope = clamp(abs(NoL) > 0 ? sqrt(saturate(1 - NoL*NoL)) / NoL : MaxSlopeDepthBias, 0, MaxSlopeDepthBias);
-	
-	const float SlopeDepthBias = PassStruct.ShadowParams.y;
-	const float SlopeBias = SlopeDepthBias * Slope;
-
-	const float ConstantDepthBias = PassStruct.ShadowParams.x;
-	const float DepthBias = SlopeBias + ConstantDepthBias;
-
-	#if PERSPECTIVE_CORRECT_DEPTH
-		ShadowDepth = OutPosition.z;
-		OutDepthBias = DepthBias;
-	#elif ONEPASS_POINTLIGHT_SHADOW
-		ShadowDepth = 0;
-		//OutPosition.z += DepthBias;
-	#else
-		// Output linear, normalized depth
-		const float InvMaxSubjectDepth = PassStruct.ShadowParams.w;
-		#if PLATFORM_NEEDS_PRECISE_SHADOW_DEPTH
-			precise 
-		#endif
-		float AdjustedDepth = ( 1 - OutPosition.z * InvMaxSubjectDepth ) + DepthBias;
-		ShadowDepth = AdjustedDepth;
-		OutPosition.z = AdjustedDepth;
-	#endif
-}
 
 /** 
  * Load physical page offset and perform the needed scale and bias to render the instance into the given physical page,
  * with the correct clipping set up.
  */
-#if VIRTUAL_SM_ENABLED
-
-void ScaleBiasClipToPhysicalSmPage(FNaniteView NaniteView, inout float4 PositionClip, inout float4 ClipPlanesInOut, FPageInfo PageInfo)
+float4 ScaleBiasClipToPhysicalSmPage(FNaniteView NaniteView, float4 InPositionClip, inout float4 ClipPlanesInOut, FPageInfo PageInfo)
 {
-	float2 vClip = PositionClip.xy;
-	float2 vUV = vClip * float2(0.5, -0.5) + 0.5 * PositionClip.w;
+	float2 vUV = InPositionClip.xy * float2(0.5, -0.5) + 0.5 * InPositionClip.w;
 	float2 vPixels = vUV * ( uint(VSM_VIRTUAL_MAX_RESOLUTION_XY) >> NaniteView.TargetMipLevel );
 
+	float4 PositionClip = InPositionClip;
 	{
-		//PointClip.xy = NaniteView.ClipSpaceScaleOffset.xy * PointClip.xy + NaniteView.ClipSpaceScaleOffset.zw * PointClip.w;
-
 		float2 Scale = exp2( -NaniteView.TargetMipLevel );
 		PositionClip.xy *= Scale;
 		PositionClip.xy += ( Scale * float2(1,-1) + float2(-1,1) ) * PositionClip.w;
@@ -137,15 +61,17 @@ void ScaleBiasClipToPhysicalSmPage(FNaniteView NaniteView, inout float4 Position
 		ClipPlanesInOut.xy = MinClip;
 		ClipPlanesInOut.zw = MaxClip;
 	}
+
+	return PositionClip;
 }
 
-void TransformToVirtualSmPage(inout float4 PointClip, inout float4 ClipPlanesInOut, FPageInfo PageInfo, float3 PointTranslatedWorld)
+float4 TransformShadowVSM(inout float4 ClipPlanesInOut, FPageInfo PageInfo, float3 PointTranslatedWorld)
 {
 	// Load nanite view and perform transform to SM space.
 	FNaniteView NaniteView = UnpackNaniteView(PassStruct.PackedNaniteViews[PageInfo.ViewId]);
 
 	PointTranslatedWorld += DFFastSubtractDemote( NaniteView.PreViewTranslation, ResolvedView.PreViewTranslation );
-	PointClip = mul( float4( PointTranslatedWorld, 1 ), NaniteView.TranslatedWorldToClip );
+	float4 PointClip = mul( float4( PointTranslatedWorld, 1 ), NaniteView.TranslatedWorldToClip );
 
 	// Clamp the vertex to the near plane if it is in front of the near plane
 	// This has problems if some vertices of a triangle get clamped and others do not, also causes artifacts with non-ortho projections
@@ -155,7 +81,7 @@ void TransformToVirtualSmPage(inout float4 PointClip, inout float4 ClipPlanesInO
 		PointClip.w = 1.0f;
 	}
 
-	ScaleBiasClipToPhysicalSmPage(NaniteView, PointClip, ClipPlanesInOut, PageInfo);
+	PointClip = ScaleBiasClipToPhysicalSmPage(NaniteView, PointClip, ClipPlanesInOut, PageInfo);
 
 #if VF_USE_VSM_CONSTANT_BIAS
 	{
@@ -163,30 +89,89 @@ void TransformToVirtualSmPage(inout float4 PointClip, inout float4 ClipPlanesInO
 		PointClip.z = clamp(PointClip.z + ConstantDepthBias * NaniteView.ViewToClip[2][2], -PointClip.w, PointClip.w);
 	}
 #endif
+
+	return PointClip;
+}
+#endif // SUPPORT_SHADOW_DEPTH_VSM
+
+float ComputeDepthBiasDirectionalSpot(float4x4 WorldToShadowMatrix, float3 WorldVertexNormal)
+{
+	const float NoL = abs(dot(
+		float3(WorldToShadowMatrix[0].z, WorldToShadowMatrix[1].z, WorldToShadowMatrix[2].z),
+		WorldVertexNormal));
+
+	const float MaxSlopeDepthBias = PassStruct.ShadowParams.z;
+	const float Slope = clamp(abs(NoL) > 0 ? sqrt(saturate(1 - NoL*NoL)) / NoL : MaxSlopeDepthBias, 0, MaxSlopeDepthBias);
+	
+	const float SlopeDepthBias = PassStruct.ShadowParams.y;
+	const float SlopeBias = SlopeDepthBias * Slope;
+
+	const float ConstantDepthBias = PassStruct.ShadowParams.x;
+	const float DepthBias = SlopeBias + ConstantDepthBias;
+
+	return DepthBias;
 }
 
-#endif // VIRTUAL_SM_ENABLED
+// NOTE: Only for non-VSM directional lights
+float4 TransformShadowDirectional(
+	float4x4 WorldToClipMatrix, 
+	float4x4 WorldToShadowMatrix, 
+	float4 WorldPosition, 
+	float3 WorldVertexNormal)
+{
+	float4 OutPosition = mul(WorldPosition, WorldToClipMatrix);
+
+	// Clamp the vertex to the near plane if it is in front of the near plane
+	// This has problems if some vertices of a triangle get clamped and others do not, also causes artifacts with non-ortho projections
+	if (PassStruct.bClampToNearPlane > 0 && OutPosition.z > OutPosition.w)
+	{
+		OutPosition.z = 0.999999f;
+		OutPosition.w = 1.0f;
+	}
+
+	const float DepthBias = ComputeDepthBiasDirectionalSpot(WorldToShadowMatrix, WorldVertexNormal);
+
+	// Output linear, normalized depth
+	const float InvMaxSubjectDepth = PassStruct.ShadowParams.w;
+#if PLATFORM_NEEDS_PRECISE_SHADOW_DEPTH
+	precise
+#endif
+	float AdjustedDepth = ( 1 - OutPosition.z * InvMaxSubjectDepth ) + DepthBias;
+	OutPosition.z = AdjustedDepth;
+
+	return OutPosition;
+}
+
+// Only for non-VSM spot lights
+float4 TransformShadowSpot(
+	float4x4 WorldToClipMatrix, 
+	float4x4 WorldToShadowMatrix, 
+	float4 WorldPosition, 
+	float3 WorldVertexNormal,
+	out float OutShadowDepth,
+	out float OutDepthBias)
+{
+	float4 OutPosition = mul(WorldPosition, WorldToClipMatrix);
+	OutShadowDepth = OutPosition.z;
+	OutDepthBias = ComputeDepthBiasDirectionalSpot(WorldToShadowMatrix, WorldVertexNormal);
+	return OutPosition;
+}
+
+// See EShadowDepthType
+#define SHADOW_DEPTH_TYPE_VSM (1U)
+#define SHADOW_DEPTH_TYPE_DIRECTIONAL (1U<<1)
+#define SHADOW_DEPTH_TYPE_SPOT (1U<<2)
+#define SHADOW_DEPTH_TYPE_ONEPASSPOINT (1U<<3)
 
 uint bUseGpuSceneInstancing;
-#if VERTEXSHADER && ONEPASS_POINTLIGHT_SHADOW
+#if VERTEXSHADER && SUPPORT_SHADOW_DEPTH_POINT
 uint LayerId;
 #endif
 
 void Main(
 	FVertexFactoryInput Input,
 	out FShadowDepthVSToPS OutParameters,
-#if VIRTUAL_SM_ENABLED
-	out nointerpolation uint PackedPageInfo : TEXCOORD8,
-#endif
-	out float4 OutPosition : SV_POSITION
-#if ONEPASS_POINTLIGHT_SHADOW
-	, out uint LayerIndex : SV_RenderTargetArrayIndex
-#endif
-#if VIRTUAL_SM_ENABLED	
-	// OLA-TODO: this collides with instanced stereo, which thankfully is not used with shadow maps, so should be fine, presumably.
-	, out float4 OutVirtualSmPageClip : SV_ClipDistance
-#endif // VIRTUAL_SM_ENABLED
-	)
+	out float4 OutPosition : SV_POSITION)
 {
 	ResolvedView = ResolveView();
 	
@@ -203,147 +188,155 @@ void Main(
 	WorldPos.xyz += GetMaterialWorldPositionOffset(VertexParameters);
 	ApplyMaterialFirstPersonTransform(VertexParameters, WorldPos.xyz);
 
-#if ONEPASS_POINTLIGHT_SHADOW
-
-	OutPosition = WorldPos;
-
-	#if INTERPOLATE_VF_ATTRIBUTES
-		// Masked materials need texture coords to clip
-		OutParameters.FactoryInterpolants = VertexFactoryGetInterpolantsVSToPS(Input, VFIntermediates, VertexParameters);
-	#endif
-
-	#if INTERPOLATE_POSITION
-		OutParameters.PixelPosition = WorldPos.xyz;
-	#endif
-
-	LayerIndex = bUseGpuSceneInstancing ? VertexFactoryGetViewIndex(VFIntermediates) : LayerId;
-	OutPosition = mul(WorldPos, PassStruct.ShadowViewProjectionMatrices[LayerIndex]);
-
-#else
-	float Dummy;
-
-	SetShadowDepthOutputs(
-		PassStruct.ProjectionMatrix,
-		PassStruct.ViewMatrix,
-		WorldPos, 
-		WorldNormal,
-		OutPosition, 
-#if !PERSPECTIVE_CORRECT_DEPTH
-		Dummy
-#else
-		OutParameters.ShadowDepth,
-		OutParameters.DepthBias
+#if INTERPOLATE_VF_ATTRIBUTES
+	// Masked materials need texture coords to clip
+	OutParameters.FactoryInterpolants = VertexFactoryGetInterpolantsVSToPS(Input, VFIntermediates, VertexParameters);
 #endif
-		);
-	
-	#if INTERPOLATE_VF_ATTRIBUTES
-		// Masked materials need texture coords to clip
-		OutParameters.FactoryInterpolants = VertexFactoryGetInterpolantsVSToPS(Input, VFIntermediates, VertexParameters);
-	#endif
-
-	#if INTERPOLATE_POSITION
-		OutParameters.PixelPosition = WorldPos.xyz;
-	#endif
-
-	#if !PERSPECTIVE_CORRECT_DEPTH && !COMPILER_SUPPORTS_EMPTY_STRUCTS
-		OutParameters.Dummy = 0;
-	#endif
 
+#if INTERPOLATE_POSITION
+	OutParameters.PixelPosition = WorldPos.xyz;
 #endif
 
-#if VIRTUAL_SM_ENABLED
-	PackedPageInfo = 0;
-	
-	OutVirtualSmPageClip = float4(1.0f, 1.0f, 1.0f, 1.0f);
-	if (PassStruct.bRenderToVirtualShadowMap != 0)
+	OutPosition = float4(0, 0, 0, 0);
+
+#if SUPPORT_SHADOW_DEPTH_VSM
+	if (PassStruct.ShadowDepthType == SHADOW_DEPTH_TYPE_VSM)
 	{
+		OutParameters.PackedPageInfo = 0;	
+		OutParameters.VirtualSmPageClip = float4(1.0f, 1.0f, 1.0f, 1.0f);
+		
 		// Get the offset from which we loaded the instance ID
 		uint InstanceIdIndex = VertexFactoryGetInstanceIdLoadIndex(VFIntermediates);
-		PackedPageInfo = InstanceCulling.PageInfoBuffer[InstanceIdIndex];
+		OutParameters.PackedPageInfo = InstanceCulling.PageInfoBuffer[InstanceIdIndex];
+		FPageInfo PageInfo = UnpackPageInfo(OutParameters.PackedPageInfo);
+		OutPosition = TransformShadowVSM(OutParameters.VirtualSmPageClip, PageInfo, WorldPos.xyz);
+	}
+	else
+	{
+		// TODO: Cleaner way to initialize unused outputs...
+		OutParameters.PackedPageInfo = 0;
+		OutParameters.VirtualSmPageClip = float4(1, 1, 1, 1);
+	}
+#endif // SUPPORT_SHADOW_DEPTH_VSM
 
-		FPageInfo PageInfo = UnpackPageInfo(PackedPageInfo);
+#if SUPPORT_SHADOW_DEPTH_DIRECTIONAL
+	// Non-VSM directional lights
+	if (PassStruct.ShadowDepthType == SHADOW_DEPTH_TYPE_DIRECTIONAL)
+	{
+		OutPosition = TransformShadowDirectional(
+			PassStruct.ProjectionMatrix,
+			PassStruct.ViewMatrix,
+			WorldPos, 
+			WorldNormal);
+	}
+#endif // SUPPORT_SHADOW_DEPTH_DIRECTIONAL
+
+#if SUPPORT_SHADOW_DEPTH_SPOT
+	if (PassStruct.ShadowDepthType == SHADOW_DEPTH_TYPE_SPOT)
+	{
+		OutPosition = TransformShadowSpot(
+			PassStruct.ProjectionMatrix,
+			PassStruct.ViewMatrix,
+			WorldPos,
+			WorldNormal,
+			OutParameters.ShadowDepth,
+			OutParameters.DepthBias);
+	}
+#endif // SUPPORT_SHADOW_DEPTH_SPOT
 
-		TransformToVirtualSmPage(OutPosition, OutVirtualSmPageClip, PageInfo, WorldPos.xyz);
+#if SUPPORT_SHADOW_DEPTH_POINT
+	if (PassStruct.ShadowDepthType == SHADOW_DEPTH_TYPE_ONEPASSPOINT)
+	{
+		OutParameters.LayerIndex = bUseGpuSceneInstancing ? VertexFactoryGetViewIndex(VFIntermediates) : LayerId;
+		OutPosition = mul(WorldPos, PassStruct.ShadowViewProjectionMatrices[OutParameters.LayerIndex]);
+	}
+	else
+	{
+		OutParameters.LayerIndex = 0;
 	}
-#endif // VIRTUAL_SM_ENABLED
+#endif // SUPPORT_SHADOW_DEPTH_POINT
 }
 
 #if POSITION_ONLY
 void PositionOnlyMain(
 	in FPositionAndNormalOnlyVertexFactoryInput Input,
-#if ONEPASS_POINTLIGHT_SHADOW
-	out uint LayerIndex : SV_RenderTargetArrayIndex,
-#endif
 	out FShadowDepthVSToPS OutParameters,
-
-#if VIRTUAL_SM_ENABLED
-	out nointerpolation uint PackedPageInfo : TEXCOORD8,
-#endif
-	out float4 OutPosition : SV_POSITION
-#if VIRTUAL_SM_ENABLED	
-	// OLA-TODO: this collides with instanced stereo, which thankfully is not used with shadow maps, so should be fine, presumably.
-	, out float4 OutVirtualSmPageClip : SV_ClipDistance
-#endif // VIRTUAL_SM_ENABLED
-	)
+	out float4 OutPosition : SV_POSITION)
 {
 	ResolvedView = ResolveView();
 
 	float4 WorldPos = VertexFactoryGetWorldPosition(Input);
-
 #if USE_WORLD_POSITION_EXCLUDING_SHADER_OFFSETS && !IS_NANITE_PASS
 	OutParameters.PixelPositionExcludingWPO = WorldPos.xyz;
 #endif
 
+	float3 WorldNormal = VertexFactoryGetWorldNormal(Input);
+
 #if INTERPOLATE_VF_ATTRIBUTES
 	OutParameters.FactoryInterpolants = (FVertexFactoryInterpolantsVSToPS)0;
 #endif
 
-	float3 WorldNormal = VertexFactoryGetWorldNormal(Input);
-
-#if ONEPASS_POINTLIGHT_SHADOW
-	LayerIndex = bUseGpuSceneInstancing ? VertexFactoryGetViewIndex(Input) : LayerId;
-	OutPosition = mul(WorldPos, PassStruct.ShadowViewProjectionMatrices[LayerIndex]);
-
-#else // !ONEPASS_POINTLIGHT_SHADOW
-	float ShadowDepth;
-	SetShadowDepthOutputs(
-		PassStruct.ProjectionMatrix, 
-		PassStruct.ViewMatrix, 
-		WorldPos, 
-		WorldNormal, 
-		OutPosition,
-	#if PERSPECTIVE_CORRECT_DEPTH
-		OutParameters.ShadowDepth,
-		OutParameters.DepthBias
-	#else
-		ShadowDepth
-	#endif
-	);
-
-	#if !PERSPECTIVE_CORRECT_DEPTH && !COMPILER_SUPPORTS_EMPTY_STRUCTS
-		OutParameters.Dummy = 0;
-	#endif
-#endif // ONEPASS_POINTLIGHT_SHADOW
-
 #if INTERPOLATE_POSITION
 	OutParameters.PixelPosition = WorldPos.xyz;
 #endif
 
-#if VIRTUAL_SM_ENABLED
-	PackedPageInfo = 0;
+	OutPosition = float4(0, 0, 0, 0);
 
-	OutVirtualSmPageClip = float4(1.0f, 1.0f, 1.0f, 1.0f);
-	if (PassStruct.bRenderToVirtualShadowMap != 0)
+#if SUPPORT_SHADOW_DEPTH_VSM
+	if (PassStruct.ShadowDepthType == SHADOW_DEPTH_TYPE_VSM)
 	{
+		OutParameters.PackedPageInfo = 0;
+		OutParameters.VirtualSmPageClip = float4(1.0f, 1.0f, 1.0f, 1.0f);
+		
 		// Get the offset from which we loaded the instance ID
 		uint InstanceIdIndex = VertexFactoryGetInstanceIdLoadIndex(Input);
-		// TODO: Maybe offset index to local buffer for this? We may not want to use a global since most passes are not supporting this
-		//       Or perhaps both are different, as they are managed somewhat differently anyway, maybe.
-		PackedPageInfo = InstanceCulling.PageInfoBuffer[InstanceIdIndex];
-		FPageInfo PageInfo = UnpackPageInfo(PackedPageInfo);
-		TransformToVirtualSmPage(OutPosition, OutVirtualSmPageClip, PageInfo, WorldPos.xyz);
+		OutParameters.PackedPageInfo = InstanceCulling.PageInfoBuffer[InstanceIdIndex];
+		FPageInfo PageInfo = UnpackPageInfo(OutParameters.PackedPageInfo);
+		OutPosition = TransformShadowVSM(OutParameters.VirtualSmPageClip, PageInfo, WorldPos.xyz);
+	}
+	else
+	{
+		// TODO: Cleaner way to initialize unused outputs...
+		OutParameters.PackedPageInfo = 0;
+		OutParameters.VirtualSmPageClip = float4(1, 1, 1, 1);
+	}	
+#endif // SUPPORT_SHADOW_DEPTH_VSM
+
+#if	SUPPORT_SHADOW_DEPTH_DIRECTIONAL
+	if (PassStruct.ShadowDepthType == SHADOW_DEPTH_TYPE_DIRECTIONAL)
+	{
+		OutPosition = TransformShadowDirectional(
+			PassStruct.ProjectionMatrix,
+			PassStruct.ViewMatrix,
+			WorldPos, 
+			WorldNormal);
+	}
+#endif // SUPPORT_SHADOW_DEPTH_DIRECTIONAL
+
+#if SUPPORT_SHADOW_DEPTH_POINT
+	if (PassStruct.ShadowDepthType == SHADOW_DEPTH_TYPE_ONEPASSPOINT)
+	{
+		OutParameters.LayerIndex = bUseGpuSceneInstancing ? VertexFactoryGetViewIndex(Input) : LayerId;
+		OutPosition = mul(WorldPos, PassStruct.ShadowViewProjectionMatrices[OutParameters.LayerIndex]);
+	}
+	else
+	{
+		OutParameters.LayerIndex = 0;
+	}	
+#endif // SUPPORT_SHADOW_DEPTH_POINT
+
+#if SUPPORT_SHADOW_DEPTH_SPOT
+	if (PassStruct.ShadowDepthType == SHADOW_DEPTH_TYPE_SPOT)
+	{
+		OutPosition = TransformShadowSpot(
+			PassStruct.ProjectionMatrix,
+			PassStruct.ViewMatrix,
+			WorldPos,
+			WorldNormal,
+			OutParameters.ShadowDepth,
+			OutParameters.DepthBias);
 	}
-#endif // VIRTUAL_SM_ENABLED
+#endif // SUPPORT_SHADOW_DEPTH_SPOT
 }
 
 
diff --git a/Engine/Shaders/Private/TranslucentShadowDepthShaders.usf b/Engine/Shaders/Private/TranslucentShadowDepthShaders.usf
--- a/Engine/Shaders/Private/TranslucentShadowDepthShaders.usf
+++ b/Engine/Shaders/Private/TranslucentShadowDepthShaders.usf
@@ -34,7 +34,7 @@ void SetShadowDepthOutputs(float4 WorldPosition, out float4 OutPosition, out flo
 		OutPosition.w = 1.0f;
 	}
 
-	#if PERSPECTIVE_CORRECT_DEPTH
+	#if SUPPORT_SHADOW_DEPTH_SPOT
 		ShadowDepth = OutPosition.z;
 	#else
 		float InvMaxSubjectDepth = TranslucentDepthPass.InvMaxSubjectDepth;
@@ -120,7 +120,7 @@ void MainOpacityPS(
 
 	Density *= GetMaterialTranslucentShadowDensityScale();
 
-#if PERSPECTIVE_CORRECT_DEPTH
+#if SUPPORT_SHADOW_DEPTH_SPOT
 	/** Used to normalize the outputted depth */
 	Inputs.ShadowDepth *= TranslucentDepthPass.InvMaxSubjectDepth;
 #endif
diff --git a/Engine/Source/Runtime/RenderCore/Private/RenderUtils.cpp b/Engine/Source/Runtime/RenderCore/Private/RenderUtils.cpp
--- a/Engine/Source/Runtime/RenderCore/Private/RenderUtils.cpp
+++ b/Engine/Source/Runtime/RenderCore/Private/RenderUtils.cpp
@@ -60,6 +60,16 @@ static FAutoConsoleVariableRef CVarRayTracingProxiesProjectEnabled(
 	TEXT("Whether ray tracing proxies are generated by default for this project."),
 	ECVF_ReadOnly);
 
+// This is a per-project main switch for Virtual Shadow Maps support
+// Disabling this can remove some shader permutations if virtual shadow maps are not needed
+int32 GVirtualShadowMapsProjectEnabled = 1;
+FAutoConsoleVariableRef CVarAllowVirtualShadowMaps(
+	TEXT("r.Shadow.Virtual.ProjectEnabled"),
+	GVirtualShadowMapsProjectEnabled,
+	TEXT("This setting allows you to disable Virtual Shadow Maps on platforms that support it to reduce the number of shaders.\n"),
+	ECVF_ReadOnly | ECVF_RenderThreadSafe
+);
+
 // This is a per-project main switch for Non-Nanite support in Virtual Shadow Maps.
 // Disabling this can remove shader permutations and slightly improve performance, but Non-Nanite objects
 // will not cast shadows when Virtual Shadow Maps are enabled.
@@ -68,7 +78,19 @@ FAutoConsoleVariableRef CVarVirtualShadowMapNonNaniteProjectEnabled(
 	TEXT("r.Shadow.Virtual.NonNanite.ProjectEnabled"),
 	GVirtualShadowMapNonNaniteProjectEnabled,
 	TEXT("This can be used to disable non-nanite rendering into Virtual Shadow Maps and related shader permutations.\n")
-	TEXT("Only use this if you do not need any non-nanite geometry to cast shadows."),
+	TEXT("Only use this if you do not need any non-nanite geometry to cast shadows in VSMs."),
+	ECVF_RenderThreadSafe | ECVF_ReadOnly
+);
+
+// This is a per-project main switch for supporting the classic shadow map techniques in Unreal for Non-Nanite geometry
+// This includes cascaded shadow maps, (classic, non-VSM) spot light shadow maps and (classic, non-VSM) point light shadows
+// Disabling this can remove shader permutations if these techniques are not needed
+int32 GClassicShadowMapsProjectEnabled = 1;
+FAutoConsoleVariableRef CVarClassicShadowMapsProjectEnabled(
+	TEXT("r.Shadow.Classic.ProjectEnabled"),
+	GClassicShadowMapsProjectEnabled,
+	TEXT("This can be used to disable classic shadow maps (cascaded shadow maps, classic spot/point/rect light shadow maps).\n")
+	TEXT("This will remove a significant number of shader permutations if these techniques are not needed."),
 	ECVF_RenderThreadSafe | ECVF_ReadOnly
 );
 
@@ -1436,39 +1458,48 @@ bool AreNaniteFallbackMeshesEnabledForPlatform(const FStaticShaderPlatform Platf
 	return (GNaniteFallbackMeshesPlatformMask[(int)Platform]);
 }
 
-bool UseVirtualShadowMaps(EShaderPlatform ShaderPlatform)
+bool DoesPlatformSupportClassicShadowMaps(EShaderPlatform Platform)
 {
-	static const auto EnableVirtualSMCVar = IConsoleManager::Get().FindConsoleVariable(TEXT("r.Shadow.Virtual.Enable"));
-	const bool bVirtualShadowMapsEnabled = EnableVirtualSMCVar ? (EnableVirtualSMCVar->GetInt() != 0) : false;
-	return bVirtualShadowMapsEnabled && DoesRuntimeSupportNanite(ShaderPlatform, true /* check for atomics */, true /* check project setting */);
-}
-
-bool UseVirtualShadowMaps(EShaderPlatform ShaderPlatform, const FStaticFeatureLevel FeatureLevel)
-{
-	return UseVirtualShadowMaps(ShaderPlatform);
+	// Unconditionally supported unless disabled by the project setting
+	static FShaderPlatformCachedIniValue<int32> CVarPlatformClassicShadowMapsProjectEnabled(TEXT("r.Shadow.Classic.ProjectEnabled"));
+	bool bIsProjectEnable = CVarPlatformClassicShadowMapsProjectEnabled.Get(Platform) != 0;
+	return bIsProjectEnable;
 }
 
 bool DoesPlatformSupportVirtualShadowMaps(EShaderPlatform Platform)
 {
-	return DoesPlatformSupportNanite(Platform);
+	// VSM allowed for this project/platform
+	static FShaderPlatformCachedIniValue<int32> CVarPlatformVSMProjectEnabled(TEXT("r.Shadow.Virtual.ProjectEnabled"));
+	bool bIsProjectEnable = CVarPlatformVSMProjectEnabled.Get(Platform) != 0;
+
+	return bIsProjectEnable && DoesPlatformSupportNanite(Platform);
 }
 
 bool DoesPlatformSupportNonNaniteVirtualShadowMaps(EShaderPlatform ShaderPlatform)
 {
-	if (GVirtualShadowMapNonNaniteProjectEnabled == 0)
-	{
-		return false;
-	}
-	return DoesPlatformSupportVirtualShadowMaps(ShaderPlatform);
+	static FShaderPlatformCachedIniValue<int32> CVarPlatformNonNaniteVSMProjectEnabled(TEXT("r.Shadow.Virtual.NonNanite.ProjectEnabled"));
+	bool bIsProjectEnable = CVarPlatformNonNaniteVSMProjectEnabled.Get(ShaderPlatform) != 0;
+
+	return bIsProjectEnable && DoesPlatformSupportVirtualShadowMaps(ShaderPlatform);
 }
 
 bool UseNonNaniteVirtualShadowMaps(EShaderPlatform ShaderPlatform, const FStaticFeatureLevel FeatureLevel)
 {
-	if (GVirtualShadowMapNonNaniteProjectEnabled == 0)
-	{
-		return false;
-	}
-	return UseVirtualShadowMaps(ShaderPlatform, FeatureLevel);
+	return DoesPlatformSupportNonNaniteVirtualShadowMaps(ShaderPlatform) && UseVirtualShadowMaps(ShaderPlatform, FeatureLevel);
+}
+
+bool UseVirtualShadowMaps(EShaderPlatform ShaderPlatform)
+{
+	static const auto EnableVirtualSMCVar = IConsoleManager::Get().FindConsoleVariable(TEXT("r.Shadow.Virtual.Enable"));
+	const bool bVirtualShadowMapsEnabled = EnableVirtualSMCVar ? (EnableVirtualSMCVar->GetInt() != 0) : false;
+	return bVirtualShadowMapsEnabled &&
+		DoesPlatformSupportVirtualShadowMaps(ShaderPlatform) &&
+		DoesRuntimeSupportNanite(ShaderPlatform, true /* check for atomics */, true /* check project setting */);
+}
+
+bool UseVirtualShadowMaps(EShaderPlatform ShaderPlatform, const FStaticFeatureLevel FeatureLevel)
+{
+	return UseVirtualShadowMaps(ShaderPlatform);
 }
 
 bool IsWaterVirtualShadowMapFilteringEnabled(const FStaticShaderPlatform Platform)
diff --git a/Engine/Source/Runtime/RenderCore/Public/RenderUtils.h b/Engine/Source/Runtime/RenderCore/Public/RenderUtils.h
--- a/Engine/Source/Runtime/RenderCore/Public/RenderUtils.h
+++ b/Engine/Source/Runtime/RenderCore/Public/RenderUtils.h
@@ -438,6 +438,11 @@ RENDERCORE_API bool UseNanite(EShaderPlatform ShaderPlatform, bool bCheckForAtom
 */
 RENDERCORE_API bool AreNaniteFallbackMeshesEnabledForPlatform(const FStaticShaderPlatform Platform);
 
+/**
+* Returns true if classic (non-virtual) shadow map techniques are supported on the given shader platform.
+*/
+RENDERCORE_API bool DoesPlatformSupportClassicShadowMaps(EShaderPlatform Platform);
+
 /**
  * Returns true if Virtual Shadow Maps should be used for the given shader platform.
  * Note: Virtual Shadow Maps require Nanite support.
diff --git a/Engine/Source/Runtime/Renderer/Private/ShadowDepthRendering.cpp b/Engine/Source/Runtime/Renderer/Private/ShadowDepthRendering.cpp
--- a/Engine/Source/Runtime/Renderer/Private/ShadowDepthRendering.cpp
+++ b/Engine/Source/Runtime/Renderer/Private/ShadowDepthRendering.cpp
@@ -171,6 +171,8 @@ void SetupShadowDepthPassUniformBuffer(
 	ShadowDepthPassParameters.ShadowParams = FVector4f(ShadowInfo->GetShaderDepthBias(), bCSMCachingEnabled ? 0.0f : ShadowInfo->GetShaderSlopeDepthBias(), ShadowInfo->GetShaderMaxSlopeDepthBias(), ShadowInfo->bOnePassPointLightShadow ? 1 : ShadowInfo->InvMaxSubjectDepth);
 	ShadowDepthPassParameters.bClampToNearPlane = ShadowInfo->ShouldClampToNearPlane() ? 1.0f : 0.0f;
 
+	ShadowDepthPassParameters.ShadowDepthType = (uint32)ShadowInfo->GetShadowDepthType();
+
 	if (ShadowInfo->bOnePassPointLightShadow)
 	{
 		check(ShadowInfo->BorderSize == 0);
@@ -185,7 +187,6 @@ void SetupShadowDepthPassUniformBuffer(
 		}
 	}
 
-	ShadowDepthPassParameters.bRenderToVirtualShadowMap = false;
 	// Reference dummy data in the UB initially
 	const uint32 DummyPageTableElement = 0xFFFFFFFF;
 	ShadowDepthPassParameters.VirtualSmPageTable = GSystemTextures.GetDefaultTexture2D(GraphBuilder, PF_R32_UINT, DummyPageTableElement);
@@ -195,7 +196,7 @@ void SetupShadowDepthPassUniformBuffer(
 	ShadowDepthPassParameters.UncachedPageRectBounds = GraphBuilder.CreateSRV(GSystemTextures.GetDefaultStructuredBuffer(GraphBuilder, sizeof(FIntVector4)));
 
 	FRDGTextureRef DepthBufferArray = GraphBuilder.CreateTexture( FRDGTextureDesc::Create2DArray( FIntPoint(4,4), PF_R32_UINT, FClearValueBinding::None, TexCreate_ShaderResource | TexCreate_UAV, 1 ), TEXT("Dummy-OutDepthBuffer") );
-
+	
 	ShadowDepthPassParameters.OutDepthBufferArray = GraphBuilder.CreateUAV( DepthBufferArray );
 }
 
@@ -213,6 +214,8 @@ void SetupShadowDepthPassUniformBuffer(
 	ShadowDepthPassParameters.ShadowParams = FVector4f(ShadowInfo->GetShaderDepthBias(), ShadowInfo->GetShaderSlopeDepthBias(), ShadowInfo->GetShaderMaxSlopeDepthBias(), ShadowInfo->InvMaxSubjectDepth);
 	ShadowDepthPassParameters.bClampToNearPlane = ShadowInfo->ShouldClampToNearPlane() ? 1.0f : 0.0f;
 
+	ShadowDepthPassParameters.ShadowDepthType = (uint32)ShadowInfo->GetShadowDepthType();
+
 	if (ShadowInfo->bOnePassPointLightShadow)
 	{
 		check(ShadowInfo->BorderSize == 0);
@@ -293,14 +296,6 @@ private:
 	LAYOUT_FIELD(FShaderParameter, bUseGpuSceneInstancing);
 };
 
-enum EShadowDepthVertexShaderMode
-{
-	VertexShadowDepth_PerspectiveCorrect,
-	VertexShadowDepth_OutputDepth,
-	VertexShadowDepth_OnePassPointLight,
-	VertexShadowDepth_VirtualShadowMap,
-};
-
 static TAutoConsoleVariable<int32> CVarSupportPointLightWholeSceneShadows(
 	TEXT("r.SupportPointLightWholeSceneShadows"),
 	1,
@@ -315,7 +310,6 @@ static bool MobileUsesPerspectiveCorrectShadowPermutation(EShaderPlatform Shader
 	return bMobileEnableMovableSpotlightShadows;
 }
 
-
 static TAutoConsoleVariable<int32> CVarDetectVertexShaderLayerAtRuntime(
 	TEXT("r.Shadow.DetectVertexShaderLayerAtRuntime"),
 	0,
@@ -323,11 +317,97 @@ static TAutoConsoleVariable<int32> CVarDetectVertexShaderLayerAtRuntime(
 	TEXT("Enabled by default for windows/SM5 as DX11 almost universally supports this at runtime."),
 	ECVF_ReadOnly | ECVF_RenderThreadSafe);
 
+enum class EShadowDepthShaderMode
+{
+	PerspectiveCorrect,
+	OutputDepth,
+	OnePassPointLight,
+	VirtualShadowMap,
+	Combined,
+};
+
+// Returns bitfield of shadow depth types supported by the given shader mode
+// NOTE: By default any modes covered by the "Combined" mode will not be included as separate shader permutations,
+// and instead that mode will be used for any passes requesting the individual modes.
+static constexpr EShadowDepthType GetShadowDepthTypesForMode(EShadowDepthShaderMode Mode)
+{
+	switch (Mode)
+	{
+		case EShadowDepthShaderMode::PerspectiveCorrect: 	return EShadowDepthType::Spot;
+		case EShadowDepthShaderMode::OutputDepth: 			return EShadowDepthType::Directional;
+		case EShadowDepthShaderMode::OnePassPointLight: 	return EShadowDepthType::OnePassPoint;
+		case EShadowDepthShaderMode::VirtualShadowMap: 		return EShadowDepthType::VSM;
+		case EShadowDepthShaderMode::Combined: 				return EShadowDepthType::VSM | EShadowDepthType::OnePassPoint;
+		default: return EShadowDepthType::None;
+	}
+}
+
+static constexpr bool IsShadowDepthTypeInCombinedMode(EShadowDepthType Types)
+{
+	EShadowDepthType CombinedShadowDepthTypes = GetShadowDepthTypesForMode(EShadowDepthShaderMode::Combined);
+	// NOTE: This is slightly messy, but the shadow depth types for mesh passes get set to "VSM | Spot" or similar
+	// in the VSM case. In these cases we only need the VSM permutation, so using "Any" allows them to correctly
+	// route down to the combined permutation.
+	return EnumHasAnyFlags(CombinedShadowDepthTypes, Types);
+}
+
+static constexpr EShadowDepthType GetShadowDepthTypesThatRequirePS()
+{
+	return EShadowDepthType::VSM | EShadowDepthType::Spot;
+}
+
+// Returns a bitfield of shadow depth types supported on the given platform
+// TODO: Cache result for current platform...?
+EShadowDepthType GetPlatformSupportedShadowDepthTypes(EShaderPlatform Platform)
+{
+	// Compile VS layer permutation if RHI supports it unconditionally OR we have forced it on (default for DX11&12 at SM5)
+	static FShaderPlatformCachedIniValue<bool> DetectVertexShaderLayerRuntimeIniValue(TEXT("r.Shadow.DetectVertexShaderLayerAtRuntime"));
+	const bool bShouldCompileVSLayer = RHISupportsVertexShaderLayer(Platform) || DetectVertexShaderLayerRuntimeIniValue.Get(Platform) != 0;
+	
+	EShadowDepthType Types = EShadowDepthType::None;
+
+	if (!IsFeatureLevelSupported(Platform, ERHIFeatureLevel::SM5))
+	{
+		// Mobile only needs OutputDepth, and optionally PerspectiveCorrect and OnePassPointLight
+		Types |= EShadowDepthType::Directional
+			| (MobileUsesPerspectiveCorrectShadowPermutation(Platform) ? EShadowDepthType::Spot : EShadowDepthType::None)
+			| (FReadOnlyCVARCache::EnablePointLightShadows(Platform) && bShouldCompileVSLayer ? EShadowDepthType::OnePassPoint : EShadowDepthType::None);
+	}
+	else
+	{
+		Types |= EShadowDepthType::Directional
+			| EShadowDepthType::Spot
+			| (bShouldCompileVSLayer ? EShadowDepthType::OnePassPoint : EShadowDepthType::None)
+			| (DoesPlatformSupportNonNaniteVirtualShadowMaps(Platform) ? EShadowDepthType::VSM : EShadowDepthType::None);
+	}
+
+	// If classic shadow maps are disabled at the platform/project level, remove those permutations
+	if (!DoesPlatformSupportClassicShadowMaps(Platform))
+	{
+		Types &= ~(EShadowDepthType::Directional | EShadowDepthType::Spot | EShadowDepthType::OnePassPoint);
+	}
+
+	return Types;
+}
+
+// Get the "effective" supported shadow depth types for a given mode given both the mode and platform support.
+EShadowDepthType GetEffectiveShadowDepthTypes(EShaderPlatform Platform, EShadowDepthShaderMode Mode)
+{
+	return GetShadowDepthTypesForMode(Mode) & GetPlatformSupportedShadowDepthTypes(Platform);
+}
+
+void SetEnvironmentDefinesForShadowDepthTypes(FShaderCompilerEnvironment& OutEnvironment, EShadowDepthType Types)
+{
+	OutEnvironment.SetDefine(TEXT("SUPPORT_SHADOW_DEPTH_VSM"), 			EnumHasAnyFlags(Types, EShadowDepthType::VSM));
+	OutEnvironment.SetDefine(TEXT("SUPPORT_SHADOW_DEPTH_DIRECTIONAL"), 	EnumHasAnyFlags(Types, EShadowDepthType::Directional));
+	OutEnvironment.SetDefine(TEXT("SUPPORT_SHADOW_DEPTH_SPOT"), 		EnumHasAnyFlags(Types, EShadowDepthType::Spot));
+	OutEnvironment.SetDefine(TEXT("SUPPORT_SHADOW_DEPTH_POINT"), 		EnumHasAnyFlags(Types, EShadowDepthType::OnePassPoint));
+}
 
 /**
 * A vertex shader for rendering the depth of a mesh.
 */
-template <EShadowDepthVertexShaderMode ShaderMode, bool bUsePositionOnlyStream>
+template <EShadowDepthShaderMode ShaderMode, bool bUsePositionOnlyStream>
 class TShadowDepthVS : public FShadowDepthVS
 {
 public:
@@ -335,36 +415,19 @@ public:
 
 	static bool ShouldCompilePermutation(const FMeshMaterialShaderPermutationParameters& Parameters)
 	{
-		const EShaderPlatform Platform = Parameters.Platform;
-
-		static const auto SupportAllShaderPermutationsVar = IConsoleManager::Get().FindTConsoleVariableDataInt(TEXT("r.SupportAllShaderPermutations"));
-		const bool bForceAllPermutations = SupportAllShaderPermutationsVar && SupportAllShaderPermutationsVar->GetValueOnAnyThread() != 0;
-		const bool bSupportPointLightWholeSceneShadows = CVarSupportPointLightWholeSceneShadows.GetValueOnAnyThread() != 0 || bForceAllPermutations;
-
-		// Mobile only needs OutputDepth, and optionally PerspectiveCorrect and OnePassPointLight
-		if (!IsFeatureLevelSupported(Platform, ERHIFeatureLevel::SM5) &&
-			!(ShaderMode == VertexShadowDepth_OutputDepth || 
-				(ShaderMode == VertexShadowDepth_PerspectiveCorrect && MobileUsesPerspectiveCorrectShadowPermutation(Platform))
-					 || ((ShaderMode == VertexShadowDepth_OnePassPointLight) && FReadOnlyCVARCache::EnablePointLightShadows(Platform))))
+		EShadowDepthType ShadowDepthTypes = GetEffectiveShadowDepthTypes(Parameters.Platform, ShaderMode);
+		if (ShadowDepthTypes == EShadowDepthType::None)
 		{
 			return false;
 		}
 
-		// Compile VS layer permutation if RHI supports it unconditionally OR we have forced it on (default for DX11&12 at SM5)
-		static FShaderPlatformCachedIniValue<bool> DetectVertexShaderLayerRuntimeIniValue(TEXT("r.Shadow.DetectVertexShaderLayerAtRuntime"));
-		const bool bShouldCompileVSLayer = RHISupportsVertexShaderLayer(Platform) || DetectVertexShaderLayerRuntimeIniValue.Get(Platform) != 0;
-		if (ShaderMode == VertexShadowDepth_OnePassPointLight && !bShouldCompileVSLayer)
+		// Don't include this permutation if it is not the combined mode, but the depth types are covered by it
+		if (ShaderMode != EShadowDepthShaderMode::Combined && IsShadowDepthTypeInCombinedMode(ShadowDepthTypes))
 		{
 			return false;
 		}
 
-		if (ShaderMode == VertexShadowDepth_VirtualShadowMap && !DoesPlatformSupportNonNaniteVirtualShadowMaps(Platform))
-		{
-			return false;
-		}
-
-
-		//Note: This logic needs to stay in sync with UseDefaultMaterialForShadowDepth!
+		// Note: This logic needs to stay in sync with UseDefaultMaterialForShadowDepth!
 		return (Parameters.MaterialParameters.bIsSpecialEngineMaterial
 			// Masked and WPO materials need their shaders but cannot be used with a position only stream.
 			|| ((!Parameters.MaterialParameters.bWritesEveryPixelShadowPass || Parameters.MaterialParameters.bMaterialMayModifyMeshPosition) && !bUsePositionOnlyStream))
@@ -378,19 +441,21 @@ public:
 	static void ModifyCompilationEnvironment(const FMaterialShaderPermutationParameters& Parameters, FShaderCompilerEnvironment& OutEnvironment)
 	{
 		FShadowDepthVS::ModifyCompilationEnvironment(Parameters, OutEnvironment);
-		OutEnvironment.SetDefine(TEXT("PERSPECTIVE_CORRECT_DEPTH"), (uint32)(ShaderMode == VertexShadowDepth_PerspectiveCorrect));
-		OutEnvironment.SetDefine(TEXT("ONEPASS_POINTLIGHT_SHADOW"), (uint32)(ShaderMode == VertexShadowDepth_OnePassPointLight));
+
+		EShadowDepthType Types = GetEffectiveShadowDepthTypes(Parameters.Platform, ShaderMode);
+		check(Types != EShadowDepthType::None);		// Should have been killed by ShouldCompilePermutation
+		SetEnvironmentDefinesForShadowDepthTypes(OutEnvironment, Types);
+		
 		OutEnvironment.SetDefine(TEXT("POSITION_ONLY"), (uint32)bUsePositionOnlyStream);
 
-		bool bEnableNonNaniteVSM = (ShaderMode == VertexShadowDepth_VirtualShadowMap);
-		OutEnvironment.SetDefine(TEXT("ENABLE_NON_NANITE_VSM"), bEnableNonNaniteVSM ? 1 : 0);
-		if (bEnableNonNaniteVSM)
+		if (EnumHasAnyFlags(Types, EShadowDepthType::VSM))
 		{
 			FVirtualShadowMapArray::SetShaderDefines(OutEnvironment);
 		}
 
-		if (ShaderMode == VertexShadowDepth_OnePassPointLight)
+		if (EnumHasAnyFlags(Types, EShadowDepthType::OnePassPoint))
 		{
+			// TODO: This is possibly harmful for other shadow depth modes in the combined case
 			OutEnvironment.CompilerFlags.Add(CFLAG_VertexUseAutoCulling);
 		}
 	}
@@ -402,24 +467,16 @@ public:
 };
 
 #define IMPLEMENT_SHADOW_DEPTH_SHADERMODE_SHADERS(ShaderMode) \
-	typedef TShadowDepthVS<ShaderMode, false> TShadowDepthVS##ShaderMode; \
-	IMPLEMENT_MATERIAL_SHADER_TYPE(template<>, TShadowDepthVS##ShaderMode, TEXT("/Engine/Private/ShadowDepthVertexShader.usf"), TEXT("Main"),       SF_Vertex);
-
-IMPLEMENT_SHADOW_DEPTH_SHADERMODE_SHADERS(VertexShadowDepth_PerspectiveCorrect);
-IMPLEMENT_SHADOW_DEPTH_SHADERMODE_SHADERS(VertexShadowDepth_OutputDepth);
-IMPLEMENT_SHADOW_DEPTH_SHADERMODE_SHADERS(VertexShadowDepth_OnePassPointLight);
-IMPLEMENT_SHADOW_DEPTH_SHADERMODE_SHADERS(VertexShadowDepth_VirtualShadowMap);
-
-// Position only vertex shaders.
-typedef TShadowDepthVS<VertexShadowDepth_PerspectiveCorrect, true> TShadowDepthVSVertexShadowDepth_PerspectiveCorrectPositionOnly;
-typedef TShadowDepthVS<VertexShadowDepth_OutputDepth,        true> TShadowDepthVSVertexShadowDepth_OutputDepthPositionOnly;
-typedef TShadowDepthVS<VertexShadowDepth_OnePassPointLight,  true> TShadowDepthVSVertexShadowDepth_OnePassPointLightPositionOnly;
-typedef TShadowDepthVS<VertexShadowDepth_VirtualShadowMap,   true> TShadowDepthVSVertexShadowDepth_VirtualShadowMapPositionOnly;
-IMPLEMENT_MATERIAL_SHADER_TYPE(template<>, TShadowDepthVSVertexShadowDepth_PerspectiveCorrectPositionOnly, TEXT("/Engine/Private/ShadowDepthVertexShader.usf"), TEXT("PositionOnlyMain"), SF_Vertex);
-IMPLEMENT_MATERIAL_SHADER_TYPE(template<>, TShadowDepthVSVertexShadowDepth_OutputDepthPositionOnly,        TEXT("/Engine/Private/ShadowDepthVertexShader.usf"), TEXT("PositionOnlyMain"), SF_Vertex);
-IMPLEMENT_MATERIAL_SHADER_TYPE(template<>, TShadowDepthVSVertexShadowDepth_OnePassPointLightPositionOnly,  TEXT("/Engine/Private/ShadowDepthVertexShader.usf"), TEXT("PositionOnlyMain"), SF_Vertex);
-IMPLEMENT_MATERIAL_SHADER_TYPE(template<>, TShadowDepthVSVertexShadowDepth_VirtualShadowMapPositionOnly,   TEXT("/Engine/Private/ShadowDepthVertexShader.usf"), TEXT("PositionOnlyMain"), SF_Vertex);
+	typedef TShadowDepthVS<EShadowDepthShaderMode::ShaderMode, false> TShadowDepthVS##ShaderMode; \
+	IMPLEMENT_MATERIAL_SHADER_TYPE(template<>, TShadowDepthVS##ShaderMode, TEXT("/Engine/Private/ShadowDepthVertexShader.usf"), TEXT("Main"), SF_Vertex); \
+	typedef TShadowDepthVS<EShadowDepthShaderMode::ShaderMode, true> TShadowDepthVSPositionOnly##ShaderMode; \
+	IMPLEMENT_MATERIAL_SHADER_TYPE(template<>, TShadowDepthVSPositionOnly##ShaderMode, TEXT("/Engine/Private/ShadowDepthVertexShader.usf"), TEXT("PositionOnlyMain"), SF_Vertex);
 
+IMPLEMENT_SHADOW_DEPTH_SHADERMODE_SHADERS(PerspectiveCorrect);
+IMPLEMENT_SHADOW_DEPTH_SHADERMODE_SHADERS(OutputDepth);
+IMPLEMENT_SHADOW_DEPTH_SHADERMODE_SHADERS(OnePassPointLight);
+IMPLEMENT_SHADOW_DEPTH_SHADERMODE_SHADERS(VirtualShadowMap);
+IMPLEMENT_SHADOW_DEPTH_SHADERMODE_SHADERS(Combined);
 
 /**
 * A pixel shader for rendering the depth of a mesh.
@@ -448,15 +505,7 @@ public:
 	FShadowDepthBasePS() = default;
 };
 
-enum EShadowDepthPixelShaderMode
-{
-	PixelShadowDepth_NonPerspectiveCorrect,
-	PixelShadowDepth_PerspectiveCorrect,
-	PixelShadowDepth_OnePassPointLight,
-	PixelShadowDepth_VirtualShadowMap
-};
-
-template <EShadowDepthPixelShaderMode ShaderMode>
+template <EShadowDepthShaderMode ShaderMode>
 class TShadowDepthPS : public FShadowDepthBasePS
 {
 	DECLARE_SHADER_TYPE(TShadowDepthPS, MeshMaterial);
@@ -464,23 +513,19 @@ public:
 
 	static bool ShouldCompilePermutation(const FMeshMaterialShaderPermutationParameters& Parameters)
 	{
-		const EShaderPlatform Platform = Parameters.Platform;
-
-		// Mobile only needs NonPerspectiveCorrect, and optionally PerspectiveCorrect and OnePassPointLight
-		if (!IsFeatureLevelSupported(Platform, ERHIFeatureLevel::SM5) &&
-			!(ShaderMode == PixelShadowDepth_NonPerspectiveCorrect || 
-				(ShaderMode == PixelShadowDepth_PerspectiveCorrect && MobileUsesPerspectiveCorrectShadowPermutation(Platform))
-					|| ((ShaderMode == PixelShadowDepth_OnePassPointLight) && FReadOnlyCVARCache::EnablePointLightShadows(Platform))))
+		EShadowDepthType ShadowDepthTypes = GetEffectiveShadowDepthTypes(Parameters.Platform, ShaderMode);
+		if (ShadowDepthTypes == EShadowDepthType::None)
 		{
 			return false;
 		}
 
-		if (ShaderMode == PixelShadowDepth_VirtualShadowMap && !DoesPlatformSupportNonNaniteVirtualShadowMaps(Platform))
+		// Don't include this permutation if it is not the combined mode, but the depth types are covered by it
+		if (ShaderMode != EShadowDepthShaderMode::Combined && IsShadowDepthTypeInCombinedMode(ShadowDepthTypes))
 		{
 			return false;
 		}
 
-		bool bModeRequiresPS = ShaderMode == PixelShadowDepth_PerspectiveCorrect || ShaderMode == PixelShadowDepth_VirtualShadowMap;
+		const bool bModeRequiresPS = EnumHasAnyFlags(GetShadowDepthTypesThatRequirePS(), ShadowDepthTypes);
 
 		//Note: This logic needs to stay in sync with UseDefaultMaterialForShadowDepth!
 		return (Parameters.MaterialParameters.bIsSpecialEngineMaterial
@@ -497,13 +542,12 @@ public:
 	static void ModifyCompilationEnvironment(const FMaterialShaderPermutationParameters& Parameters, FShaderCompilerEnvironment& OutEnvironment)
 	{
 		FShadowDepthBasePS::ModifyCompilationEnvironment(Parameters, OutEnvironment);
-		OutEnvironment.SetDefine(TEXT("PERSPECTIVE_CORRECT_DEPTH"), (uint32)(ShaderMode == PixelShadowDepth_PerspectiveCorrect));
-		OutEnvironment.SetDefine(TEXT("ONEPASS_POINTLIGHT_SHADOW"), (uint32)(ShaderMode == PixelShadowDepth_OnePassPointLight));
-		OutEnvironment.SetDefine(TEXT("VIRTUAL_TEXTURE_TARGET"), (uint32)(ShaderMode == PixelShadowDepth_VirtualShadowMap));
 
-		bool bEnableNonNaniteVSM = (ShaderMode == PixelShadowDepth_VirtualShadowMap);
-		OutEnvironment.SetDefine(TEXT("ENABLE_NON_NANITE_VSM"), bEnableNonNaniteVSM ? 1 : 0);
-		if (bEnableNonNaniteVSM != 0)
+		EShadowDepthType Types = GetEffectiveShadowDepthTypes(Parameters.Platform, ShaderMode);
+		check(Types != EShadowDepthType::None);		// Should have been killed by ShouldCompilePermutation
+		SetEnvironmentDefinesForShadowDepthTypes(OutEnvironment, Types);		
+
+		if (EnumHasAnyFlags(Types, EShadowDepthType::VSM))
 		{
 			FVirtualShadowMapArray::SetShaderDefines(OutEnvironment);
 		}
@@ -521,13 +565,14 @@ public:
 
 // typedef required to get around macro expansion failure due to commas in template argument list for TShadowDepthPixelShader
 #define IMPLEMENT_SHADOWDEPTHPASS_PIXELSHADER_TYPE(ShaderMode) \
-	typedef TShadowDepthPS<ShaderMode> TShadowDepthPS##ShaderMode; \
+	typedef TShadowDepthPS<EShadowDepthShaderMode::ShaderMode> TShadowDepthPS##ShaderMode; \
 	IMPLEMENT_MATERIAL_SHADER_TYPE(template<>,TShadowDepthPS##ShaderMode,TEXT("/Engine/Private/ShadowDepthPixelShader.usf"),TEXT("Main"),SF_Pixel);
 
-IMPLEMENT_SHADOWDEPTHPASS_PIXELSHADER_TYPE(PixelShadowDepth_NonPerspectiveCorrect);
-IMPLEMENT_SHADOWDEPTHPASS_PIXELSHADER_TYPE(PixelShadowDepth_PerspectiveCorrect);
-IMPLEMENT_SHADOWDEPTHPASS_PIXELSHADER_TYPE(PixelShadowDepth_OnePassPointLight);
-IMPLEMENT_SHADOWDEPTHPASS_PIXELSHADER_TYPE(PixelShadowDepth_VirtualShadowMap);
+IMPLEMENT_SHADOWDEPTHPASS_PIXELSHADER_TYPE(PerspectiveCorrect);
+IMPLEMENT_SHADOWDEPTHPASS_PIXELSHADER_TYPE(OutputDepth);
+IMPLEMENT_SHADOWDEPTHPASS_PIXELSHADER_TYPE(OnePassPointLight);
+IMPLEMENT_SHADOWDEPTHPASS_PIXELSHADER_TYPE(VirtualShadowMap);
+IMPLEMENT_SHADOWDEPTHPASS_PIXELSHADER_TYPE(Combined);
 
 /**
  * Check if default material can be used for shadow depth rendering 
@@ -560,94 +605,117 @@ bool GetShadowDepthPassShaders(
 		&& Material.WritesEveryPixel(true, bVFTypeSupportsNullPixelShader)
 		&& !bMaterialModifiesMeshPosition;
 
-	// Use perspective correct shadow depths for shadow types which typically render low poly meshes into the shadow depth buffer.
-	// Depth will be interpolated to the pixel shader and written out, which disables HiZ and double speed Z.
-	// Directional light shadows use an ortho projection and can use the non-perspective correct path without artifacts.
-	// One pass point lights don't output a linear depth, so they are already perspective correct.
-	bool bUsePerspectiveCorrectShadowDepths = !EnumHasAnyFlags(ShadowDepthType, EShadowDepthType::Directional | EShadowDepthType::OnePassPoint);
-	bool bOnePassPointLightShadow = EnumHasAnyFlags(ShadowDepthType, EShadowDepthType::OnePassPoint);
-
-	bool bVirtualShadowMap = EnumHasAnyFlags(ShadowDepthType, EShadowDepthType::VSM);
-	if (bVirtualShadowMap)
-	{
-		bUsePerspectiveCorrectShadowDepths = false;
-		bOnePassPointLightShadow = false;
-	}
+	const EShaderPlatform Platform = GShaderPlatformForFeatureLevel[FeatureLevel];
 
 	FMaterialShaderTypes ShaderTypes;
 
-	// Vertex related shaders
-	if (bOnePassPointLightShadow)
+	const bool bUseCombinedMode = IsShadowDepthTypeInCombinedMode(ShadowDepthType);
+
+	// Vertex shaders
+	if (bUseCombinedMode)
 	{
-		if (DoesRuntimeSupportOnePassPointLightShadows(GShaderPlatformForFeatureLevel[FeatureLevel]))
+		if (bPositionOnlyVS)
 		{
-			if (bPositionOnlyVS)
-			{
-				ShaderTypes.AddShaderType<TShadowDepthVS<VertexShadowDepth_OnePassPointLight, true>>();
-			}
-			else
-			{
-				ShaderTypes.AddShaderType<TShadowDepthVS<VertexShadowDepth_OnePassPointLight, false>>();
-			}
+			ShaderTypes.AddShaderType<TShadowDepthVS<EShadowDepthShaderMode::Combined, true>>();
 		}
 		else
 		{
-			return false;
+			ShaderTypes.AddShaderType<TShadowDepthVS<EShadowDepthShaderMode::Combined, false>>();
 		}
 	}
-	else if (bVirtualShadowMap)
+	else if (EnumHasAnyFlags(ShadowDepthType, EShadowDepthType::VSM))
 	{
 		if (bPositionOnlyVS)
 		{
-			ShaderTypes.AddShaderType<TShadowDepthVS<VertexShadowDepth_VirtualShadowMap, true>>();
+			ShaderTypes.AddShaderType<TShadowDepthVS<EShadowDepthShaderMode::VirtualShadowMap, true>>();
 		}
 		else
 		{
-			ShaderTypes.AddShaderType<TShadowDepthVS<VertexShadowDepth_VirtualShadowMap, false>>();
+			ShaderTypes.AddShaderType<TShadowDepthVS<EShadowDepthShaderMode::VirtualShadowMap, false>>();
 		}
 	}
-	else if (bUsePerspectiveCorrectShadowDepths)
+	else
 	{
-		if (bPositionOnlyVS)
+		// Classic shadow map techniques. If disabled, we bail as we can't render this batch.
+		if (!DoesPlatformSupportClassicShadowMaps(Platform))
 		{
-			ShaderTypes.AddShaderType<TShadowDepthVS<VertexShadowDepth_PerspectiveCorrect, true>>();
+			return false;
 		}
-		else
+
+		if (EnumHasAnyFlags(ShadowDepthType, EShadowDepthType::OnePassPoint))
+		{
+			if (DoesRuntimeSupportOnePassPointLightShadows(Platform))
+			{
+				if (bPositionOnlyVS)
+				{
+					ShaderTypes.AddShaderType<TShadowDepthVS<EShadowDepthShaderMode::OnePassPointLight, true>>();
+				}
+				else
+				{
+					ShaderTypes.AddShaderType<TShadowDepthVS<EShadowDepthShaderMode::OnePassPointLight, false>>();
+				}
+			}
+			else
+			{
+				return false;
+			}
+		}
+		else if (EnumHasAnyFlags(ShadowDepthType, EShadowDepthType::Spot))
 		{
-			ShaderTypes.AddShaderType<TShadowDepthVS<VertexShadowDepth_PerspectiveCorrect, false>>();
+			if (bPositionOnlyVS)
+			{
+				ShaderTypes.AddShaderType<TShadowDepthVS<EShadowDepthShaderMode::PerspectiveCorrect, true>>();
+			}
+			else
+			{
+				ShaderTypes.AddShaderType<TShadowDepthVS<EShadowDepthShaderMode::PerspectiveCorrect, false>>();
+			}
 		}
-	}
-	else
-	{
-		if (bPositionOnlyVS)
+		else if (EnumHasAnyFlags(ShadowDepthType, EShadowDepthType::Directional))
 		{
-			ShaderTypes.AddShaderType<TShadowDepthVS<VertexShadowDepth_OutputDepth, true>>();
+			if (bPositionOnlyVS)
+			{
+				ShaderTypes.AddShaderType<TShadowDepthVS<EShadowDepthShaderMode::OutputDepth, true>>();
+			}
+			else
+			{
+				ShaderTypes.AddShaderType<TShadowDepthVS<EShadowDepthShaderMode::OutputDepth, false>>();
+			}
 		}
 		else
 		{
-			ShaderTypes.AddShaderType<TShadowDepthVS<VertexShadowDepth_OutputDepth, false>>();
+			return false;
 		}
 	}
 
 	// Pixel shaders
-	const bool bNullPixelShader = Material.WritesEveryPixel(true, bVFTypeSupportsNullPixelShader) && !bUsePerspectiveCorrectShadowDepths && !bVirtualShadowMap && bVFTypeSupportsNullPixelShader;
+	const bool bNullPixelShader = 
+		bVFTypeSupportsNullPixelShader &&
+		Material.WritesEveryPixel(true, bVFTypeSupportsNullPixelShader) &&
+		!EnumHasAnyFlags(GetShadowDepthTypesThatRequirePS(), ShadowDepthType);
+
 	if (!bNullPixelShader)
 	{
-		if (bVirtualShadowMap)
+		// TODO: Move this around to the same branch above, just respect NullPixelShader as appropriate
+		if (bUseCombinedMode)
+		{
+			ShaderTypes.AddShaderType<TShadowDepthPS<EShadowDepthShaderMode::Combined>>();
+		}
+		else if (EnumHasAnyFlags(ShadowDepthType, EShadowDepthType::VSM))
 		{
-			ShaderTypes.AddShaderType<TShadowDepthPS<PixelShadowDepth_VirtualShadowMap>>();
+			ShaderTypes.AddShaderType<TShadowDepthPS<EShadowDepthShaderMode::VirtualShadowMap>>();
 		}
-		else if (bUsePerspectiveCorrectShadowDepths)
+		else if (EnumHasAnyFlags(ShadowDepthType, EShadowDepthType::OnePassPoint))
 		{
-			ShaderTypes.AddShaderType<TShadowDepthPS<PixelShadowDepth_PerspectiveCorrect>>();
+			ShaderTypes.AddShaderType<TShadowDepthPS<EShadowDepthShaderMode::OnePassPointLight>>();
 		}
-		else if (bOnePassPointLightShadow)
+		else if (EnumHasAnyFlags(ShadowDepthType, EShadowDepthType::Spot))
 		{
-			ShaderTypes.AddShaderType<TShadowDepthPS<PixelShadowDepth_OnePassPointLight>>();
+			ShaderTypes.AddShaderType<TShadowDepthPS<EShadowDepthShaderMode::PerspectiveCorrect>>();
 		}
 		else
 		{
-			ShaderTypes.AddShaderType<TShadowDepthPS<PixelShadowDepth_NonPerspectiveCorrect>>();
+			ShaderTypes.AddShaderType<TShadowDepthPS<EShadowDepthShaderMode::OutputDepth>>();
 		}
 	}
 
diff --git a/Engine/Source/Runtime/Renderer/Private/ShadowRendering.h b/Engine/Source/Runtime/Renderer/Private/ShadowRendering.h
--- a/Engine/Source/Runtime/Renderer/Private/ShadowRendering.h
+++ b/Engine/Source/Runtime/Renderer/Private/ShadowRendering.h
@@ -87,7 +87,7 @@ enum class EShadowDepthType
 	None = 0,
 	VSM = 1 << 0,
 	Directional = 1 << 1,
-	Point = 1 << 2,
+	Spot = 1 << 2,
 	OnePassPoint = 1 << 3
 };
 
@@ -254,7 +254,7 @@ BEGIN_GLOBAL_SHADER_PARAMETER_STRUCT(FShadowDepthPassUniformParameters,)
 	SHADER_PARAMETER(float, bClampToNearPlane)
 	SHADER_PARAMETER_ARRAY(FMatrix44f, ShadowViewProjectionMatrices, [6])
 	SHADER_PARAMETER_ARRAY(FMatrix44f, ShadowViewMatrices, [6])
-	SHADER_PARAMETER(int, bRenderToVirtualShadowMap)
+	SHADER_PARAMETER(uint32, ShadowDepthType)
 
 	SHADER_PARAMETER_RDG_TEXTURE(Texture2D<uint>, VirtualSmPageTable)
 	SHADER_PARAMETER_RDG_BUFFER_SRV(StructuredBuffer< FPackedNaniteView >, PackedNaniteViews)
@@ -270,6 +270,7 @@ BEGIN_GLOBAL_SHADER_PARAMETER_STRUCT(FMobileShadowDepthPassUniformParameters, RE
 	SHADER_PARAMETER(FVector4f, ShadowParams)
 	SHADER_PARAMETER(float, bClampToNearPlane)
 	SHADER_PARAMETER_ARRAY(FMatrix44f, ShadowViewProjectionMatrices, [6])
+	SHADER_PARAMETER(uint32, ShadowDepthType)
 END_GLOBAL_SHADER_PARAMETER_STRUCT()
 
 /**
@@ -796,7 +797,7 @@ public:
 		}
 		else
 		{
-			ShadowDepthType |= EShadowDepthType::Point;
+			ShadowDepthType |= EShadowDepthType::Spot;
 		}
 
 		return ShadowDepthType;
diff --git a/Engine/Source/Runtime/Renderer/Private/TranslucentLighting.cpp b/Engine/Source/Runtime/Renderer/Private/TranslucentLighting.cpp
--- a/Engine/Source/Runtime/Renderer/Private/TranslucentLighting.cpp
+++ b/Engine/Source/Runtime/Renderer/Private/TranslucentLighting.cpp
@@ -429,7 +429,7 @@ public:
 	static void ModifyCompilationEnvironment(const FMaterialShaderPermutationParameters& Parameters, FShaderCompilerEnvironment& OutEnvironment)
 	{
 		FTranslucencyShadowDepthVS::ModifyCompilationEnvironment(Parameters, OutEnvironment);
-		OutEnvironment.SetDefine(TEXT("PERSPECTIVE_CORRECT_DEPTH"), (uint32)(ShaderMode == TranslucencyShadowDepth_PerspectiveCorrect ? 1 : 0));
+		OutEnvironment.SetDefine(TEXT("SUPPORT_SHADOW_DEPTH_SPOT"), (uint32)(ShaderMode == TranslucencyShadowDepth_PerspectiveCorrect ? 1 : 0));
 	}
 };
 
@@ -482,7 +482,7 @@ public:
 	static void ModifyCompilationEnvironment(const FMaterialShaderPermutationParameters& Parameters, FShaderCompilerEnvironment& OutEnvironment)
 	{
 		FTranslucencyShadowDepthPS::ModifyCompilationEnvironment(Parameters, OutEnvironment);
-		OutEnvironment.SetDefine(TEXT("PERSPECTIVE_CORRECT_DEPTH"), (uint32)(ShaderMode == TranslucencyShadowDepth_PerspectiveCorrect ? 1 : 0));
+		OutEnvironment.SetDefine(TEXT("SUPPORT_SHADOW_DEPTH_SPOT"), (uint32)(ShaderMode == TranslucencyShadowDepth_PerspectiveCorrect ? 1 : 0));
 		OutEnvironment.SetDefine(TEXT("SUBSTRATE_INLINE_SHADING"), 1);
 	}
 
diff --git a/Engine/Source/Runtime/Renderer/Private/VirtualShadowMaps/VirtualShadowMapArray.cpp b/Engine/Source/Runtime/Renderer/Private/VirtualShadowMaps/VirtualShadowMapArray.cpp
--- a/Engine/Source/Runtime/Renderer/Private/VirtualShadowMaps/VirtualShadowMapArray.cpp
+++ b/Engine/Source/Runtime/Renderer/Private/VirtualShadowMaps/VirtualShadowMapArray.cpp
@@ -4076,7 +4076,7 @@ void FVirtualShadowMapArray::RenderVirtualShadowMapsNonNanite(FRDGBuilder& Graph
 		ShadowDepthPassParameters->ProjectionMatrix = FMatrix44f::Identity;
 		ShadowDepthPassParameters->ViewMatrix = FMatrix44f::Identity;
 		ShadowDepthPassParameters->ShadowParams = FVector4f(0.0f, 0.0f, 0.0f, 1.0f);
-		ShadowDepthPassParameters->bRenderToVirtualShadowMap = true;
+		ShadowDepthPassParameters->ShadowDepthType = (uint32)EShadowDepthType::VSM;
 
 		ShadowDepthPassParameters->VirtualSmPageTable = PageTableRDG;
 		ShadowDepthPassParameters->PackedNaniteViews = GraphBuilder.CreateSRV(VirtualShadowViewsRDG);
