[HWRT] Moved DynamicReadBufferForRayTracing.Commit(...) out of FDynamicRayTracingInstancesContext::Finish(...).
- the same DynamicReadBufferForRayTracing can be shared by multiple FDynamicRayTracingInstancesContext.
- also removed some used parameters from RayTracing::FinishGatherInstances(...).

diff --git a/Engine/Source/Runtime/Renderer/Private/DeferredShadingRenderer.cpp b/Engine/Source/Runtime/Renderer/Private/DeferredShadingRenderer.cpp
--- a/Engine/Source/Runtime/Renderer/Private/DeferredShadingRenderer.cpp
+++ b/Engine/Source/Runtime/Renderer/Private/DeferredShadingRenderer.cpp
@@ -2168,9 +2168,7 @@ void FDeferredShadingSceneRenderer::Render(FRDGBuilder& GraphBuilder, const FSce
 				GraphBuilder,
 				*InitViewTaskDatas.RayTracingGatherInstances,
 				RayTracingScene,
-				RayTracingSBT,
-				DynamicReadBufferForRayTracing,
-				Allocator);
+				RayTracingSBT);
 		}
 #endif // RHI_RAYTRACING
 
diff --git a/Engine/Source/Runtime/Renderer/Private/RayTracing/RayTracing.cpp b/Engine/Source/Runtime/Renderer/Private/RayTracing/RayTracing.cpp
--- a/Engine/Source/Runtime/Renderer/Private/RayTracing/RayTracing.cpp
+++ b/Engine/Source/Runtime/Renderer/Private/RayTracing/RayTracing.cpp
@@ -815,8 +815,6 @@ namespace RayTracing
 		DynamicIndexBuffer.Commit();
 		RHICmdList->FinishRecording();
 
-		FSceneRenderer::DynamicReadBufferForRayTracing.Commit(InRHICmdList);
-
 		InRHICmdList.QueueAsyncCommandListSubmit(RHICmdList);
 	}
 
@@ -1978,9 +1976,7 @@ namespace RayTracing
 		FRDGBuilder& GraphBuilder,
 		FGatherInstancesTaskData& TaskData,
 		FRayTracingScene& RayTracingScene,
-		FRayTracingShaderBindingTable& RayTracingSBT,
-		FGlobalDynamicReadBuffer& InDynamicReadBuffer,
-		FSceneRenderingBulkObjectAllocator& InBulkAllocator)
+		FRayTracingShaderBindingTable& RayTracingSBT)
 	{
 		TRACE_CPUPROFILER_EVENT_SCOPE(RayTracing_FinishGatherInstances);
 		SCOPE_CYCLE_COUNTER(STAT_RayTracing_FinishGatherInstances);
@@ -2013,6 +2009,8 @@ namespace RayTracing
 			TaskData.DynamicRayTracingInstancesContext->GatherDynamicRayTracingInstances_RenderThread();
 			TaskData.DynamicRayTracingInstancesContext->Finish(GraphBuilder.RHICmdList);
 
+			FSceneRenderer::DynamicReadBufferForRayTracing.Commit(GraphBuilder.RHICmdList);
+
 			for (FGatherInstancesViewTaskData& ViewTaskData : TaskData.ViewTaskDatas)
 			{
 				ViewTaskData.Scene.GPUScene.UploadDynamicPrimitiveShaderDataForView(GraphBuilder, ViewTaskData.View, /*bRayTracing*/ true);
diff --git a/Engine/Source/Runtime/Renderer/Private/RayTracing/RayTracing.h b/Engine/Source/Runtime/Renderer/Private/RayTracing/RayTracing.h
--- a/Engine/Source/Runtime/Renderer/Private/RayTracing/RayTracing.h
+++ b/Engine/Source/Runtime/Renderer/Private/RayTracing/RayTracing.h
@@ -67,9 +67,7 @@ namespace RayTracing
 		FRDGBuilder& GraphBuilder,
 		FGatherInstancesTaskData& TaskData,
 		FRayTracingScene& RayTracingScene,
-		FRayTracingShaderBindingTable& RayTracingSBT,
-		FGlobalDynamicReadBuffer& InDynamicReadBuffer,
-		FSceneRenderingBulkObjectAllocator& InBulkAllocator);
+		FRayTracingShaderBindingTable& RayTracingSBT);
 
 	void WaitForDynamicBindings(FGatherInstancesTaskData& TaskData);
 
