Skip indirect diffuse composition pass when project has disabled static lighting, and only use SSAO+sky lighting as GI solution.

This is currently only skip for Substrate for A/B comparison purpose. When the function runs with EAmbientOcclusionMethod::Disabled, the composition is skipped at a later point within the loop.

diff --git a/Engine/Source/Runtime/Renderer/Private/IndirectLightRendering.cpp b/Engine/Source/Runtime/Renderer/Private/IndirectLightRendering.cpp
--- a/Engine/Source/Runtime/Renderer/Private/IndirectLightRendering.cpp
+++ b/Engine/Source/Runtime/Renderer/Private/IndirectLightRendering.cpp
@@ -102,11 +102,6 @@ static TAutoConsoleVariable<bool> CVarDiffuseIndirectForceCopyPass(
 	TEXT("Forces use of copy pass instead of dual source blend. (for debugging)"),
 	ECVF_RenderThreadSafe);
 
-static TAutoConsoleVariable<bool> CVarDiffuseIndirectAllowed(
-	TEXT("r.DiffuseIndirect.Allowed"), true,
-	TEXT("Allow indirect diffuse composition pass to run."),
-	ECVF_RenderThreadSafe);
-
 DECLARE_GPU_STAT_NAMED(ReflectionEnvironment, TEXT("Reflection Environment"));
 DECLARE_GPU_STAT_NAMED(RayTracingReflections, TEXT("Ray Tracing Reflections"));
 DECLARE_GPU_STAT(SkyLightDiffuse);
@@ -1021,8 +1016,7 @@ void FDeferredShadingSceneRenderer::RenderDiffuseIndirectAndAmbientOcclusion(
 
 	extern int32 GLumenVisualizeIndirectDiffuse;
 	if ((GLumenVisualizeIndirectDiffuse != 0) != bIsVisualizePass
-		|| ViewFamily.EngineShowFlags.VisualizeLightCulling 
-		|| !CVarDiffuseIndirectAllowed.GetValueOnRenderThread())
+		|| ViewFamily.EngineShowFlags.VisualizeLightCulling)
 	{
 		return;
 	}
@@ -1036,12 +1030,14 @@ void FDeferredShadingSceneRenderer::RenderDiffuseIndirectAndAmbientOcclusion(
 
 	const FRDGSystemTextures& SystemTextures = FRDGSystemTextures::Get(GraphBuilder);
 
-	auto ShouldSkipView = [bCompositeRegularLumenOnly, &AsyncLumenIndirectLightingOutputs](const FPerViewPipelineState& ViewPipelineState)
+	auto ShouldSkipView = [bCompositeRegularLumenOnly, &AsyncLumenIndirectLightingOutputs](const FPerViewPipelineState& ViewPipelineState, bool bApplyAmbientCubmap)
 	{
 		const bool bOnlyCompositeStepLeft = AsyncLumenIndirectLightingOutputs.StepsLeft == ELumenIndirectLightingSteps::Composite;
 		const bool bRegularLumenView = ViewPipelineState.DiffuseIndirectMethod == EDiffuseIndirectMethod::Lumen;
 
-		return bCompositeRegularLumenOnly != (bOnlyCompositeStepLeft && bRegularLumenView);
+		// If Static lighting is disabled and this pass only composes (SS)AO on top of the scene color, we can skip this step all together.
+		const bool bSkipSSAOComposition = Substrate::IsSubstrateEnabled() && ViewPipelineState.DiffuseIndirectMethod == EDiffuseIndirectMethod::Disabled && ViewPipelineState.AmbientOcclusionMethod == EAmbientOcclusionMethod::SSAO && !IsStaticLightingAllowed() && !bApplyAmbientCubmap;
+		return bCompositeRegularLumenOnly != (bOnlyCompositeStepLeft && bRegularLumenView) || bSkipSSAOComposition;
 	};
 
 	for (int32 ViewIndex = 0; ViewIndex < Views.Num(); ++ViewIndex)
@@ -1052,8 +1048,9 @@ void FDeferredShadingSceneRenderer::RenderDiffuseIndirectAndAmbientOcclusion(
 		RDG_EVENT_SCOPE_CONDITIONAL(GraphBuilder, Views.Num() > 1, "View%d", ViewIndex);
 
 		const FPerViewPipelineState& ViewPipelineState = GetViewPipelineState(View);
+		const bool bApplyAmbientCubmap = IsAmbientCubemapPassRequired(View) && !bIsVisualizePass;
 
-		if (ShouldSkipView(ViewPipelineState))
+		if (ShouldSkipView(ViewPipelineState, bApplyAmbientCubmap))
 		{
 			continue;
 		}
@@ -1667,7 +1664,7 @@ void FDeferredShadingSceneRenderer::RenderDiffuseIndirectAndAmbientOcclusion(
 		}
 
 		// Apply the ambient cubemaps
-		if (IsAmbientCubemapPassRequired(View) && !bIsVisualizePass)
+		if (bApplyAmbientCubmap)
 		{
 			auto ApplyAmbientCubemapComposite = [&](ESubstrateTileType TileType)
 			{			
