Fix for light function vis

#tests PC editor dx12
#rnx

diff --git a/Engine/Source/Runtime/Renderer/Private/LightFunctionAtlas.cpp b/Engine/Source/Runtime/Renderer/Private/LightFunctionAtlas.cpp
--- a/Engine/Source/Runtime/Renderer/Private/LightFunctionAtlas.cpp
+++ b/Engine/Source/Runtime/Renderer/Private/LightFunctionAtlas.cpp
@@ -825,13 +825,34 @@ FString FLightFunctionAtlas::GetOutOfSlotWarningMessage()
 	return FString();
 }
 
-FScreenPassTexture FLightFunctionAtlas::AddDebugVisualizationPasses(FRDGBuilder& GraphBuilder, const FViewInfo& View, FScreenPassTexture& ScreenPassSceneColor) const
+FScreenPassTexture FLightFunctionAtlas::AddDebugVisualizationPasses(FRDGBuilder& GraphBuilder, const FViewInfo& View, FScreenPassTexture& ScreenPassSceneColor, FScreenPassRenderTarget& OverrideOutput) const
 {
+	FScreenPassRenderTarget OutputRT = OverrideOutput;
+
 #if WITH_EDITOR
 
 	if (!IsLightFunctionAtlasEnabled())
 	{
-		return ScreenPassSceneColor;
+		return MoveTemp(OutputRT);
+	}
+
+	if (!OutputRT.IsValid())
+	{
+		OutputRT = FScreenPassRenderTarget::CreateFromInput(GraphBuilder, ScreenPassSceneColor, View.GetOverwriteLoadAction(), TEXT("VisualizeLightFunctionAtlas"));
+	}
+
+	if (ScreenPassSceneColor.Texture != OutputRT.GetRenderTargetBinding().GetTexture())
+	{
+		// Copy the scene color because below we do not write and copy all the view pixel. We only stamp the atlas and its tiles on top.
+		AddDrawTexturePass(
+			GraphBuilder,
+			View,
+			ScreenPassSceneColor.Texture,
+			OutputRT.GetRenderTargetBinding().GetTexture(),
+			FIntPoint::ZeroValue,
+			OutputRT.GetRenderTargetBinding().GetTexture()->Desc.Extent,
+			FIntPoint::ZeroValue,
+			OutputRT.GetRenderTargetBinding().GetTexture()->Desc.Extent);
 	}
 
 	const uint32 AtlasSlotResolution	= AtlasSetup.SlotResolution;
@@ -850,7 +871,7 @@ FScreenPassTexture FLightFunctionAtlas::AddDebugVisualizationPasses(FRDGBuilder&
 		GraphBuilder,
 		View,
 		GSystemTextures.GetMidGreyDummy(GraphBuilder),
-		ScreenPassSceneColor.Texture,
+		OutputRT.GetRenderTargetBinding().GetTexture(),
 		SrcPoint,
 		SrcSize,
 		DstPoint,
@@ -864,16 +885,15 @@ FScreenPassTexture FLightFunctionAtlas::AddDebugVisualizationPasses(FRDGBuilder&
 			GraphBuilder,
 			View,
 			RDGAtlasTexture2D,
-			ScreenPassSceneColor.Texture,
+			OutputRT.GetRenderTargetBinding().GetTexture(),
 			AtlasSlot.Min,
 			AtlasSlot.Max-AtlasSlot.Min,
 			SlotDstPoint,
 			SlotDstSize);
 	}
 
-
 	// Now debug print
-	AddDrawCanvasPass(GraphBuilder, {}, View, FScreenPassRenderTarget(ScreenPassSceneColor, ERenderTargetLoadAction::ELoad),
+	AddDrawCanvasPass(GraphBuilder, {}, View, FScreenPassRenderTarget(OutputRT, ERenderTargetLoadAction::ELoad),
 		[&View, this, DstPoint, DstSize, AtlasResolution, DisplayResolutionRatio](FCanvas& Canvas)
 	{
 		FString Text;
@@ -969,7 +989,7 @@ FScreenPassTexture FLightFunctionAtlas::AddDebugVisualizationPasses(FRDGBuilder&
 
 #endif // WITH_EDITOR
 
-	return MoveTemp(ScreenPassSceneColor);
+	return MoveTemp(OutputRT);
 }
 
 bool IsEnabled(const FViewInfo& InView, ELightFunctionAtlasSystem In)
diff --git a/Engine/Source/Runtime/Renderer/Private/LightFunctionAtlas.h b/Engine/Source/Runtime/Renderer/Private/LightFunctionAtlas.h
--- a/Engine/Source/Runtime/Renderer/Private/LightFunctionAtlas.h
+++ b/Engine/Source/Runtime/Renderer/Private/LightFunctionAtlas.h
@@ -16,6 +16,7 @@ class FLightSceneInfo;
 class FRDGTexture;
 
 struct FScreenPassTexture;
+struct FScreenPassRenderTarget;
 struct FSortedLightSetSceneInfo;
 struct IPooledRenderTarget;
 
@@ -150,7 +151,7 @@ struct FLightFunctionAtlas
 	void RenderLightFunctionAtlas(FRDGBuilder& GraphBuilder, TArray<FViewInfo>& Views);
 
 
-	FScreenPassTexture AddDebugVisualizationPasses(FRDGBuilder& GraphBuilder, const FViewInfo& View, FScreenPassTexture& ScreenPassSceneColor)  const;
+	FScreenPassTexture AddDebugVisualizationPasses(FRDGBuilder& GraphBuilder, const FViewInfo& View, FScreenPassTexture& ScreenPassSceneColor, FScreenPassRenderTarget& OverrideOutput)  const;
 
 	FLightFunctionAtlasGlobalParameters*						GetLightFunctionAtlasGlobalParametersStruct(FRDGBuilder& GraphBuilder, uint32 ViewIndex);
 	TRDGUniformBufferRef<FLightFunctionAtlasGlobalParameters>	GetLightFunctionAtlasGlobalParameters(FRDGBuilder& GraphBuilder, uint32 ViewIndex);
diff --git a/Engine/Source/Runtime/Renderer/Private/PostProcess/PostProcessing.cpp b/Engine/Source/Runtime/Renderer/Private/PostProcess/PostProcessing.cpp
--- a/Engine/Source/Runtime/Renderer/Private/PostProcess/PostProcessing.cpp
+++ b/Engine/Source/Runtime/Renderer/Private/PostProcess/PostProcessing.cpp
@@ -564,9 +564,9 @@ void AddPostProcessingPasses(
 	PassSequence.SetEnabled(EPass::VisualizeMegaLightsLightComplexity, View.Family && MegaLights::ShouldAddVisualizeLightComplexityPostProcessingPass(*View.Family, LumenFrameTemporaries));
 
 #if WITH_EDITOR
-	PassSequence.SetEnabled(EPass::VisualizeSkyAtmosphere, Scene&& View.Family && View.Family->EngineShowFlags.VisualizeSkyAtmosphere && ShouldRenderSkyAtmosphereDebugPasses(Scene, View.Family->EngineShowFlags));
-	PassSequence.SetEnabled(EPass::VisualizeSkyLightIlluminanceMeter, Scene&& Scene->SkyLight && View.Family && View.Family->EngineShowFlags.VisualizeSkyLightIlluminance);
-	PassSequence.SetEnabled(EPass::VisualizeLightFunctionAtlas, Scene && Scene->LightFunctionAtlasSceneData.GetLightFunctionAtlasEnabled() && View.Family && View.Family->EngineShowFlags.VisualizeLightFunctionAtlas);
+	PassSequence.SetEnabled(EPass::VisualizeSkyAtmosphere, GIsEditor && Scene && View.Family && View.Family->EngineShowFlags.VisualizeSkyAtmosphere && ShouldRenderSkyAtmosphereDebugPasses(Scene, View.Family->EngineShowFlags));
+	PassSequence.SetEnabled(EPass::VisualizeSkyLightIlluminanceMeter, GIsEditor && Scene && Scene->SkyLight && View.Family && View.Family->EngineShowFlags.VisualizeSkyLightIlluminance);
+	PassSequence.SetEnabled(EPass::VisualizeLightFunctionAtlas, GIsEditor && Scene && Scene->LightFunctionAtlasSceneData.GetLightFunctionAtlasEnabled() && View.Family && View.Family->EngineShowFlags.VisualizeLightFunctionAtlas);
 	PassSequence.SetEnabled(EPass::VisualizeLevelInstance, GIsEditor && EngineShowFlags.EditingLevelInstance && EngineShowFlags.VisualizeLevelInstanceEditing && !bVisualizeHDR);
 	PassSequence.SetEnabled(EPass::ChannelMask, View.Family->ChannelMaskParams.ColorChannelMask != EColorChannelMask::All || View.Family->ChannelMaskParams.bDrawAlphaBlendedCheckerboard);
 	PassSequence.SetEnabled(EPass::SelectionOutline, GIsEditor && EngineShowFlags.Selection && EngineShowFlags.SelectionOutline && !EngineShowFlags.Wireframe && !bVisualizeHDR);
@@ -1698,11 +1698,11 @@ void AddPostProcessingPasses(
 
 	if (PassSequence.IsEnabled(EPass::VisualizeLightFunctionAtlas))
 	{
-		FScreenPassRenderTarget OverrideOutput;
-		PassSequence.AcceptOverrideIfLastPass(EPass::VisualizeLightFunctionAtlas, OverrideOutput);
 		if (Scene->LightFunctionAtlasSceneData.GetLightFunctionAtlas())
 		{
-			SceneColor = Scene->LightFunctionAtlasSceneData.GetLightFunctionAtlas()->AddDebugVisualizationPasses(GraphBuilder, View, SceneColor);
+			FScreenPassRenderTarget OverrideOutput;
+			PassSequence.AcceptOverrideIfLastPass(EPass::VisualizeLightFunctionAtlas, OverrideOutput);
+			SceneColor = Scene->LightFunctionAtlasSceneData.GetLightFunctionAtlas()->AddDebugVisualizationPasses(GraphBuilder, View, SceneColor, OverrideOutput);
 		}
 	}
 
