Only precache dynamic raytracing geometry PSOs when raytracing is enabled

diff --git a/Engine/Source/Runtime/Renderer/Private/RayTracing/RayTracingDynamicGeometryUpdateManager.cpp b/Engine/Source/Runtime/Renderer/Private/RayTracing/RayTracingDynamicGeometryUpdateManager.cpp
--- a/Engine/Source/Runtime/Renderer/Private/RayTracing/RayTracingDynamicGeometryUpdateManager.cpp
+++ b/Engine/Source/Runtime/Renderer/Private/RayTracing/RayTracingDynamicGeometryUpdateManager.cpp
@@ -242,6 +242,13 @@ void FRayTracingDynamicGeometryPSOCollector::CollectPSOInitializers(
 	TArray<FPSOPrecacheData>& PSOInitializers
 )
 {
+	// check if raytracing is runtime enabled
+	static IConsoleVariable* CVarRayTracingEnable = IConsoleManager::Get().FindConsoleVariable(TEXT("r.RayTracing.Enable"));
+	if (CVarRayTracingEnable->GetInt() == 0)
+	{
+		return;
+	}
+
 	if (!VertexFactoryData.VertexFactoryType->SupportsRayTracingDynamicGeometry())
 	{
 		return;
