Mobile PSO precaching improvements to not precache PSOs when not used:
- add cvar to disable PSO precaching for water depth rendering
- don't precache distortion PSOs when refraction effects are disabled
- setup correct custom depth mode on PSOPrecacheParams for primitives and change default value to precache for all depth (doesn't change default behavior for platforms which don't modify the cva but setting 1 of r.PSOPrecache.CustomDepth will actually work now)

diff --git a/Engine/Source/Runtime/Engine/Private/Components/PrimitiveComponent.cpp b/Engine/Source/Runtime/Engine/Private/Components/PrimitiveComponent.cpp
--- a/Engine/Source/Runtime/Engine/Private/Components/PrimitiveComponent.cpp
+++ b/Engine/Source/Runtime/Engine/Private/Components/PrimitiveComponent.cpp
@@ -5171,9 +5171,7 @@ void UPrimitiveComponent::SetupPrecachePSOParams(FPSOPrecacheParams& Params)
 	Params.bUsesIndirectLightingCache = Params.bStaticLighting && IndirectLightingCacheQuality != ILCQ_Off && (!IsPrecomputedLightingValid() || GetLightmapType() == ELightmapType::ForceVolumetric);
 	Params.bAffectDynamicIndirectLighting = bAffectDynamicIndirectLighting;
 	Params.bCastShadow = CastShadow;
-	// Custom depth can be toggled at runtime with PSO precache call so assume it might be needed when depth pass is needed
-	// Ideally precache those with lower priority and don't wait on these (UE-174426)
-	Params.bRenderCustomDepth = bRenderInDepthPass;
+	Params.bRenderCustomDepth = bRenderCustomDepth;
 	Params.bCastShadowAsTwoSided = bCastShadowAsTwoSided;
 	Params.SetMobility(Mobility);	
 	Params.SetStencilWriteMask(FRendererStencilMaskEvaluation::ToStencilMask(CustomDepthStencilWriteMask));
diff --git a/Engine/Source/Runtime/Renderer/Private/CustomDepthRendering.cpp b/Engine/Source/Runtime/Renderer/Private/CustomDepthRendering.cpp
--- a/Engine/Source/Runtime/Renderer/Private/CustomDepthRendering.cpp
+++ b/Engine/Source/Runtime/Renderer/Private/CustomDepthRendering.cpp
@@ -34,7 +34,7 @@ static TAutoConsoleVariable<bool> CVarCustomDepthEnableFastClear(
 
 static TAutoConsoleVariable<int32> CVarPSOPrecacheCustomDepth(
 	TEXT("r.PSOPrecache.CustomDepth"),
-	1,
+	2,
 	TEXT("Also Precache PSOs with for custom depth pass.")  \
 	TEXT(" 0: No PSOs are compiled for this pass.\n") \
 	TEXT(" 1: PSOs are compiled for all primitives which explicitly request custom depth rendering (default).\n") \
diff --git a/Engine/Source/Runtime/Renderer/Private/DistortionRendering.cpp b/Engine/Source/Runtime/Renderer/Private/DistortionRendering.cpp
--- a/Engine/Source/Runtime/Renderer/Private/DistortionRendering.cpp
+++ b/Engine/Source/Runtime/Renderer/Private/DistortionRendering.cpp
@@ -96,18 +96,15 @@ uint32 GetRefractionOffsetQuality(EShaderPlatform InPlatform)
 	return CVarRefractionOffsetQualityPlatform.Get(InPlatform);
 }
 
-int32 FSceneRenderer::GetRefractionQuality(const FSceneViewFamily& ViewFamily)
+int32 GetGlobalRefractionQuality()
 {
 	static const auto ICVar = IConsoleManager::Get().FindTConsoleVariableDataInt(TEXT("r.RefractionQuality"));
+	return ICVar->GetValueOnAnyThread();
+}
 
-	int32 Value = 0;
-
-	if (ViewFamily.EngineShowFlags.Refraction)
-	{
-		Value = ICVar->GetValueOnRenderThread();
-	}
-
-	return Value;
+int32 FSceneRenderer::GetRefractionQuality(const FSceneViewFamily& ViewFamily)
+{
+	return ViewFamily.EngineShowFlags.Refraction ? GetGlobalRefractionQuality() : 0;
 }
 
 void SetupDistortionParams(FVector4f& DistortionParams, const FViewInfo& View)
@@ -1025,6 +1022,12 @@ void FDistortionMeshProcessor::CollectPSOInitializers(const FSceneTexturesConfig
 		return;
 	}
 
+	// If refraction effects are disabled then don't precache distortion PSOs
+	if (GetGlobalRefractionQuality() == 0)
+	{
+		return;
+	}
+
 	TMeshProcessorShaders<
 		FDistortionMeshVS,
 		FDistortionMeshPS> DistortionPassShaders;
diff --git a/Engine/Source/Runtime/Renderer/Private/WaterInfoTextureRendering.cpp b/Engine/Source/Runtime/Renderer/Private/WaterInfoTextureRendering.cpp
--- a/Engine/Source/Runtime/Renderer/Private/WaterInfoTextureRendering.cpp
+++ b/Engine/Source/Runtime/Renderer/Private/WaterInfoTextureRendering.cpp
@@ -32,6 +32,13 @@ static TAutoConsoleVariable<int32> CVarRenderCaptureNextWaterInfoDraws(
 	TEXT("Enable capturing of the water info texture for the next N draws"),
 	ECVF_RenderThreadSafe);
 
+static TAutoConsoleVariable<int32> CVarPSOPrecacheWaterInfoDepthPass(
+	TEXT("r.PSOPrecache.WaterInfoDepthPass"),
+	0,
+	TEXT("Precaches PSOs for WaterInfoDepthPass.\n"),
+	ECVF_ReadOnly
+);
+
 BEGIN_SHADER_PARAMETER_STRUCT(FWaterInfoTexturePassParameters, )
 	SHADER_PARAMETER_STRUCT_INCLUDE(FViewShaderParameters, View)
 	SHADER_PARAMETER_STRUCT_INCLUDE(FInstanceCullingDrawParams, InstanceCullingDrawParams)
@@ -546,6 +553,13 @@ bool FWaterInfoTextureDepthPassMeshProcessor::Process(
 
 void FWaterInfoTextureDepthPassMeshProcessor::CollectPSOInitializers(const FSceneTexturesConfig& SceneTexturesConfig, const FMaterial& Material, const FPSOPrecacheVertexFactoryData& VertexFactoryData, const FPSOPrecacheParams& PreCacheParams, TArray<FPSOPrecacheData>& PSOInitializers)
 {
+	// Early out if disabled
+	int32 CollectPSOs = CVarPSOPrecacheWaterInfoDepthPass.GetValueOnRenderThread();
+	if (CollectPSOs == 0)
+	{
+		return;
+	}
+
 	// We need to support all materials that could possibly be rendered in a depth-only pass. Unfortunately there doesn't seem to be a way to filter for bUseForWaterInfoTextureDepth at this point.
 	if (Material.GetMaterialDomain() == MD_Surface && !IsTranslucentBlendMode(Material))
 	{
