[HWRT] Gather dynamic ray tracing instances using multiple threads.

diff --git a/Engine/Source/Runtime/Renderer/Private/RayTracing/RayTracing.cpp b/Engine/Source/Runtime/Renderer/Private/RayTracing/RayTracing.cpp
--- a/Engine/Source/Runtime/Renderer/Private/RayTracing/RayTracing.cpp
+++ b/Engine/Source/Runtime/Renderer/Private/RayTracing/RayTracing.cpp
@@ -29,6 +29,24 @@ static FAutoConsoleVariableRef CVarRayTracingDynamicInstances(
 	ECVF_RenderThreadSafe
 );
 
+static int32 GParallelGatherDynamicRayTracingInstancesNumTasks = 4;
+static FAutoConsoleVariableRef CVarParallelGatherDynamicRayTracingInstancesNumTasks(
+	TEXT("r.RayTracing.DynamicInstances.Parallel.NumTasks"),
+	GParallelGatherDynamicRayTracingInstancesNumTasks,
+	TEXT("Controls the number of gather dynamic ray tracing instances tasks to run asynchronously."),
+	ECVF_RenderThreadSafe
+);
+
+inline uint32 GetNumDynamicRayTracingInstancesContexts()
+{
+	if (!IsParallelGatherDynamicRayTracingInstancesEnabled())
+	{
+		return 1;
+	}
+
+	return FMath::Clamp<int32>(GParallelGatherDynamicRayTracingInstancesNumTasks, 1, LowLevelTasks::FScheduler::Get().GetNumWorkers());
+}
+
 static int32 GRayTracingParallelMeshBatchSetup = 1;
 static FAutoConsoleVariableRef CRayTracingParallelMeshBatchSetup(
 	TEXT("r.RayTracing.ParallelMeshBatchSetup"),
@@ -283,12 +301,41 @@ namespace RayTracing
 		uint32 ViewMask : 8;
 	};
 
+	class FDynamicPrimitiveIndexQueue
+	{
+	public:
+		FDynamicPrimitiveIndexQueue(TArray<FDynamicPrimitiveIndex, SceneRenderingAllocator>&& InList)
+			: List(InList)
+		{
+		}
+
+		bool Pop(FDynamicPrimitiveIndex& PrimitiveIndex)
+		{
+			const int32 Index = NextIndex.fetch_add(1, std::memory_order_relaxed);
+			if (Index < List.Num())
+			{
+				PrimitiveIndex = List[Index];
+				return true;
+			}
+			return false;
+		}
+
+		int32 Num()
+		{
+			return List.Num();
+		}
+
+	private:
+		TArray<FDynamicPrimitiveIndex, SceneRenderingAllocator> List;
+		std::atomic_int32_t NextIndex = { 0 };
+	};
+
 	class FDynamicRayTracingInstancesContext
 	{
 	public:
 		FDynamicRayTracingInstancesContext(FScene& InScene, TArrayView<FViewInfo*> InViews, TConstArrayView<const FSceneOptions*> InViewSceneOptions, FSceneRenderingBulkObjectAllocator& InBulkAllocator);
 
-		void GatherDynamicRayTracingInstances(TConstArrayView<FDynamicPrimitiveIndex> InDynamicPrimitives);
+		void GatherDynamicRayTracingInstances(FDynamicPrimitiveIndexQueue* PrimitiveIndexQueue);
 
 		void GatherDynamicRayTracingInstances_RenderThread();
 
@@ -298,6 +345,8 @@ namespace RayTracing
 
 		void CollectRDGResources(FRDGBuilder& RDGBuilder);
 
+		void SetRenderThreadDynamicPrimitives(TArray<FDynamicPrimitiveIndex, SceneRenderingAllocator>&& InDynamicPrimitives);
+
 	private:
 
 		void GatherDynamicRayTracingInstances_Internal(const FDynamicPrimitiveIndex& PrimitiveIndex);
@@ -316,9 +365,9 @@ namespace RayTracing
 
 		FRayTracingInstanceCollector RayTracingInstanceCollector;
 
-		TArray<FDynamicRayTracingPrimitive> DynamicRayTracingPrimitives;
+		TArray<FDynamicRayTracingPrimitive, SceneRenderingAllocator> DynamicRayTracingPrimitives;
 
-		TArray<FDynamicPrimitiveIndex> RenderThreadDynamicPrimitives;
+		TArray<FDynamicPrimitiveIndex, SceneRenderingAllocator> RenderThreadDynamicPrimitives;
 	};
 
 	struct FGatherInstancesViewTaskData
@@ -405,7 +454,7 @@ namespace RayTracing
 		UE::Tasks::FPipe AddInstancesPipe{ UE_SOURCE_LOCATION };
 		UE::Tasks::FPipe FinalizeGatherRelevantPrimitivesPipe{ UE_SOURCE_LOCATION };
 
-		FDynamicRayTracingInstancesContext* DynamicRayTracingInstancesContext;
+		TArray<FDynamicRayTracingInstancesContext*> DynamicRayTracingInstancesContexts;
 
 		// This task must complete before accessing DynamicRayTracingInstancesContext.
 		UE::Tasks::FTask GatherDynamicRayTracingInstancesTask;
@@ -516,7 +565,7 @@ namespace RayTracing
 				}
 			}
 		}
-	}	
+	}
 
 	class FRaytracingShaderBindingLayout : public FShaderBindingLayoutContainer
 	{
@@ -569,7 +618,7 @@ namespace RayTracing
 		// Setup the static uniform buffers used by the RTPSO if enabled
 		const FShaderBindingLayout* ShaderBindingLayout = GetShaderBindingLayout(View.GetShaderPlatform());
 		if (ShaderBindingLayout)
-		{			
+		{
 			FUniformBufferStaticBindings StaticUniformBuffers(&ShaderBindingLayout->RHILayout);
 			StaticUniformBuffers.AddUniformBuffer(View.ViewUniformBuffer.GetReference());
 			StaticUniformBuffers.AddUniformBuffer(SceneUniformBuffer);
@@ -695,10 +744,7 @@ namespace RayTracing
 	void FDynamicRayTracingInstancesContext::GatherDynamicRayTracingInstances_Internal(const FDynamicPrimitiveIndex& PrimitiveIndex)
 	{
 		FPrimitiveSceneProxy* SceneProxy = Scene.PrimitiveSceneProxies[PrimitiveIndex.Index];
-		FPrimitiveSceneInfo* SceneInfo = Scene.Primitives[PrimitiveIndex.Index];
-		const FPersistentPrimitiveIndex PersistentPrimitiveIndex = SceneInfo->GetPersistentIndex();
 
-		//RayTracingInstanceCollector.SetPrimitive(SceneProxy, PrimitiveIndex);
 		RayTracingInstanceCollector.SetPrimitive(SceneProxy, FHitProxyId::InvisibleHitProxyId);
 		RayTracingInstanceCollector.SetVisibilityMap(PrimitiveIndex.ViewMask);
 
@@ -718,21 +764,19 @@ namespace RayTracing
 		DynamicRayTracingPrimitives.Add(MoveTemp(Tmp));
 	}
 
-	void FDynamicRayTracingInstancesContext::GatherDynamicRayTracingInstances(TConstArrayView<FDynamicPrimitiveIndex> InDynamicPrimitives)
+	void FDynamicRayTracingInstancesContext::GatherDynamicRayTracingInstances(FDynamicPrimitiveIndexQueue* PrimitiveIndexQueue)
 	{
 		TRACE_CPUPROFILER_EVENT_SCOPE(RayTracing_GatherDynamicRayTracingInstances);
 
-		DynamicRayTracingPrimitives.Reserve(DynamicRayTracingPrimitives.Num() + InDynamicPrimitives.Num());
+		check(IsParallelGatherDynamicRayTracingInstancesEnabled());
 
-		if (!IsParallelGatherDynamicRayTracingInstancesEnabled())
-		{
-			RenderThreadDynamicPrimitives = InDynamicPrimitives;
-			return;
-		}
+		// Reserve for worst case even though queue might be processed by multiple tasks
+		DynamicRayTracingPrimitives.Reserve(DynamicRayTracingPrimitives.Num() + PrimitiveIndexQueue->Num());
 
 		// TODO: Could filter primitives whose proxy supports ParallelGDRTI during GatherRayTracingRelevantPrimitives_Parallel
 
-		for (const FDynamicPrimitiveIndex& PrimitiveIndex : InDynamicPrimitives)
+		FDynamicPrimitiveIndex PrimitiveIndex;
+		while (PrimitiveIndexQueue->Pop(PrimitiveIndex))
 		{
 			FPrimitiveSceneProxy* SceneProxy = Scene.PrimitiveSceneProxies[PrimitiveIndex.Index];
 
@@ -761,6 +805,13 @@ namespace RayTracing
 		RenderThreadDynamicPrimitives.Empty();
 	}
 
+	void FDynamicRayTracingInstancesContext::SetRenderThreadDynamicPrimitives(TArray<FDynamicPrimitiveIndex, SceneRenderingAllocator>&& InDynamicPrimitives)
+	{
+		check(RenderThreadDynamicPrimitives.IsEmpty());
+
+		RenderThreadDynamicPrimitives = MoveTemp(InDynamicPrimitives);
+	}
+
 	void FDynamicRayTracingInstancesContext::Finish(FRHICommandListImmediate& InRHICmdList)
 	{
 		TRACE_CPUPROFILER_EVENT_SCOPE(RayTracing_GatherDynamicRayTracingInstances_Finish);
@@ -1817,7 +1868,10 @@ namespace RayTracing
 	{
 		TRACE_CPUPROFILER_EVENT_SCOPE(RayTracingScene_AddDynamicInstances);
 
-		TaskData.DynamicRayTracingInstancesContext->AddInstancesToScene(RayTracingScene, RayTracingSBT, SharedBufferGenerationID);
+		for (FDynamicRayTracingInstancesContext* DynamicRayTracingInstancesContext : TaskData.DynamicRayTracingInstancesContexts)
+		{
+			DynamicRayTracingInstancesContext->AddInstancesToScene(RayTracingScene, RayTracingSBT, SharedBufferGenerationID);
+		}
 	}
 
 	void BeginGatherInstances(FGatherInstancesTaskData& TaskData, UE::Tasks::FTask FrustumCullTask)
@@ -1916,11 +1970,16 @@ namespace RayTracing
 
 		// Dynamic instance gathering
 		{
-			TaskData.DynamicRayTracingInstancesContext = TaskData.Allocator.Create<FDynamicRayTracingInstancesContext>(TaskData.Scene, TaskData.Views, TaskData.ViewSceneOptions, TaskData.Allocator);
+			const int32 NumContexts = GetNumDynamicRayTracingInstancesContexts();
+			TaskData.DynamicRayTracingInstancesContexts.Reserve(NumContexts);
+
+			for (int32 Index = 0; Index < NumContexts; ++Index)
+			{
+				TaskData.DynamicRayTracingInstancesContexts.Add(TaskData.Allocator.Create<FDynamicRayTracingInstancesContext>(TaskData.Scene, TaskData.Views, TaskData.ViewSceneOptions, TaskData.Allocator));
+			}
 
 			TaskData.GatherDynamicRayTracingInstancesPrerequisites.AddPrerequisites(TaskData.Scene.GetGPUSkinCacheTask());
 
-			// TODO: Could gather dynamic ray tracing instances using multiple tasks / FDynamicRayTracingInstancesContext
 			TaskData.GatherDynamicRayTracingInstancesTask = UE::Tasks::Launch(UE_SOURCE_LOCATION, [&TaskData]
 				{
 					FTaskTagScope Scope(ETaskTag::EParallelRenderingThread);
@@ -1930,7 +1989,7 @@ namespace RayTracing
 					//	and then here we could just loop over all primitives to collect the relevant dynamic ones.
 					//	(loop over all primitives vs loop over dynamic primitives for each view and combine using map + array)
 
-					TArray<FDynamicPrimitiveIndex> DynamicPrimitivesIndices;
+					TArray<FDynamicPrimitiveIndex, SceneRenderingAllocator> DynamicPrimitivesIndices;
 					TMap<int32, int32> PrimitiveIndexMap;
 
 					for (int32 ViewIndex = 0; ViewIndex < TaskData.ViewTaskDatas.Num(); ++ViewIndex)
@@ -1956,7 +2015,25 @@ namespace RayTracing
 						}
 					}
 
-					TaskData.DynamicRayTracingInstancesContext->GatherDynamicRayTracingInstances(DynamicPrimitivesIndices);
+					if (IsParallelGatherDynamicRayTracingInstancesEnabled())
+					{
+						FDynamicPrimitiveIndexQueue* Queue = TaskData.Allocator.Create<FDynamicPrimitiveIndexQueue>(MoveTemp(DynamicPrimitivesIndices));
+
+						for (FDynamicRayTracingInstancesContext* DynamicRayTracingInstancesContext : TaskData.DynamicRayTracingInstancesContexts)
+						{
+							UE::Tasks::FTask ContextTask = UE::Tasks::Launch(UE_SOURCE_LOCATION, [DynamicRayTracingInstancesContext, Queue]
+								{
+									FTaskTagScope Scope(ETaskTag::EParallelRenderingThread);
+									DynamicRayTracingInstancesContext->GatherDynamicRayTracingInstances(Queue);
+								});
+
+							UE::Tasks::AddNested(ContextTask);
+						}
+					}
+					else
+					{
+						TaskData.DynamicRayTracingInstancesContexts.Last()->SetRenderThreadDynamicPrimitives(MoveTemp(DynamicPrimitivesIndices));
+					}
 				}, TaskData.GatherDynamicRayTracingInstancesPrerequisites, TaskPriority);
 		}
 
@@ -2006,8 +2083,12 @@ namespace RayTracing
 
 		{
 			TaskData.GatherDynamicRayTracingInstancesTask.Wait();
-			TaskData.DynamicRayTracingInstancesContext->GatherDynamicRayTracingInstances_RenderThread();
-			TaskData.DynamicRayTracingInstancesContext->Finish(GraphBuilder.RHICmdList);
+
+			for (FDynamicRayTracingInstancesContext* DynamicRayTracingInstancesContext : TaskData.DynamicRayTracingInstancesContexts)
+			{
+				DynamicRayTracingInstancesContext->GatherDynamicRayTracingInstances_RenderThread();
+				DynamicRayTracingInstancesContext->Finish(GraphBuilder.RHICmdList);
+			}
 
 			FSceneRenderer::DynamicReadBufferForRayTracing.Commit(GraphBuilder.RHICmdList);
 
@@ -2016,7 +2097,10 @@ namespace RayTracing
 				ViewTaskData.Scene.GPUScene.UploadDynamicPrimitiveShaderDataForView(GraphBuilder, ViewTaskData.View, /*bRayTracing*/ true);
 			}
 
-			TaskData.DynamicRayTracingInstancesContext->CollectRDGResources(GraphBuilder);
+			for (FDynamicRayTracingInstancesContext* DynamicRayTracingInstancesContext : TaskData.DynamicRayTracingInstancesContexts)
+			{
+				DynamicRayTracingInstancesContext->CollectRDGResources(GraphBuilder);
+			}
 		}
 
 		bool bAnyViewLightingChannelsUsingAHS = false;
