[DOF] Make the slight-out-of-focus (SOF) history rejection boost optional, as this introduces noise in areas with high variance (e.g. bright specular).
This variable increases contrast in SOF areas to reduce denoising strength, based on CoC radius. This enhances sharpness, but can also introduce noise.
This CL turns off this boost by default, and can be turned back on with r.DOF.TAA.SOFHistoryRejectionBoost 0.125

diff --git a/Engine/Shaders/Private/DiaphragmDOF/DOFRecombine.usf b/Engine/Shaders/Private/DiaphragmDOF/DOFRecombine.usf
--- a/Engine/Shaders/Private/DiaphragmDOF/DOFRecombine.usf
+++ b/Engine/Shaders/Private/DiaphragmDOF/DOFRecombine.usf
@@ -141,6 +141,7 @@ float4 SeparateTranslucencyBilinearUVMinMax;
 uint SeparateTranslucencyUpscaling;
 float EncodedCocRadiusToRecombineCocRadius;
 float MaxRecombineAbsCocRadius;
+float SOFHistoryRejectionBoost;
 
 Texture2D BokehLUT;
 
@@ -963,7 +964,7 @@ void RecombineMainCS(
 			#if 1
 			{
 				float4 HD = Min;
-				float Small = 0.125 * (1.0 - saturate(CenterSample.CocRadius * CenterSample.CocRadius * rcp(64.0)));
+				float Small = SOFHistoryRejectionBoost * (1.0 - saturate(CenterSample.CocRadius * CenterSample.CocRadius * rcp(64.0)));
 				Max += HD * Small;
 				Min -= HD * Small;
 
diff --git a/Engine/Source/Runtime/Renderer/Private/PostProcess/DiaphragmDOF.cpp b/Engine/Source/Runtime/Renderer/Private/PostProcess/DiaphragmDOF.cpp
--- a/Engine/Source/Runtime/Renderer/Private/PostProcess/DiaphragmDOF.cpp
+++ b/Engine/Source/Runtime/Renderer/Private/PostProcess/DiaphragmDOF.cpp
@@ -173,7 +173,14 @@ TAutoConsoleVariable<float> CVarCoCBilateralFilterStrength(
 	TEXT(" (Default: 0)."),
 	ECVF_RenderThreadSafe);
 
-
+TAutoConsoleVariable<float> CVarSOFHistoryRejectionBoost(
+	TEXT("r.DOF.TAA.SOFHistoryRejectionBoost"),
+	0.0,
+	TEXT("In the recombine pass, reduce the strength of denoising for slight-out-of-focus areas.\n")
+	TEXT("Reasonable values are 0 to 0.125.\n")
+	TEXT("If 0, this is disabled.\n")
+	TEXT(" (Default: 0)."),
+	ECVF_RenderThreadSafe);
 
 } // namespace
 
@@ -1422,6 +1429,7 @@ class FDiaphragmDOFRecombineCS : public FDiaphragmDOFShader
 		SHADER_PARAMETER(int32, SeparateTranslucencyUpscaling)
 		SHADER_PARAMETER(float, EncodedCocRadiusToRecombineCocRadius)
 		SHADER_PARAMETER(float, MaxRecombineAbsCocRadius)
+		SHADER_PARAMETER(float, SOFHistoryRejectionBoost)
 		
 		SHADER_PARAMETER_RDG_TEXTURE(Texture2D, BokehLUT)
 
@@ -3023,6 +3031,7 @@ bool DiaphragmDOF::AddPasses(
 			(GatheringViewSize.Y - 0.5f) / float(RefBufferSize.Y));
 		PassParameters->EncodedCocRadiusToRecombineCocRadius = float(GatheringViewSize.X) / EncodedCocRadiusBasis; 
 		PassParameters->MaxRecombineAbsCocRadius = float(kMaxSlightOutOfFocusCocRadius) * PassParameters->EncodedCocRadiusToRecombineCocRadius;
+		PassParameters->SOFHistoryRejectionBoost = CVarSOFHistoryRejectionBoost.GetValueOnRenderThread();
 
 		PassParameters->SeparateTranslucencyBilinearUVMinMax.X = (SeparateTranslucencyViewport.Rect.Min.X + 0.5f) / float(SeparateTranslucencyViewport.Extent.X);
 		PassParameters->SeparateTranslucencyBilinearUVMinMax.Y = (SeparateTranslucencyViewport.Rect.Min.Y + 0.5f) / float(SeparateTranslucencyViewport.Extent.Y);
