VolumetricFog - fixed race condition with initialising/writting volumetric fog data on viewstate on SetupVolumetricFog() winthin LaunchVisibilityTasks threads in flight while a thread is also initialising View buffers also reading that data.
Now, The volumetric fog view state is cleared only if needed before threads are launched. We have:
FDeferredShadingSceneRenderer::Render
- OnRenderBegin(...) setting up volumetric fog UB not touching anything on ViewState
- BeginInitViews(...) initialising view resource, safely reading ViewState andinitialising it if needed
- VFog rendering, setting up the view state for next frame.

#tests PC editor dx12
#rnx

diff --git a/Engine/Source/Runtime/Renderer/Private/RendererScene.cpp b/Engine/Source/Runtime/Renderer/Private/RendererScene.cpp
--- a/Engine/Source/Runtime/Renderer/Private/RendererScene.cpp
+++ b/Engine/Source/Runtime/Renderer/Private/RendererScene.cpp
@@ -295,6 +295,18 @@ TRefCountPtr<IPooledRenderTarget> FPersistentSkyAtmosphereData::GetCurrentCamera
 	return CameraAerialPerspectiveVolumesRayOnly[CameraAerialPerspectiveVolumeIndex];
 }
 
+void FSceneViewState::ResetVolumetricFogState()
+{
+	LightScatteringHistoryPreExposure = 1.0f;
+	PrevLightScatteringViewGridUVToViewRectVolumeUV = FVector2f::One();
+	VolumetricFogPrevViewGridRectUVToResourceUV = FVector2f::One();
+	VolumetricFogPrevUVMax = FVector2f::One();
+	VolumetricFogPrevUVMaxForTemporalBlend = FVector2f::One();
+	VolumetricFogPrevResourceGridSize = FIntVector(1, 1, 1);
+	LightScatteringHistory = nullptr;
+	PrevLightScatteringConservativeDepthTexture = nullptr;
+}
+
 /** Default constructor. */
 FSceneViewState::FSceneViewState(ERHIFeatureLevel::Type FeatureLevel, FSceneViewState* ShareOriginTarget)
 	: OcclusionQueryPool(RHICreateRenderQueryPool(RQT_Occlusion))
@@ -382,6 +394,8 @@ FSceneViewState::FSceneViewState(ERHIFeatureLevel::Type FeatureLevel, FSceneView
 	}
 
 	Occlusion.LastOcclusionQueryArray.SetNumZeroed(FOcclusionQueryHelpers::MaxBufferedOcclusionFrames);
+
+	ResetVolumetricFogState();
 }
 
 FSceneViewState::~FSceneViewState()
diff --git a/Engine/Source/Runtime/Renderer/Private/ScenePrivate.h b/Engine/Source/Runtime/Renderer/Private/ScenePrivate.h
--- a/Engine/Source/Runtime/Renderer/Private/ScenePrivate.h
+++ b/Engine/Source/Runtime/Renderer/Private/ScenePrivate.h
@@ -1186,6 +1186,8 @@ public:
 		DistanceFieldTemporalSampleIndex = 0;
 		PreExposure = 1.f;
 
+		ResetVolumetricFogState();
+
 		ReleaseRHI();
 	}
 
@@ -1207,6 +1209,8 @@ public:
 		return DistanceFieldTemporalSampleIndex;
 	}
 
+	void ResetVolumetricFogState();
+
 	/** Default constructor. */
 	FSceneViewState(ERHIFeatureLevel::Type FeatureLevel, FSceneViewState* ShareOriginTarget);
 
diff --git a/Engine/Source/Runtime/Renderer/Private/SceneRendering.cpp b/Engine/Source/Runtime/Renderer/Private/SceneRendering.cpp
--- a/Engine/Source/Runtime/Renderer/Private/SceneRendering.cpp
+++ b/Engine/Source/Runtime/Renderer/Private/SceneRendering.cpp
@@ -4061,6 +4061,9 @@ IVisibilityTaskData* FSceneRenderer::OnRenderBegin(FRDGBuilder& GraphBuilder, co
 			LightFunctionAtlas::OnRenderBegin(LightFunctionAtlas, *Scene, Views, ViewFamily);
 		}
 
+		// Notify the volumetric fog system that we start a new frame, deal with its state, etc.
+		VolumetricFogOnRenderBegin(*Scene, Views);
+
 		VisibilityTaskData = LaunchVisibilityTasks(GraphBuilder.RHICmdList, *this, StaticMeshUpdateTask);
 
 		if (GraphBuilder.IsParallelSetupEnabled())
diff --git a/Engine/Source/Runtime/Renderer/Private/VolumetricFog.cpp b/Engine/Source/Runtime/Renderer/Private/VolumetricFog.cpp
--- a/Engine/Source/Runtime/Renderer/Private/VolumetricFog.cpp
+++ b/Engine/Source/Runtime/Renderer/Private/VolumetricFog.cpp
@@ -658,6 +658,20 @@ bool LightNeedsSeparateInjectionIntoVolumetricFogForOpaqueShadow(const FViewInfo
 	return false;
 }
 
+void VolumetricFogOnRenderBegin(FScene& Scene, TArray<FViewInfo>& Views)
+{
+	for (int32 ViewIndex = 0; ViewIndex < Views.Num(); ViewIndex++)
+	{
+		FViewInfo& View = Views[ViewIndex];
+
+		// Similar test to what is in SetupVolumetricFog()
+		if (View.ViewState && View.Family && !ShouldRenderVolumetricFog(&Scene, *View.Family))
+		{
+			View.ViewState->ResetVolumetricFogState();
+		}
+	}
+}
+
 FIntPoint CalculateVolumetricFogBoundsForLight(const FSphere& LightBounds, const FViewInfo& View, FIntVector VolumetricFogGridSize, FVector GridZParams)
 {
 	FIntPoint VolumeZBounds;
@@ -1362,26 +1376,6 @@ void FSceneRenderer::SetupVolumetricFog()
 			View.VolumetricFogResources.VolumetricFogGlobalData = TUniformBufferRef<FVolumetricFogGlobalData>::CreateUniformBufferImmediate(GlobalData, UniformBuffer_SingleFrame);
 		}
 	}
-	else
-	{
-		for (int32 ViewIndex = 0; ViewIndex < Views.Num(); ViewIndex++)
-		{
-			FViewInfo& View = Views[ViewIndex];
-
-			if (View.ViewState)
-			{
-				View.ViewState->LightScatteringHistory = NULL;
-				View.ViewState->LightScatteringHistoryPreExposure = 1.0f;
-				View.ViewState->PrevLightScatteringViewGridUVToViewRectVolumeUV = FVector2f::One();
-				View.ViewState->VolumetricFogPrevViewGridRectUVToResourceUV = FVector2f::One();
-				View.ViewState->VolumetricFogPrevUVMax = FVector2f::One();
-				View.ViewState->VolumetricFogPrevUVMaxForTemporalBlend = FVector2f::One();
-
-				int32 VolumetricFogResourceGridPixelSize;
-				View.ViewState->VolumetricFogPrevResourceGridSize = GetVolumetricFogResourceGridSize(View, VolumetricFogResourceGridPixelSize);
-			}
-		}
-	}
 }
 
 void FSceneRenderer::ComputeVolumetricFog(FRDGBuilder& GraphBuilder,
diff --git a/Engine/Source/Runtime/Renderer/Private/VolumetricFog.h b/Engine/Source/Runtime/Renderer/Private/VolumetricFog.h
--- a/Engine/Source/Runtime/Renderer/Private/VolumetricFog.h
+++ b/Engine/Source/Runtime/Renderer/Private/VolumetricFog.h
@@ -29,3 +29,5 @@ extern bool ShouldRenderVolumetricFog(const FScene* Scene, const FSceneViewFamil
 extern const FProjectedShadowInfo* GetShadowForInjectionIntoVolumetricFog(const FVisibleLightInfo& VisibleLightInfo);
 
 extern bool LightNeedsSeparateInjectionIntoVolumetricFogForOpaqueShadow(const FViewInfo& View, const FLightSceneInfo* LightSceneInfo, const FVisibleLightInfo& VisibleLightInfo, const FScene& InScene);
+
+extern void VolumetricFogOnRenderBegin(FScene& Scene, TArray<FViewInfo>& Views);
