Use correct cvar thread function for PSO precaching - not on render thread

diff --git a/Engine/Source/Runtime/Renderer/Private/WaterInfoTextureRendering.cpp b/Engine/Source/Runtime/Renderer/Private/WaterInfoTextureRendering.cpp
--- a/Engine/Source/Runtime/Renderer/Private/WaterInfoTextureRendering.cpp
+++ b/Engine/Source/Runtime/Renderer/Private/WaterInfoTextureRendering.cpp
@@ -554,7 +554,7 @@ bool FWaterInfoTextureDepthPassMeshProcessor::Process(
 void FWaterInfoTextureDepthPassMeshProcessor::CollectPSOInitializers(const FSceneTexturesConfig& SceneTexturesConfig, const FMaterial& Material, const FPSOPrecacheVertexFactoryData& VertexFactoryData, const FPSOPrecacheParams& PreCacheParams, TArray<FPSOPrecacheData>& PSOInitializers)
 {
 	// Early out if disabled
-	int32 CollectPSOs = CVarPSOPrecacheWaterInfoDepthPass.GetValueOnRenderThread();
+	int32 CollectPSOs = CVarPSOPrecacheWaterInfoDepthPass.GetValueOnAnyThread();
 	if (CollectPSOs == 0)
 	{
 		return;
