Fix translucent materials not showing up in shader complexity view mode.

- Since we now only compile the skylight version of BasePassPS for all translucent materials (to save shaders), the non-skylight version of the shader is missing.
- Debug viewmode rendering was always asking the material for the non-skylight version with the assumption that this shader would always exist.
- Since the shader no longer exists I needed to special case the logic.  If we are a translucent material we need to get the skylight version of teh shader to get instruction count.

#rnx

diff --git a/Engine/Source/Runtime/Renderer/Private/DebugViewModeRendering.cpp b/Engine/Source/Runtime/Renderer/Private/DebugViewModeRendering.cpp
--- a/Engine/Source/Runtime/Renderer/Private/DebugViewModeRendering.cpp
+++ b/Engine/Source/Runtime/Renderer/Private/DebugViewModeRendering.cpp
@@ -516,7 +516,15 @@ void FDebugViewModeMeshProcessor::UpdateInstructionCount(FDebugViewModeShaderEle
 
 			FMaterialShaderTypes ShaderTypes;
 			ShaderTypes.AddShaderType<TBasePassVS<TUniformLightMapPolicy<LMP_NO_LIGHTMAP>>>();
-			ShaderTypes.AddShaderType<TBasePassPS<TUniformLightMapPolicy<LMP_NO_LIGHTMAP>, false, GBL_Default>>();
+
+			if (InBatchMaterial->GetBlendMode() == EBlendMode::BLEND_Translucent)
+			{
+				ShaderTypes.AddShaderType<TBasePassPS<TUniformLightMapPolicy<LMP_NO_LIGHTMAP>, true, GBL_Default>>();
+			}
+			else
+			{
+				ShaderTypes.AddShaderType<TBasePassPS<TUniformLightMapPolicy<LMP_NO_LIGHTMAP>, false, GBL_Default>>();
+			}
 
 			FMaterialShaders Shaders;
 			if (InBatchMaterial->TryGetShaders(ShaderTypes, InVertexFactoryType, Shaders))
