Fix usage of LowResCubeViewPtr in RenderCubeFaces_SkyCloud

diff --git a/Engine/Source/Runtime/Renderer/Private/ReflectionEnvironmentRealTimeCapture.cpp b/Engine/Source/Runtime/Renderer/Private/ReflectionEnvironmentRealTimeCapture.cpp
--- a/Engine/Source/Runtime/Renderer/Private/ReflectionEnvironmentRealTimeCapture.cpp
+++ b/Engine/Source/Runtime/Renderer/Private/ReflectionEnvironmentRealTimeCapture.cpp
@@ -645,7 +645,7 @@ void FScene::AllocateAndCaptureFrameSkyEnvMap(
 						FMatrix LowResCubeProjectionMatrix;
 						FViewInfo* LowResCubeViewPtr = nullptr;
 						CreateMainViewSnapshotForRealTimeCapture(LowResCubeViewPtr, LowResCubeProjectionMatrix, CubeWidthLow);
-						FViewInfo& LowResCubeView = *CubeViewPtr;
+						FViewInfo& LowResCubeView = *LowResCubeViewPtr;
 
 						FViewMatrices CubeViewMatrices;
 						SetupCommonViewUniformBufferParameters(LowResCubeView, LowResCubeProjectionMatrix, CubeViewMatrices, CubeWidthLow, CubeFace);
