[HWRT] Fix missing calls to FRayTracingGeometryManager::PreRender() and FRayTracingGeometryManager::ProcessBuildRequests(...).
- those functions where not being called when the first renderer didn't output FinalSceneColor.

diff --git a/Engine/Source/Runtime/Renderer/Private/DeferredShadingRenderer.cpp b/Engine/Source/Runtime/Renderer/Private/DeferredShadingRenderer.cpp
--- a/Engine/Source/Runtime/Renderer/Private/DeferredShadingRenderer.cpp
+++ b/Engine/Source/Runtime/Renderer/Private/DeferredShadingRenderer.cpp
@@ -1722,14 +1722,12 @@ void FDeferredShadingSceneRenderer::Render(FRDGBuilder& GraphBuilder, const FSce
 	GPU_MESSAGE_SCOPE(GraphBuilder);
 
 #if RHI_RAYTRACING
-	if (SceneUpdateInputs && RendererOutput == FSceneRenderer::ERendererOutput::FinalSceneColor)
+	if (SceneUpdateInputs)
 	{
-		GRayTracingGeometryManager->PreRender();
-
-		// TODO: should only process build requests once per frame
 		RHI_BREADCRUMB_EVENT_STAT(GraphBuilder.RHICmdList, RayTracingGeometry, "RayTracingGeometry");
 		SCOPED_GPU_STAT(GraphBuilder.RHICmdList, RayTracingGeometry);
 
+		// TODO: should only process build requests once per frame
 		GRayTracingGeometryManager->ProcessBuildRequests(GraphBuilder.RHICmdList);
 	}
 
@@ -1747,10 +1745,17 @@ void FDeferredShadingSceneRenderer::Render(FRDGBuilder& GraphBuilder, const FSce
 		View.SetRayTracingSceneViewHandle(RayTracingScene.AddView(View.GetViewKey()));
 		RayTracingScene.SetViewParams(View.GetRayTracingSceneViewHandle(), View.ViewMatrices, View.RayTracingCullingParameters);
 	}
-#endif
+#endif // RHI_RAYTRACING
 
 	FInitViewTaskDatas InitViewTaskDatas = OnRenderBegin(GraphBuilder, SceneUpdateInputs);
 
+#if RHI_RAYTRACING
+	if (RendererOutput == FSceneRenderer::ERendererOutput::FinalSceneColor && FamilyPipelineState[&FFamilyPipelineState::bRayTracing])
+	{
+		GRayTracingGeometryManager->PreRender();
+	}
+#endif // RHI_RAYTRACING
+
 	FUpdateExposureCompensationCurveLUTTaskData UpdateExposureCompensationCurveLUTTaskData;
 	BeginUpdateExposureCompensationCurveLUT(Views, &UpdateExposureCompensationCurveLUTTaskData);
 
