Remove unused stencil plane from SkyLight.CubeDepthTexture

D3DDebug was also reporting that we weren't tracking the unused plane correctly.

diff --git a/Engine/Source/Runtime/Renderer/Private/ReflectionEnvironmentRealTimeCapture.cpp b/Engine/Source/Runtime/Renderer/Private/ReflectionEnvironmentRealTimeCapture.cpp
--- a/Engine/Source/Runtime/Renderer/Private/ReflectionEnvironmentRealTimeCapture.cpp
+++ b/Engine/Source/Runtime/Renderer/Private/ReflectionEnvironmentRealTimeCapture.cpp
@@ -754,7 +754,7 @@ void FScene::AllocateAndCaptureFrameSkyEnvMap(
 					
 								if (bUseDepthBuffer)
 								{
-									FRDGTextureDesc CubeDepthTextureDesc = FRDGTextureDesc::Create2D(FIntPoint(CubeWidth, CubeWidth), PF_DepthStencil,
+									FRDGTextureDesc CubeDepthTextureDesc = FRDGTextureDesc::Create2D(FIntPoint(CubeWidth, CubeWidth), PF_ShadowDepth,
 										MainView.GetSceneTexturesConfig().DepthClearValue,
 										TexCreate_DepthStencilTargetable | TexCreate_ShaderResource);
 									CubeDepthTexture = GraphBuilder.CreateTexture(CubeDepthTextureDesc, TEXT("SkyLight.CubeDepthTexture"));
@@ -832,7 +832,7 @@ void FScene::AllocateAndCaptureFrameSkyEnvMap(
 								// Setup the depth buffer
 								if (bUseDepthBuffer)
 								{
-									FRDGTextureDesc CubeDepthTextureDesc = FRDGTextureDesc::Create2D(FIntPoint(CubeWidth, CubeWidth), PF_DepthStencil,
+									FRDGTextureDesc CubeDepthTextureDesc = FRDGTextureDesc::Create2D(FIntPoint(CubeWidth, CubeWidth), PF_ShadowDepth,
 										MainView.GetSceneTexturesConfig().DepthClearValue,
 										TexCreate_DepthStencilTargetable | TexCreate_ShaderResource);
 									CubeDepthTexture = GraphBuilder.CreateTexture(CubeDepthTextureDesc, TEXT("SkyLight.CubeDepthTexture"));
