Make sure the RT segment is setup again when the vertex count changed as well

diff --git a/Engine/Source/Runtime/Renderer/Private/RayTracing/RayTracingDynamicGeometryUpdateManager.cpp b/Engine/Source/Runtime/Renderer/Private/RayTracing/RayTracingDynamicGeometryUpdateManager.cpp
--- a/Engine/Source/Runtime/Renderer/Private/RayTracing/RayTracingDynamicGeometryUpdateManager.cpp
+++ b/Engine/Source/Runtime/Renderer/Private/RayTracing/RayTracingDynamicGeometryUpdateManager.cpp
@@ -422,7 +422,9 @@ void FRayTracingDynamicGeometryUpdateManager::AddDynamicGeometryToUpdate(
 
 	check(Geometry.IsInitialized());
 
-	if (Geometry.Initializer.TotalPrimitiveCount != UpdateParams.NumTriangles && UpdateParams.NumTriangles > 0)
+	// Primitive count or triangle count changed then setup the initializer and segment data again
+	bool bVertexCountChanges = Geometry.Initializer.Segments.Num() > 0 ? Geometry.Initializer.Segments[0].MaxVertices != NumVertices : false;
+	if ((Geometry.Initializer.TotalPrimitiveCount != UpdateParams.NumTriangles || bVertexCountChanges) && UpdateParams.NumTriangles > 0)
 	{
 		checkf(Geometry.Initializer.Segments.Num() <= 1, TEXT("Dynamic ray tracing geometry '%s' has an unexpected number of segments."), *Geometry.Initializer.DebugName.ToString());
 		Geometry.Initializer.TotalPrimitiveCount = UpdateParams.NumTriangles;
