[FastGeo] Fixed primitives incorrectly shown when running PIE and the player is ejected.
- Deprecated FSceneInterface::UpdatePrimitivesDrawnInGame_RenderThread in favor of FSceneInterface::UpdatePrimitivesIsDrawn_RenderThread.
- Deprecated FPrimitiveSceneProxy::SetDrawnInGame_RenderThread in favor of FPrimitiveSceneProxy::SetIsDrawn_RenderThread

diff --git a/Engine/Plugins/Experimental/FastGeoStreaming/Source/FastGeoStreaming/Private/FastGeoComponentCluster.cpp b/Engine/Plugins/Experimental/FastGeoStreaming/Source/FastGeoStreaming/Private/FastGeoComponentCluster.cpp
--- a/Engine/Plugins/Experimental/FastGeoStreaming/Source/FastGeoStreaming/Private/FastGeoComponentCluster.cpp
+++ b/Engine/Plugins/Experimental/FastGeoStreaming/Source/FastGeoStreaming/Private/FastGeoComponentCluster.cpp
@@ -203,7 +203,7 @@ void FFastGeoComponentCluster::UpdateVisibility_Internal(TArray<FFastGeoPrimitiv
 						}
 						else
 						{
-							Proxy->GetScene().UpdatePrimitivesDrawnInGame_RenderThread(MakeArrayView(&Proxy, 1), bShow);
+							Proxy->GetScene().UpdatePrimitivesIsDrawn_RenderThread(MakeArrayView(&Proxy, 1), bShow);
 						}
 					}
 				}
diff --git a/Engine/Plugins/Experimental/FastGeoStreaming/Source/FastGeoStreaming/Private/FastGeoPrimitiveComponent.cpp b/Engine/Plugins/Experimental/FastGeoStreaming/Source/FastGeoStreaming/Private/FastGeoPrimitiveComponent.cpp
--- a/Engine/Plugins/Experimental/FastGeoStreaming/Source/FastGeoStreaming/Private/FastGeoPrimitiveComponent.cpp
+++ b/Engine/Plugins/Experimental/FastGeoStreaming/Source/FastGeoStreaming/Private/FastGeoPrimitiveComponent.cpp
@@ -256,6 +256,7 @@ void FFastGeoPrimitiveComponent::UpdateVisibility()
 #if !UE_BUILD_SHIPPING
 	SceneProxyDesc.bIsVisible = SceneProxyDesc.bIsVisible && UFastGeoWorldSubsystem::IsFastGeoVisible();
 #endif
+	SceneProxyDesc.bIsVisibleEditor = SceneProxyDesc.bIsVisible;
 }
 
 void FFastGeoPrimitiveComponent::InitializeSceneProxyDescDynamicProperties()
diff --git a/Engine/Source/Runtime/Engine/Private/PrimitiveSceneProxy.cpp b/Engine/Source/Runtime/Engine/Private/PrimitiveSceneProxy.cpp
--- a/Engine/Source/Runtime/Engine/Private/PrimitiveSceneProxy.cpp
+++ b/Engine/Source/Runtime/Engine/Private/PrimitiveSceneProxy.cpp
@@ -1101,16 +1101,18 @@ void FPrimitiveSceneProxy::SetCustomDepthEnabled_RenderThread(const bool bInRend
 }
 
 /**
-* Set the draw in game flag (RENDER THREAD)
+* Set the DrawInGame and the DrawInEditor flags (RENDER THREAD)
 *
 * @param the new value
 */
-void FPrimitiveSceneProxy::SetDrawnInGame_RenderThread(bool bInDrawInGame)
+void FPrimitiveSceneProxy::SetIsDrawn_RenderThread(bool bIsDrawn)
 {
 	check(IsInRenderingThread());
-	if (DrawInGame != bInDrawInGame)
+	
+	if ((DrawInGame != DrawInEditor) || (DrawInGame != bIsDrawn))
 	{
-		DrawInGame = bInDrawInGame;
+		DrawInGame = bIsDrawn;
+		DrawInEditor = bIsDrawn;
 
 		if (PrimitiveSceneInfo)
 		{
diff --git a/Engine/Source/Runtime/Engine/Public/PrimitiveSceneProxy.h b/Engine/Source/Runtime/Engine/Public/PrimitiveSceneProxy.h
--- a/Engine/Source/Runtime/Engine/Public/PrimitiveSceneProxy.h
+++ b/Engine/Source/Runtime/Engine/Public/PrimitiveSceneProxy.h
@@ -1176,11 +1176,14 @@ public:
 	/** Tell us if this proxy is drawn in game.*/
 	inline bool IsDrawnInGame() const { return DrawInGame; }
 
+	UE_DEPRECATED(5.8, "Use SetIsDrawn_RenderThread instead.")
+	inline void SetDrawnInGame_RenderThread(bool bInDrawInGame) {}
+
 	/**
-	 * Set the drawn in game flag, from the render thread.
-	 * @param bInDrawInGame		The new value.
+	 * Set the DrawInGame and the DrawInEditor flags, from the render thread.
+	 * @param bIsDrawn		The new value.
  	 */
-	ENGINE_API void SetDrawnInGame_RenderThread(bool bInDrawInGame);
+	ENGINE_API void SetIsDrawn_RenderThread(bool bIsDrawn);
 
 	/** Tell us if this proxy is drawn in editor.*/
 	inline bool IsDrawnInEditor() const { return DrawInEditor; }
diff --git a/Engine/Source/Runtime/Engine/Public/SceneInterface.h b/Engine/Source/Runtime/Engine/Public/SceneInterface.h
--- a/Engine/Source/Runtime/Engine/Public/SceneInterface.h
+++ b/Engine/Source/Runtime/Engine/Public/SceneInterface.h
@@ -562,7 +562,10 @@ public:
 	UE_DEPRECATED(5.6, "StaticMeshBounds is no longer a parameter of this interface, please update FPrimitiveSceneDesc version.")
 	virtual void UpdatePrimitiveInstances(FInstancedStaticMeshSceneDesc* Primitive) = 0;	
 
-	virtual void UpdatePrimitivesDrawnInGame_RenderThread(TArrayView<FPrimitiveSceneProxy*> PrimitiveSceneProxies, bool bDrawnInGame) = 0;
+	UE_DEPRECATED(5.8, "Use UpdatePrimitivesIsDrawn_RenderThread instead.")
+	virtual void UpdatePrimitivesDrawnInGame_RenderThread(TArrayView<FPrimitiveSceneProxy*> PrimitiveSceneProxies, bool bDrawnInGame) final {}
+
+	virtual void UpdatePrimitivesIsDrawn_RenderThread(TArrayView<FPrimitiveSceneProxy*> PrimitiveSceneProxies, bool bIsDrawn) = 0;
 
 	/**
 	 * Release this scene and remove it from the rendering thread
diff --git a/Engine/Source/Runtime/Renderer/Private/RendererScene.cpp b/Engine/Source/Runtime/Renderer/Private/RendererScene.cpp
--- a/Engine/Source/Runtime/Renderer/Private/RendererScene.cpp
+++ b/Engine/Source/Runtime/Renderer/Private/RendererScene.cpp
@@ -191,7 +191,7 @@ static TAutoConsoleVariable<int32> CVarLumenSceneSurfaceCacheAllowPersistOnVisib
 	TEXT("r.LumenScene.SurfaceCache.AllowPersistOnVisibilityChange"),
 	1,
 	TEXT("Whether to allow primitives that affect indirect lighting while hidden to persist their surface cache on visibility change.\n")
-	TEXT("Only supported through FScene::UpdatePrimitivesDrawnInGame_RenderThread currently."),
+	TEXT("Only supported through FScene::UpdatePrimitivesIsDrawn_RenderThread currently."),
 	ECVF_RenderThreadSafe
 );
 
@@ -1815,13 +1815,13 @@ static bool AllowPersistSurfaceCacheOnVisibilityChange(const FPrimitiveSceneProx
 		&& !PrimitiveSceneProxy->IsEmissiveLightSource();
 }
 
-void FScene::UpdatePrimitivesDrawnInGame_RenderThread(TArrayView<FPrimitiveSceneProxy*> InPrimitiveSceneProxies, bool bDrawnInGame)
+void FScene::UpdatePrimitivesIsDrawn_RenderThread(TArrayView<FPrimitiveSceneProxy*> InPrimitiveSceneProxies, bool bInIsDrawn)
 {
 	check(IsInRenderingThread());
 
 	for (FPrimitiveSceneProxy* PrimitiveSceneProxy : InPrimitiveSceneProxies)
 	{
-		if (PrimitiveSceneProxy->IsDrawnInGame() != bDrawnInGame)
+		if ((PrimitiveSceneProxy->IsDrawnInGame() != PrimitiveSceneProxy->IsDrawnInEditor()) || (PrimitiveSceneProxy->IsDrawnInGame() != bInIsDrawn))
 		{
 			FPrimitiveSceneInfo* PrimitiveSceneInfo = PrimitiveSceneProxy->GetPrimitiveSceneInfo();
 			const int32 PrimitiveIndex = PrimitiveSceneInfo->GetIndex();
@@ -1829,7 +1829,7 @@ void FScene::UpdatePrimitivesDrawnInGame_RenderThread(TArrayView<FPrimitiveScene
 			check(bPrimitiveIndexValid);
 			const bool bPersistSurfaceCache = AllowPersistSurfaceCacheOnVisibilityChange(PrimitiveSceneProxy);
 
-			if (!bDrawnInGame && bPrimitiveIndexValid)
+			if (!bInIsDrawn && bPrimitiveIndexValid)
 			{
 				DistanceFieldSceneData.RemovePrimitive(PrimitiveSceneInfo);
 
@@ -1839,9 +1839,9 @@ void FScene::UpdatePrimitivesDrawnInGame_RenderThread(TArrayView<FPrimitiveScene
 				}
 			}
 
-			PrimitiveSceneProxy->SetDrawnInGame_RenderThread(bDrawnInGame);
+			PrimitiveSceneProxy->SetIsDrawn_RenderThread(bInIsDrawn);
 
-			if (bDrawnInGame && bPrimitiveIndexValid)
+			if (bInIsDrawn && bPrimitiveIndexValid)
 			{
 				DistanceFieldSceneData.AddPrimitive(PrimitiveSceneInfo);
 				// 5.6.1 Temporary Fix
@@ -1857,7 +1857,7 @@ void FScene::UpdatePrimitivesDrawnInGame_RenderThread(TArrayView<FPrimitiveScene
 			if (bPrimitiveIndexValid && PrimitiveSceneProxy->HasRayTracingRepresentation())
 			{
 				FScene::FPrimitiveRayTracingData& PrimitiveRayTracingData = PrimitiveRayTracingDatas[PrimitiveIndex];
-				PrimitiveRayTracingData.bDrawInGame = bDrawnInGame;
+				PrimitiveRayTracingData.bDrawInGame = bInIsDrawn;
 
 				UpdateCachedRayTracingState(PrimitiveSceneProxy);
 			}
@@ -6837,7 +6837,7 @@ public:
 	virtual void UpdatePrimitiveInstances(FPrimitiveSceneDesc* Primitive) override {};
 	virtual void UpdatePrimitiveInstances(FInstancedStaticMeshSceneDesc* Primitive) override {};	
 
-	virtual void UpdatePrimitivesDrawnInGame_RenderThread(TArrayView<FPrimitiveSceneProxy*> InPrimitiveSceneProxies, bool bDrawnInGame) override {}
+	virtual void UpdatePrimitivesIsDrawn_RenderThread(TArrayView<FPrimitiveSceneProxy*> PrimitiveSceneProxies, bool bIsDrawn) override {}
 
 private:
 	UWorld* World;
diff --git a/Engine/Source/Runtime/Renderer/Private/ScenePrivate.h b/Engine/Source/Runtime/Renderer/Private/ScenePrivate.h
--- a/Engine/Source/Runtime/Renderer/Private/ScenePrivate.h
+++ b/Engine/Source/Runtime/Renderer/Private/ScenePrivate.h
@@ -3924,8 +3924,8 @@ public:
 	/** Returns true if the scene requires the use of debug material permutations. */
 	bool RequiresDebugMaterials() const { return PersistentViewStateDebugFlags != 0; }
 
-     /** Update the DrawnInGame flag on the provided scene proxies */
-	virtual void UpdatePrimitivesDrawnInGame_RenderThread(TArrayView<FPrimitiveSceneProxy*> PrimitiveSceneProxies, bool bDrawnInGame) override;
+     /** Update the DrawnInGame and the DrawnInEditor flags on the provided scene proxies */
+	virtual void UpdatePrimitivesIsDrawn_RenderThread(TArrayView<FPrimitiveSceneProxy*> PrimitiveSceneProxies, bool bIsDrawn) override;
 
 protected:
 
