PSO Precache fixes and improvements for mobile renderer
- Don't precache LOCAL_LIGHTS_DISABLED when MobileLocalLightsUseSinglePermutation is set because it's not needed anymore
- Don't precache decal pass PSOs in mobile renderer when dbuffer is not enabled
- Enable precaching of deferred decal mesh PSOs because those are actually used
- Fix PSO precaching of text component (vertex element setup has been fixed in other CL)
- Enable PSO precaching of all default material permutation also when manual vertex fetch is disabled
- Don't track misses for global slate PSOs when global shader precaching is disabled (default on mobile)
- Fix PSO precaching misses for mobile lights due to inset shadows support

diff --git a/Engine/Source/Runtime/Engine/Private/Components/TextRenderComponent.cpp b/Engine/Source/Runtime/Engine/Private/Components/TextRenderComponent.cpp
--- a/Engine/Source/Runtime/Engine/Private/Components/TextRenderComponent.cpp
+++ b/Engine/Source/Runtime/Engine/Private/Components/TextRenderComponent.cpp
@@ -1588,20 +1588,28 @@ FVertexDeclarationElementList InitDummyVertexDeclarationElementsForText()
 
 void UTextRenderComponent::PrecachePSOs()
 {
-	if (((IsComponentPSOPrecachingEnabled() && RHISupportsManualVertexFetch(GMaxRHIShaderPlatform))) && TextMaterial)
+	if (IsComponentPSOPrecachingEnabled() && TextMaterial)
 	{
-		// FIXME: need to collect an actual vertex declaration for non-MVF path
 		FPSOPrecacheParams PrecachePSOParams;
 		SetupPrecachePSOParams(PrecachePSOParams);
 
 		// Overwrite the cast shadow flag because MeshBatch setup in DrawStaticElements doesn't set the correct CastShadow flag from the proxy or component
 		// and leaves the default CastShadow value which is true
 		PrecachePSOParams.bCastShadow = true;
-
-		const FVertexFactoryType* VFType = &FLocalVertexFactory::StaticType;
-		FPSOPrecacheVertexFactoryDataList VFDataList;
-		VFDataList.Add(FPSOPrecacheVertexFactoryData(VFType));
-		TextMaterial->PrecachePSOs(VFType, PrecachePSOParams);
+	
+		FPSOPrecacheVertexFactoryDataList VertexFactoryDataList;
+		if (RHISupportsManualVertexFetch(GMaxRHIShaderPlatform))
+		{
+			VertexFactoryDataList.Add(FPSOPrecacheVertexFactoryData(&FLocalVertexFactory::StaticType));
+		}
+		else
+		{
+			const static FVertexDeclarationElementList Elements = InitDummyVertexDeclarationElementsForText();
+			VertexFactoryDataList.Add(FPSOPrecacheVertexFactoryData(&FLocalVertexFactory::StaticType, Elements));
+		}
+				
+		TArray<FMaterialPSOPrecacheRequestID> MaterialPrecacheRequestIDs;
+		TextMaterial->PrecachePSOs(VertexFactoryDataList, PrecachePSOParams, EPSOPrecachePriority::Medium, MaterialPrecacheRequestIDs);
 	}
 }
 
diff --git a/Engine/Source/Runtime/Engine/Private/Materials/Material.cpp b/Engine/Source/Runtime/Engine/Private/Materials/Material.cpp
--- a/Engine/Source/Runtime/Engine/Private/Materials/Material.cpp
+++ b/Engine/Source/Runtime/Engine/Private/Materials/Material.cpp
@@ -673,33 +673,19 @@ void UMaterialInterface::PrecacheDefaultMaterialPSOs()
 		AllVertexFactoryTypes.Add(VFData);
 	}
 
-	// Skip platforms that do not support MVF, non-MVF path needs mesh information for PSO 
-	if (RHISupportsManualVertexFetch(GMaxRHIShaderPlatform) || IsPSOShaderPreloadingEnabled())
+	for (int32 Domain = 0; Domain < MD_MAX; ++Domain)
 	{
-		for (int32 Domain = 0; Domain < MD_MAX; ++Domain)
+		if (GDefaultMaterials[Domain])
 		{
-			if (GDefaultMaterials[Domain])
+			PrecachePSOParams.Mobility = EComponentMobility::Static;
+			if (PipelineStateCache::IsPSOPrecachingEnabled() || IsPSOShaderPreloadingEnabled())
 			{
-				PrecachePSOParams.Mobility = EComponentMobility::Static;
-				if (PipelineStateCache::IsPSOPrecachingEnabled() || IsPSOShaderPreloadingEnabled())
-				{
-					GDefaultMaterials[Domain]->PrecachePSOs(AllVertexFactoryTypes, PrecachePSOParams, EPSOPrecachePriority::High, MaterialPrecacheRequestIDs);
-				}
-				PrecachePSOParams.Mobility = EComponentMobility::Movable;
-				if (PipelineStateCache::IsPSOPrecachingEnabled() || IsPSOShaderPreloadingEnabled())
-				{
-					GDefaultMaterials[Domain]->PrecachePSOs(AllVertexFactoryTypes, PrecachePSOParams, EPSOPrecachePriority::High, MaterialPrecacheRequestIDs);
-				}
+				GDefaultMaterials[Domain]->PrecachePSOs(AllVertexFactoryTypes, PrecachePSOParams, EPSOPrecachePriority::High, MaterialPrecacheRequestIDs);
 			}
-		}
-	}
-	else
-	{
-		if (GDefaultMaterials[MD_LightFunction])
-		{
-			if (PipelineStateCache::IsPSOPrecachingEnabled())
+			PrecachePSOParams.Mobility = EComponentMobility::Movable;
+			if (PipelineStateCache::IsPSOPrecachingEnabled() || IsPSOShaderPreloadingEnabled())
 			{
-				GDefaultMaterials[MD_LightFunction]->PrecachePSOs(AllVertexFactoryTypes, PrecachePSOParams, EPSOPrecachePriority::High, MaterialPrecacheRequestIDs);
+				GDefaultMaterials[Domain]->PrecachePSOs(AllVertexFactoryTypes, PrecachePSOParams, EPSOPrecachePriority::High, MaterialPrecacheRequestIDs);
 			}
 		}
 	}
diff --git a/Engine/Source/Runtime/Renderer/Private/CompositionLighting/PostProcessMeshDecals.cpp b/Engine/Source/Runtime/Renderer/Private/CompositionLighting/PostProcessMeshDecals.cpp
--- a/Engine/Source/Runtime/Renderer/Private/CompositionLighting/PostProcessMeshDecals.cpp
+++ b/Engine/Source/Runtime/Renderer/Private/CompositionLighting/PostProcessMeshDecals.cpp
@@ -30,6 +30,13 @@ static TAutoConsoleVariable<int32> CVarParallelMeshDecal(
 	ECVF_RenderThreadSafe
 );
 
+static int32 GDeferredDecalPassMaterialPSOPrecache = 1;
+static FAutoConsoleVariableRef CVarGDeferredDecalPassMaterialPSOPrecache(
+	TEXT("r.PSOPrecache.DeferredDecalPassMaterial"), 
+	GDeferredDecalPassMaterialPSOPrecache, 
+	TEXT("Precache all possible required PSOs for loaded DeferredDecalPass Materials."),
+	ECVF_ReadOnly);
+
 static int32 GDeferredDecalMeshMaterialPSOPrecache = 1;
 static FAutoConsoleVariableRef CVarGDeferredDecalMeshMaterialPSOPrecache(
 	TEXT("r.PSOPrecache.DeferredDecalMeshMaterial"), 
@@ -477,8 +484,18 @@ void FMeshDecalMeshProcessor::CollectPSOInitializers(
 			continue;
 		}
 		
-		// Collect decal pass PSOs
-		CollectDeferredDecalPassPSOInitializers(PSOCollectorIndex, FeatureLevel, SceneTexturesConfig, Material, LocalDecalRenderStage, PSOInitializers);
+		// Mobile only runs decal pass when dbuffers are enabled
+		bool bCollectDecalPassPSOInitializers = GDeferredDecalPassMaterialPSOPrecache > 0;
+		if (IsMobilePlatform(ShaderPlatform) && !IsUsingDBuffers(ShaderPlatform))
+		{
+			bCollectDecalPassPSOInitializers = false;
+		}
+		
+		if (bCollectDecalPassPSOInitializers)
+		{
+			// Collect decal pass PSOs
+			CollectDeferredDecalPassPSOInitializers(PSOCollectorIndex, FeatureLevel, SceneTexturesConfig, Material, LocalDecalRenderStage, PSOInitializers);
+		}
 
 		if (GDeferredDecalMeshMaterialPSOPrecache)
 		{
diff --git a/Engine/Source/Runtime/Renderer/Private/MobileBasePass.cpp b/Engine/Source/Runtime/Renderer/Private/MobileBasePass.cpp
--- a/Engine/Source/Runtime/Renderer/Private/MobileBasePass.cpp
+++ b/Engine/Source/Runtime/Renderer/Private/MobileBasePass.cpp
@@ -1098,7 +1098,10 @@ void FMobileBasePassMeshProcessor::CollectPSOInitializers(const FSceneTexturesCo
 	{
 		LocalLightSetting = GetMobileForwardLocalLightSetting(ShaderPlatform);
 	}
+		
 	const bool bUseLocalLightPermutation = (LocalLightSetting != EMobileLocalLightSetting::LOCAL_LIGHTS_DISABLED);
+	const bool bUseNoLocalLightPermutation = !bUseLocalLightPermutation || !MobileLocalLightsUseSinglePermutation(ShaderPlatform);
+	check(bUseLocalLightPermutation || bUseNoLocalLightPermutation);
 
 	const bool bCanReceiveCSM = ((Flags & EFlags::CanReceiveCSM) == EFlags::CanReceiveCSM) && MaterialCanReceiveCSM(Material);
 
@@ -1106,7 +1109,10 @@ void FMobileBasePassMeshProcessor::CollectPSOInitializers(const FSceneTexturesCo
 	
 	for (ELightMapPolicyType LightMapPolicyType : UniformLightMapPolicyTypes)
 	{
-		CollectPSOInitializersForLMPolicy(VertexFactoryData, DrawRenderState, RenderTargetsInfo, Material, EMobileLocalLightSetting::LOCAL_LIGHTS_DISABLED, LightMapPolicyType, MeshFillMode, MeshCullMode, (EPrimitiveType)PreCacheParams.PrimitiveType, PSOInitializers);
+		if (bUseNoLocalLightPermutation)
+		{
+			CollectPSOInitializersForLMPolicy(VertexFactoryData, DrawRenderState, RenderTargetsInfo, Material, EMobileLocalLightSetting::LOCAL_LIGHTS_DISABLED, LightMapPolicyType, MeshFillMode, MeshCullMode, (EPrimitiveType)PreCacheParams.PrimitiveType, PSOInitializers);
+		}
 		if (bUseLocalLightPermutation)
 		{
 			CollectPSOInitializersForLMPolicy(VertexFactoryData, DrawRenderState, RenderTargetsInfo, Material, LocalLightSetting, LightMapPolicyType, MeshFillMode, MeshCullMode, (EPrimitiveType)PreCacheParams.PrimitiveType, PSOInitializers);
diff --git a/Engine/Source/Runtime/Renderer/Private/MobileDeferredShadingPass.cpp b/Engine/Source/Runtime/Renderer/Private/MobileDeferredShadingPass.cpp
--- a/Engine/Source/Runtime/Renderer/Private/MobileDeferredShadingPass.cpp
+++ b/Engine/Source/Runtime/Renderer/Private/MobileDeferredShadingPass.cpp
@@ -1148,6 +1148,10 @@ void FDeferredMobileLightMaterialPSOCollector::CollectPSOInitializersDirectional
 	{
 		return;
 	}
+	
+	static auto CVarMobileSupportInsetShadows = IConsoleManager::Get().FindConsoleVariable(TEXT("r.Mobile.SupportInsetShadows"));
+	const bool bInsetShadows = (CVarMobileSupportInsetShadows != nullptr && CVarMobileSupportInsetShadows->GetInt() != 0);
+	const bool bMobileUsesShadowMaskTexture = (MobileUsesShadowMaskTexture(GMaxRHIShaderPlatform) || bInsetShadows);
 
 	auto AddPSOInitializer =
 		[&](int32 PassIndex, bool bInlineReflectionAndSky, bool bOnlyDefaultLitInView, int32 ShadowQuality, bool bHasBoxSphere)
@@ -1160,14 +1164,14 @@ void FDeferredMobileLightMaterialPSOCollector::CollectPSOInitializersDirectional
 			const bool bClustredReflection = bInlineReflectionAndSky && bHasBoxSphere;
 			const bool bPlanarReflection = false;
 			const bool bEnableSkyLight = bInlineReflectionAndSky && bDynamicSkyLight;
-			const bool bMobileUsesShadowMaskTexture = MobileUsesShadowMaskTexture(GMaxRHIShaderPlatform);
 			const bool bApplySkyShadowing = false;
 			const int32 LightingChannel = 0;
-			
+						
 			FMobileDirectionalLightFunctionPS::FPermutationDomain PermutationVector;
 			PermutationVector.Set<FMobileDirectionalLightFunctionPS::FEnableShadingModelSupport>(bShadingModelSupport);
 			PermutationVector.Set<FMobileDirectionalLightFunctionPS::FEnableClustredLights>(bUseClusteredLights);
 			PermutationVector.Set<FMobileDirectionalLightFunctionPS::FEnableSkyLight>(bEnableSkyLight);
+			PermutationVector.Set<FMobileDirectionalLightFunctionPS::FEnableScreenSpaceShadowMask>(bMobileUsesShadowMaskTexture);
 			PermutationVector.Set<FMobileDirectionalLightFunctionPS::FEnableCSM>(ShadowQuality > 0);
 			PermutationVector.Set<FMobileDirectionalLightFunctionPS::FShadowQuality>(FMath::Clamp(ShadowQuality, 1, 3));
 			PermutationVector.Set<FMobileDirectionalLightFunctionPS::FSkyShadowing>(bInlineReflectionAndSky && bApplySkyShadowing);
diff --git a/Engine/Source/Runtime/SlateRHIRenderer/Private/SlateRHIRenderingPolicy.cpp b/Engine/Source/Runtime/SlateRHIRenderer/Private/SlateRHIRenderingPolicy.cpp
--- a/Engine/Source/Runtime/SlateRHIRenderer/Private/SlateRHIRenderingPolicy.cpp
+++ b/Engine/Source/Runtime/SlateRHIRenderer/Private/SlateRHIRenderingPolicy.cpp
@@ -1423,7 +1423,8 @@ void DrawSlateRenderBatch(
 		State.GraphicsPSOInit.PrimitiveType = GetRHIPrimitiveType(RenderBatch.DrawPrimitiveType);
 
 #if PSO_PRECACHING_VALIDATE
-		if (PSOCollectorStats::IsFullPrecachingValidationEnabled())
+		static IConsoleVariable* PrecacheGlobalShadersCVar = IConsoleManager::Get().FindConsoleVariable(TEXT("r.PSOPrecache.GlobalShaders"));
+		if (PSOCollectorStats::IsFullPrecachingValidationEnabled() && PrecacheGlobalShadersCVar && PrecacheGlobalShadersCVar->GetInt() > 0)
 		{			
 			static const int32 GlobalPSOCollectorIndex = FGlobalPSOCollectorManager::GetIndex(SlateGlobalPSOCollectorName);
 			PSOCollectorStats::CheckGlobalGraphicsPipelineStateInCache(State.GraphicsPSOInit, GlobalPSOCollectorIndex);
