Fixed a latent bug in the instance culling context if a situation occurs where there are zero instances in GPU-scene, while the MDC had a nonzero number of instances. This would cause the draw to fall back to a direct draw which then would try to render instances fetched from the dummy buffer.

#rnx

diff --git a/Engine/Source/Runtime/Renderer/Private/InstanceCulling/InstanceCullingContext.cpp b/Engine/Source/Runtime/Renderer/Private/InstanceCulling/InstanceCullingContext.cpp
--- a/Engine/Source/Runtime/Renderer/Private/InstanceCulling/InstanceCullingContext.cpp
+++ b/Engine/Source/Runtime/Renderer/Private/InstanceCulling/InstanceCullingContext.cpp
@@ -235,6 +235,7 @@ void FInstanceCullingContext::ResetCommands(int32 MaxNumCommands)
 	InstanceIdOffsets.Empty(MaxNumCommands);
 	PayloadData.Empty(MaxNumCommands);
 	TotalInstances = 0U;
+	bHasAnyIndirectDraw = false;
 
 	DrawCommandCompactionData.Empty(MaxNumCommands);
 	CompactionBlockDataIndices.Reset();
@@ -692,7 +693,8 @@ void FInstanceCullingContext::BuildRenderingCommands(FRDGBuilder& GraphBuilder,
 
 bool FInstanceCullingContext::HasCullingCommands() const
 {
-	check(!SyncPrerequisitesFunc);  return TotalInstances > 0;
+	check(!SyncPrerequisitesFunc);  
+	return bHasAnyIndirectDraw || TotalInstances > 0u;
 }
 
 void FInstanceCullingContext::BuildRenderingCommandsInternal(
@@ -956,6 +958,12 @@ void FInstanceCullingContext::BuildRenderingCommandsInternal(
 		}
 	}
 
+	// If there were no instances & thus the compute passes were skipped we need to clear the buffer to signal to RDG that it is ok to bind it later.
+	if (TotalInstances == 0u)
+	{
+		AddClearUAVPass(GraphBuilder, InstanceIdsBufferUAV, 0u);
+	}
+
 	InstanceCullingDrawParams->DrawIndirectArgsBuffer = DrawIndirectArgsRDG;
 	InstanceCullingDrawParams->InstanceIdOffsetBuffer = InstanceIdOffsetBufferRDG;
 	
@@ -1562,6 +1570,7 @@ void FInstanceCullingContext::SetupDrawCommands(
 
 				if (bUseIndirectDraw)
 				{
+					bHasAnyIndirectDraw = true;
  					DrawCmd.IndirectArgsOffsetOrNumInstances = CurrentIndirectArgsOffset * FInstanceCullingContext::IndirectArgsNumWords * sizeof(uint32);
 				}
 				else
diff --git a/Engine/Source/Runtime/Renderer/Public/InstanceCulling/InstanceCullingContext.h b/Engine/Source/Runtime/Renderer/Public/InstanceCulling/InstanceCullingContext.h
--- a/Engine/Source/Runtime/Renderer/Public/InstanceCulling/InstanceCullingContext.h
+++ b/Engine/Source/Runtime/Renderer/Public/InstanceCulling/InstanceCullingContext.h
@@ -246,7 +246,9 @@ public:
 	EShaderPlatform ShaderPlatform = SP_NumPlatforms;
 	TArray<int32, TInlineAllocator<6, SceneRenderingAllocator>> ViewIds;
 	TRefCountPtr<IPooledRenderTarget> PrevHZB = nullptr;
-	bool bIsEnabled = false;
+	bool bIsEnabled:1 = false;
+	// This is set during SetupDrawCommands to indicate if there are draw commands using indirect args, as this requires setting up the relevant buffers (even if the instance count is zero).
+	bool bHasAnyIndirectDraw:1 = false;
 	EInstanceCullingMode InstanceCullingMode = EInstanceCullingMode::Normal;
 	EInstanceCullingFlags Flags = EInstanceCullingFlags::None;
 
