[PostProcessing] Avoid additional copies for postprocess materials when slices are not available.

diff --git a/Engine/Source/Runtime/Renderer/Private/PostProcess/PostProcessing.cpp b/Engine/Source/Runtime/Renderer/Private/PostProcess/PostProcessing.cpp
--- a/Engine/Source/Runtime/Renderer/Private/PostProcess/PostProcessing.cpp
+++ b/Engine/Source/Runtime/Renderer/Private/PostProcess/PostProcessing.cpp
@@ -825,6 +825,35 @@ void AddPostProcessingPasses(
 		}
 
 		auto ApplyPostProcessMaterials = [&GraphBuilder, &View, ViewIndex, bOutputInHDR, &VisualizeBufferContext, GetPostProcessMaterialInputs](
+			EBlendableLocation Location,
+			FPostProcessMaterialInputs Inputs,
+			const FPostProcessMaterialChain& MaterialChain,
+			const FScreenPassTexture& InputTexture,
+			EPostProcessMaterialInput MaterialInput = EPostProcessMaterialInput::SceneColor
+		) -> FScreenPassTexture
+		{
+			AddVisualizeBufferMaterialPass(GraphBuilder, View, Location, Inputs, bOutputInHDR, VisualizeBufferContext);
+
+			if (MaterialChain.Num())
+			{
+				return AddPostProcessMaterialChain(GraphBuilder, View, ViewIndex, Inputs, MaterialChain, MaterialInput);
+			}
+			else
+			{
+				FScreenPassTextureSlice CurrentInputSlice = Inputs.GetInput(MaterialInput);
+				if (Inputs.OverrideOutput.IsValid())
+				{
+					AddDrawTexturePass(GraphBuilder, View, CurrentInputSlice, Inputs.OverrideOutput);
+					return Inputs.OverrideOutput;
+				}
+				else
+				{
+					return InputTexture;
+				}
+			}
+		};
+
+		auto ApplyPostProcessMaterialsToSlice = [&GraphBuilder, &View, ViewIndex, bOutputInHDR, &VisualizeBufferContext, GetPostProcessMaterialInputs](
 			EBlendableLocation Location,
 			FPostProcessMaterialInputs Inputs,
 			const FPostProcessMaterialChain& MaterialChain,
@@ -856,8 +885,7 @@ void AddPostProcessingPasses(
 		// Post Process Material Chain - BL_SceneColorBeforeDOF
 		{
 			FScreenPassTextureSlice SceneColorSlice = FScreenPassTextureSlice::CreateFromScreenPassTexture(GraphBuilder, SceneColor);
-			SceneColorSlice = ApplyPostProcessMaterials(BL_SceneColorBeforeDOF, GetPostProcessMaterialInputs(SceneColorSlice), MaterialChainSceneColorBeforeDOF);
-			SceneColor = FScreenPassTexture::CopyFromSlice(GraphBuilder, SceneColorSlice);
+			SceneColor = ApplyPostProcessMaterials(BL_SceneColorBeforeDOF, GetPostProcessMaterialInputs(SceneColorSlice), MaterialChainSceneColorBeforeDOF, SceneColor);
 		}
 
 		// Diaphragm Depth of Field
@@ -903,7 +931,7 @@ void AddPostProcessingPasses(
 
 		// Post Process Material Chain - BL_SceneColorAfterDOF
 		{
-			SceneColorSlice = ApplyPostProcessMaterials(BL_SceneColorAfterDOF, GetPostProcessMaterialInputs(SceneColorSlice), MaterialChainSceneColorAfterDOF);
+			SceneColorSlice = ApplyPostProcessMaterialsToSlice(BL_SceneColorAfterDOF, GetPostProcessMaterialInputs(SceneColorSlice), MaterialChainSceneColorAfterDOF);
 		}
 
 		// Post Process Material Chain - BL_TranslucencyAfterDOF
@@ -928,13 +956,14 @@ void AddPostProcessingPasses(
 
 			if (MaterialChainTranslucencyAfterDOF.Num())
 			{
-				FScreenPassTextureSlice PostDOFTranslucencySlice = ApplyPostProcessMaterials(
+				FScreenPassTexture PostDOFTranslucency {};
+				PostDOFTranslucency = ApplyPostProcessMaterials(
 					BL_TranslucencyAfterDOF,
 					GetPostProcessMaterialInputs(SceneColorSlice), 
 					MaterialChainTranslucencyAfterDOF, 
+					PostDOFTranslucency,
 					EPostProcessMaterialInput::SeparateTranslucency);
-
-				FScreenPassTexture PostDOFTranslucency = FScreenPassTexture::CopyFromSlice(GraphBuilder, PostDOFTranslucencySlice);
+				ensure(PostDOFTranslucency.IsValid());
 
 				PostDOFTranslucencyResources.ColorTexture = PostDOFTranslucency.Texture;
 				ensure(PostDOFTranslucencyResources.ViewRect == PostDOFTranslucency.ViewRect);
@@ -1080,31 +1109,29 @@ void AddPostProcessingPasses(
 		// SVE/Post Process Material Chain - SSR Input
 		if (View.ViewState && !View.bStatePrevViewInfoIsReadOnly)
 		{
-			FScreenPassTextureSlice PassOutput;
+			FScreenPassTexture PassOutput;
 			FPostProcessMaterialInputs PassInputs;
 			const FPostProcessMaterialChain MaterialChain = GetPostProcessMaterialChain(View, BL_SSRInput);
 
 			if (SceneViewExtensionDelegates[static_cast<uint32>(ISceneViewExtension::EPostProcessingPass::SSRInput)].Num())
 			{
 				PassInputs = GetPostProcessMaterialInputs(SceneColorSlice);
-				FScreenPassTexture PassOutputTexture = AddSceneViewExtensionPassChain(GraphBuilder, View, PassInputs,
+				PassOutput = AddSceneViewExtensionPassChain(GraphBuilder, View, PassInputs,
 					SceneViewExtensionDelegates[static_cast<uint32>(ISceneViewExtension::EPostProcessingPass::SSRInput)]);
-				PassOutput = FScreenPassTextureSlice::CreateFromScreenPassTexture(GraphBuilder, PassOutputTexture);
 			}
 
 			{
-				FScreenPassTextureSlice InputSceneColorSlice = PassOutput.IsValid() ? PassOutput : SceneColorSlice;
-				PassOutput = ApplyPostProcessMaterials(BL_SSRInput, GetPostProcessMaterialInputs(InputSceneColorSlice), MaterialChain);
+				FScreenPassTextureSlice InputSceneColorSlice = PassOutput.IsValid() ? FScreenPassTextureSlice::CreateFromScreenPassTexture(GraphBuilder, PassOutput) : SceneColorSlice;
+				PassOutput = ApplyPostProcessMaterials(BL_SSRInput, GetPostProcessMaterialInputs(InputSceneColorSlice), MaterialChain, PassOutput);
 			}
 
 			if (PassOutput.IsValid())
 			{
-				FScreenPassTexture PassOutputTexture = FScreenPassTexture::CopyFromSlice(GraphBuilder, PassOutput);
 				// Save off SSR post process output for the next frame.
-				GraphBuilder.QueueTextureExtraction(PassOutputTexture.Texture, &View.ViewState->PrevFrameViewInfo.CustomSSRInput.RT[0]);
+				GraphBuilder.QueueTextureExtraction(PassOutput.Texture, &View.ViewState->PrevFrameViewInfo.CustomSSRInput.RT[0]);
 
-				View.ViewState->PrevFrameViewInfo.CustomSSRInput.ViewportRect = PassOutputTexture.ViewRect;
-				View.ViewState->PrevFrameViewInfo.CustomSSRInput.ReferenceBufferSize = PassOutputTexture.Texture->Desc.Extent;
+				View.ViewState->PrevFrameViewInfo.CustomSSRInput.ViewportRect = PassOutput.ViewRect;
+				View.ViewState->PrevFrameViewInfo.CustomSSRInput.ReferenceBufferSize = PassOutput.Texture->Desc.Extent;
 			}
 		}
 
@@ -1179,7 +1206,7 @@ void AddPostProcessingPasses(
 		{
 			FPostProcessMaterialInputs PostProcessMaterialInputs = GetPostProcessMaterialInputs(SceneColorSlice);
 			PassSequence.AcceptOverrideIfLastPass(EPass::PostProcessMaterialBeforeBloom, PostProcessMaterialInputs.OverrideOutput);
-			SceneColorSlice = ApplyPostProcessMaterials(BL_SceneColorBeforeBloom, PostProcessMaterialInputs, PostProcessMaterialBeforeBloomChain);
+			SceneColorSlice = ApplyPostProcessMaterialsToSlice(BL_SceneColorBeforeBloom, PostProcessMaterialInputs, PostProcessMaterialBeforeBloomChain);
 		}
 
 		// Generate before bloom lower res scene color if they have not been generated.
@@ -1560,8 +1587,7 @@ void AddPostProcessingPasses(
 			PassInputs.SetInput(GraphBuilder, EPostProcessMaterialInput::PostTonemapHDRColor, SceneColorAfterTonemap);
 			PassInputs.SceneTextures = GetSceneTextureShaderParameters(Inputs.SceneTextures);
 
-			SceneColorSlice = ApplyPostProcessMaterials(BL_SceneColorAfterTonemapping, PassInputs, PostProcessMaterialAfterTonemappingChain);
-			SceneColor = FScreenPassTexture::CopyFromSlice(GraphBuilder, SceneColorSlice);
+			SceneColor = ApplyPostProcessMaterials(BL_SceneColorAfterTonemapping, PassInputs, PostProcessMaterialAfterTonemappingChain, SceneColor);
 		}
 
 		if (PassSequence.IsEnabled(EPass::VisualizeLumenScene))
