Move CombineLUTs to async compute

diff --git a/Engine/Source/Runtime/Renderer/Private/PostProcess/PostProcessCombineLUTs.cpp b/Engine/Source/Runtime/Renderer/Private/PostProcess/PostProcessCombineLUTs.cpp
--- a/Engine/Source/Runtime/Renderer/Private/PostProcess/PostProcessCombineLUTs.cpp
+++ b/Engine/Source/Runtime/Renderer/Private/PostProcess/PostProcessCombineLUTs.cpp
@@ -40,6 +40,12 @@ TAutoConsoleVariable<int32> CVarLUTSize(
 	TEXT("Size of film LUT"),
 	ECVF_RenderThreadSafe | ECVF_Scalability);
 
+TAutoConsoleVariable<bool> CVarLUTAsyncCompute(
+	TEXT("r.LUT.AsyncCompute"),
+	1,
+	TEXT("Controls whether the tonemapping LUT pass is run on async compute."),
+	ECVF_RenderThreadSafe);
+
 TAutoConsoleVariable<int32> CVarColorGrading(
 	TEXT("r.Color.Grading"), 1,
 	TEXT("Controls whether post process settings's color grading settings should be applied."),
@@ -594,6 +600,7 @@ FRDGTextureRef AddCombineLUTPass(FRDGBuilder& GraphBuilder, const FViewInfo& Vie
 	FComputeShaderUtils::AddPass(
 		GraphBuilder,
 		RDG_EVENT_NAME("CombineLUTs %d (CS)", LUTSize),
+		CVarLUTAsyncCompute.GetValueOnRenderThread() ? ERDGPassFlags::AsyncCompute : ERDGPassFlags::Compute,
 		ComputeShader,
 		PassParameters,
 		FIntVector(GroupSize, GroupSize, GroupSize));
diff --git a/Engine/Source/Runtime/Renderer/Private/PostProcess/PostProcessing.cpp b/Engine/Source/Runtime/Renderer/Private/PostProcess/PostProcessing.cpp
--- a/Engine/Source/Runtime/Renderer/Private/PostProcess/PostProcessing.cpp
+++ b/Engine/Source/Runtime/Renderer/Private/PostProcess/PostProcessing.cpp
@@ -994,6 +994,28 @@ void AddPostProcessingPasses(
 
 		ensure(bSceneColorHasPostDOFTranslucency != bComposeSeparateTranslucencyInTSR);
 
+		// For async compute, schedule AddCombineLUTPass before TSR, which might be waiting on Single Layer Water.
+		const FPostProcessingPassDelegateArray& ReplacingTonemapperDelegates = SceneViewExtensionDelegates[static_cast<uint32>(ISceneViewExtension::EPostProcessingPass::ReplacingTonemapper)];
+		const FPostProcessMaterialChain ReplacingTonemapperMaterialChain = GetPostProcessMaterialChain(View, BL_ReplacingTonemapper);
+		FRDGTextureRef ColorGradingTexture = nullptr;
+		if (PassSequence.IsEnabled(EPass::Tonemap) && ReplacingTonemapperDelegates.IsEmpty() && ReplacingTonemapperMaterialChain.IsEmpty())
+		{
+			if (bPrimaryView)
+			{
+				ColorGradingTexture = AddCombineLUTPass(GraphBuilder, View);
+			}
+			// We can re-use the color grading texture from the primary view.
+			else if (View.GetTonemappingLUT())
+			{
+				ColorGradingTexture = TryRegisterExternalTexture(GraphBuilder, View.GetTonemappingLUT());
+			}
+			else
+			{
+				const FViewInfo* PrimaryView = static_cast<const FViewInfo*>(View.Family->Views[0]);
+				ColorGradingTexture = TryRegisterExternalTexture(GraphBuilder, PrimaryView->GetTonemappingLUT());
+			}
+		}
+
 		FScreenPassTextureSlice HalfResSceneColor;
 		FScreenPassTextureSlice QuarterResSceneColor;
 		FScreenPassTextureSlice EighthResSceneColor;
@@ -1472,9 +1494,6 @@ void AddPostProcessingPasses(
 
 		if (PassSequence.IsEnabled(EPass::Tonemap))
 		{
-			const FPostProcessingPassDelegateArray& ReplacingTonemapperDelegates = SceneViewExtensionDelegates[static_cast<uint32>(ISceneViewExtension::EPostProcessingPass::ReplacingTonemapper)];
-			const FPostProcessMaterialChain MaterialChain = GetPostProcessMaterialChain(View, BL_ReplacingTonemapper);
-
 			// GPU Skin Cache for next frame can overlap with the tone-mapping pass.
 			if (FGPUSkinCache* GPUSkinCache = Scene->GetGPUSkinCache(); GPUSkinCache && View.IsLastInFamily())
 			{
@@ -1500,31 +1519,14 @@ void AddPostProcessingPasses(
 
 				SceneColor = HighestPriorityDelegate.Execute(GraphBuilder, View, GetReplaceTonemapperInputs());
 			}
-			else if (MaterialChain.Num())
+			else if (ReplacingTonemapperMaterialChain.Num())
 			{
-				const UMaterialInterface* HighestPriorityMaterial = MaterialChain[0];
+				const UMaterialInterface* HighestPriorityMaterial = ReplacingTonemapperMaterialChain[0];
 
 				SceneColor = AddPostProcessMaterialPass(GraphBuilder, View, GetReplaceTonemapperInputs(), HighestPriorityMaterial);
 			}
 			else
 			{
-				FRDGTextureRef ColorGradingTexture = nullptr;
-
-				if (bPrimaryView)
-				{
-					ColorGradingTexture = AddCombineLUTPass(GraphBuilder, View);
-				}
-				// We can re-use the color grading texture from the primary view.
-				else if (View.GetTonemappingLUT())
-				{
-					ColorGradingTexture = TryRegisterExternalTexture(GraphBuilder, View.GetTonemappingLUT());
-				}
-				else
-				{
-					const FViewInfo* PrimaryView = static_cast<const FViewInfo*>(View.Family->Views[0]);
-					ColorGradingTexture = TryRegisterExternalTexture(GraphBuilder, PrimaryView->GetTonemappingLUT());
-				}
-
 				FTonemapInputs PassInputs;
 				PassSequence.AcceptOverrideIfLastPass(EPass::Tonemap, PassInputs.OverrideOutput);
 				PassInputs.SceneColor = SceneColorSlice;
