[Renderer] Rename CVarRealTimeReflectionCaptureDepthBuffer to fix 1126510

diff --git a/Engine/Source/Runtime/Renderer/Private/SkyPassRendering.cpp b/Engine/Source/Runtime/Renderer/Private/SkyPassRendering.cpp
--- a/Engine/Source/Runtime/Renderer/Private/SkyPassRendering.cpp
+++ b/Engine/Source/Runtime/Renderer/Private/SkyPassRendering.cpp
@@ -244,8 +244,8 @@ void FSkyPassMeshProcessor::CollectPSOInitializersInternal(T& SkyPassShaders, co
 	const ERasterizerFillMode MeshFillMode = ComputeMeshFillMode(Material, OverrideSettings);
 	const ERasterizerCullMode MeshCullMode = ComputeMeshCullMode(Material, OverrideSettings);
 
-	static auto CVarRealTimeReflectionCaptureDepthBuffer = IConsoleManager::Get().FindConsoleVariable(TEXT("r.SkyLight.RealTimeReflectionCapture.DepthBuffer"));
-	const bool bUseDepthBuffer = CVarRealTimeReflectionCaptureDepthBuffer ? CVarRealTimeReflectionCaptureDepthBuffer->GetInt() > 0 : false;
+	static auto CVarRealTimeReflectionCaptureDepthBufferStatic = IConsoleManager::Get().FindConsoleVariable(TEXT("r.SkyLight.RealTimeReflectionCapture.DepthBuffer"));
+	const bool bUseDepthBuffer = CVarRealTimeReflectionCaptureDepthBufferStatic ? CVarRealTimeReflectionCaptureDepthBufferStatic->GetInt() > 0 : false;
 
 	FGraphicsPipelineRenderTargetsInfo RenderTargetsInfo;
 	SetupGBufferRenderTargetInfo(SceneTexturesConfig, RenderTargetsInfo, bUseDepthBuffer);
