[DOF] Fix typo in reduce pass permutation range

diff --git a/Engine/Source/Runtime/Renderer/Private/PostProcess/DiaphragmDOF.cpp b/Engine/Source/Runtime/Renderer/Private/PostProcess/DiaphragmDOF.cpp
--- a/Engine/Source/Runtime/Renderer/Private/PostProcess/DiaphragmDOF.cpp
+++ b/Engine/Source/Runtime/Renderer/Private/PostProcess/DiaphragmDOF.cpp
@@ -888,7 +888,7 @@ class FDiaphragmDOFReduceCS : public FDiaphragmDOFShader
 	DECLARE_GLOBAL_SHADER(FDiaphragmDOFReduceCS);
 	SHADER_USE_PARAMETER_STRUCT(FDiaphragmDOFReduceCS, FDiaphragmDOFShader);
 
-	class FReduceMipCount : SHADER_PERMUTATION_RANGE_INT("DIM_REDUCE_MIP_COUNT", 0, kMaxMipLevelCount);
+	class FReduceMipCount : SHADER_PERMUTATION_RANGE_INT("DIM_REDUCE_MIP_COUNT", 0, kMaxMipLevelCount + 1);
 	class FHybridScatterForeground : SHADER_PERMUTATION_BOOL("DIM_HYBRID_SCATTER_FGD");
 	class FHybridScatterBackground : SHADER_PERMUTATION_BOOL("DIM_HYBRID_SCATTER_BGD");
 	class FDisableOutputMip0 : SHADER_PERMUTATION_BOOL("DIM_DISABLE_OUTPUT_MIP0");
