Fix for BeforeDOF translucent rendering (luminance and modulate) to be composed correctly onto the scene the scene even when r.SeparateTranslucencyScreenPercentage !=100 (SeparateTranslucencyDimensions.Scale != 1.0f) and this also fixes r.Translucency.StandardSeparated 1.

#tests PC editor dx12
#rnx
[RN] Fix for BeforeDOF translucent rendering to be composed correctly onto the scene the scene even when r.SeparateTranslucencyScreenPercentage !=100 and r.Translucency.StandardSeparated is 0/1.

diff --git a/Engine/Source/Runtime/Renderer/Private/TranslucentRendering.cpp b/Engine/Source/Runtime/Renderer/Private/TranslucentRendering.cpp
--- a/Engine/Source/Runtime/Renderer/Private/TranslucentRendering.cpp
+++ b/Engine/Source/Runtime/Renderer/Private/TranslucentRendering.cpp
@@ -1492,7 +1492,7 @@ void FDeferredShadingSceneRenderer::RenderTranslucencyInner(
 	const bool bRenderInParallel = IsParallelTranslucencyEnabled();
 	const bool bIsScalingTranslucency = SeparateTranslucencyDimensions.Scale < 1.0f;
 	const bool bIsStandardSeparatedTranslucency = bStandardTranslucentCanRenderSeparate && TranslucencyPass == ETranslucencyPass::TPT_TranslucencyStandard && ViewFamily.AllowStandardTranslucencySeparated();
-	const bool bRenderInSeparateTranslucency = IsSeparateTranslucencyEnabled(TranslucencyPass, SeparateTranslucencyDimensions.Scale) || bIsStandardSeparatedTranslucency;
+	const bool bRenderInSeparateTranslucency = (IsSeparateTranslucencyEnabled(TranslucencyPass, SeparateTranslucencyDimensions.Scale) && !(TranslucencyPass == ETranslucencyPass::TPT_TranslucencyStandard || TranslucencyPass == ETranslucencyPass::TPT_TranslucencyStandardModulate)) || bIsStandardSeparatedTranslucency;
 	
 	// Holdout rendering 
 	const bool bRenderTranslucencyHold = TranslucencyPass == ETranslucencyPass::TPT_TranslucencyHoldout;
