Distance Fields - allow distance fields to run on GRHIDeviceIsIntegrated. This check was added a decade ago to prevent issues on some outdated drivers, but nowadays it prevents some pretty capable SM5 GPUs to run distance field features.

diff --git a/Engine/Source/Runtime/Renderer/Private/DistanceFieldAmbientOcclusion.cpp b/Engine/Source/Runtime/Renderer/Private/DistanceFieldAmbientOcclusion.cpp
--- a/Engine/Source/Runtime/Renderer/Private/DistanceFieldAmbientOcclusion.cpp
+++ b/Engine/Source/Runtime/Renderer/Private/DistanceFieldAmbientOcclusion.cpp
@@ -687,8 +687,6 @@ bool SupportsDistanceFieldAO(ERHIFeatureLevel::Type FeatureLevel, EShaderPlatfor
 		// Pre-GCN AMD cards have a driver bug that prevents the global distance field from being generated correctly
 		// Better to disable entirely than to display garbage
 		&& !GRHIDeviceIsAMDPreGCNArchitecture
-		// In case of iGPU we use the maximum feature level to differentiate between older and newer more capable iGPU
-		&& (!GRHIDeviceIsIntegrated || GMaxRHIFeatureLevel >= ERHIFeatureLevel::SM6)
 		&& DoesPlatformSupportDistanceFieldAO(ShaderPlatform)
 		&& IsUsingDistanceFields(ShaderPlatform);
 }
@@ -807,12 +805,6 @@ bool ShouldPrepareDistanceFieldScene(const FSceneRenderUpdateInputs& SceneUpdate
 		return false;
 	}
 
-	if (GRHIDeviceIsIntegrated && GMaxRHIFeatureLevel < ERHIFeatureLevel::SM6)
-	{
-		// In case of iGPU we use the maximum feature level to differentiate between older and newer more capable iGPU
-		return false;
-	}
-
 	if (ShouldPrepareForDistanceFieldAO(SceneUpdateInputs))
 	{
 		return true;
