[HWRT] Prevent RDG warnings about InstanceBuffer being accessed even though it has been written into.

diff --git a/Engine/Source/Runtime/Renderer/Private/DeferredShadingRenderer.cpp b/Engine/Source/Runtime/Renderer/Private/DeferredShadingRenderer.cpp
--- a/Engine/Source/Runtime/Renderer/Private/DeferredShadingRenderer.cpp
+++ b/Engine/Source/Runtime/Renderer/Private/DeferredShadingRenderer.cpp
@@ -3701,7 +3701,7 @@ void FDeferredShadingSceneRenderer::Render(FRDGBuilder& GraphBuilder, const FSce
 		FPathTracingResources PathTracingResources;
 
 #if RHI_RAYTRACING
-		if (IsRayTracingEnabled())
+		if (FamilyPipelineState[&FFamilyPipelineState::bRayTracing])
 		{
 			// Path tracer requires the full ray tracing pipeline support, as well as specialized extra shaders.
 			// Most of the ray tracing debug visualizations also require the full pipeline, but some support inline mode.
diff --git a/Engine/Source/Runtime/Renderer/Private/RayTracing/RayTracingScene.cpp b/Engine/Source/Runtime/Renderer/Private/RayTracing/RayTracingScene.cpp
--- a/Engine/Source/Runtime/Renderer/Private/RayTracing/RayTracingScene.cpp
+++ b/Engine/Source/Runtime/Renderer/Private/RayTracing/RayTracingScene.cpp
@@ -432,6 +432,16 @@ void FRayTracingScene::Update(FRDGBuilder& GraphBuilder, FSceneUniformBuffer& Sc
 						});
 				}
 			}
+			else
+			{
+				// when RayTracingBuildInstanceBuffer pass is skipped need to clear buffers so RDG doesn't complain they have never been written to
+				AddClearUAVPass(GraphBuilder, GraphBuilder.CreateUAV(LayerView.InstanceBuffer), 0, ComputePassFlags);
+				
+				if (GRHIGlobals.RayTracing.RequiresSeparateHitGroupContributionsBuffer)
+				{
+					AddClearUAVPass(GraphBuilder, GraphBuilder.CreateUAV(LayerView.HitGroupContributionsBuffer), 0, ComputePassFlags);
+				}
+			}
 
 			// Feedback
 			if (bTracingFeedbackEnabled)
