LFV - added per instance start distance.

#tests PC editor dx12
[RN] Local Fog Volume - added per volume start distance control. Changed r.LocalFogVolume.GlobalStartDistance to be a more reasonable 10m distance.

diff --git a/Engine/Shaders/Private/LocalFogVolumes/LocalFogVolumeCommon.ush b/Engine/Shaders/Private/LocalFogVolumes/LocalFogVolumeCommon.ush
--- a/Engine/Shaders/Private/LocalFogVolumes/LocalFogVolumeCommon.ush
+++ b/Engine/Shaders/Private/LocalFogVolumes/LocalFogVolumeCommon.ush
@@ -17,6 +17,7 @@ struct FLocalFogVolumeGPUInstanceData
 	float3		TranslatedWorlPos;
 	float		UniformScale;
 	float		UniformScaleInv;		// Derived from UniformScale
+	float		StartDistance;
 
 	float		RadialFogExtinction;
 	float		HeightFogExtinction;
@@ -108,6 +109,7 @@ FLocalFogVolumeGPUInstanceData GetLocalFogVolumeGPUInstanceData(uint Index)
 	float3 ZVec = cross(XVec, YVec);
 
 	Data.PreTranslatedInvTransform = float3x3(float3(XVec * Data.UniformScaleInv), float3(YVec * Data.UniformScaleInv), float3(ZVec * Data.UniformScaleInv));
+	Data.StartDistance		= Data1.w;		// Could be 2xfp16
 
 	// Data2
 	float3 Data2X			= UnpackR11G11B10F(asuint(Data2.x));
@@ -128,7 +130,7 @@ FLocalFogVolumeGPUInstanceData GetLocalFogVolumeGPUInstanceData(uint Index)
 	float4 Data2Z			= UnpackRGBA8(asuint(Data2.z));
 	Data.Albedo				= Data2Z.rgb;
 	Data.PhaseG				= Data2Z.a;
-	Data.HeightFogOffset	= Data2.w;
+	Data.HeightFogOffset	= Data2.w;		// Could be 2xfp16
 
 	return Data;
 }
@@ -351,7 +353,8 @@ float4 GetLFVInstanceContribution(
 	float2 TsU = RayIntersectSphere(RayPosU, RayDirU, float4(0.0, 0.0, 0.0, 1.0));
 
 	//Clamp to make sure we only consider intersection in front of the camera
-	TsU = max(LFVGlobalStartDistance.xx * FogInstance.UniformScaleInv, TsU);
+	const float LFVStartDistance = max(LFVGlobalStartDistance, FogInstance.StartDistance);
+	TsU = max(LFVStartDistance * FogInstance.UniformScaleInv, TsU);
 
 	// Clamp the traced distance to the depth buffer and compute world length.
 	float3 ViewToDepthVector = DepthBufferPosU - RayPosU;
diff --git a/Engine/Shaders/Private/VolumetricFog.usf b/Engine/Shaders/Private/VolumetricFog.usf
--- a/Engine/Shaders/Private/VolumetricFog.usf
+++ b/Engine/Shaders/Private/VolumetricFog.usf
@@ -157,13 +157,6 @@ void MaterialSetupCS(
 	{
 		const float3 FrontCellTranslatedWorldPosition	= ComputeCellTranslatedWorldPosition(GridCoordinate, float3(CellOffset, CellOffset, 0.0f), SceneDepth);
 		const float DistanceToCellFront					= GetDistanceToCameraFromViewVector(FrontCellTranslatedWorldPosition - PrimaryView.TranslatedWorldCameraOrigin);
-		float SoftDensityScale = 1;
-		if (DistanceToCellFront < LFVGlobalStartDistance && LFVGlobalStartDistance < DistanceToCellBack)
-		{
-			// Smoothly transition from the front clip planes of LFV to avoid bad straircase/aliasing visual effect.
-			// Extinction is basically scaled as a fonction of the start distance position inside the froxel along V.
-			SoftDensityScale = (DistanceToCellBack - LFVGlobalStartDistance) / max(0.0001, DistanceToCellBack - DistanceToCellFront);
-		}
 
 		float4 VertexClipSpacePosition = mul(float4(TranslatedWorldPosition.xyz, 1), PrimaryView.TranslatedWorldToClip);
 		float2 SvPosition = (VertexClipSpacePosition.xy / VertexClipSpacePosition.w * float2(.5f, -.5f) + .5f) * PrimaryView.ViewSizeAndInvSize.xy;
@@ -178,6 +171,23 @@ void MaterialSetupCS(
 			// LFV_TODO skip if out of sphere of influence
 			FLocalFogVolumeGPUInstanceData FogInstance = GetLocalFogVolumeGPUInstanceData(InstanceIndex);
 
+			float SoftDensityScale = 1;
+			const float LFVStartDistance = max(LFVGlobalStartDistance, FogInstance.StartDistance);
+
+			if (DistanceToCellBack < LFVStartDistance)
+			{
+				continue;
+			}
+
+			if (DistanceToCellFront < LFVStartDistance && LFVStartDistance < DistanceToCellBack)
+			{
+				// Smoothly transition from the front clip planes of LFV to avoid bad straircase/aliasing visual effect.
+				// Extinction is basically scaled as a fonction of the start distance position inside the froxel along V.
+				SoftDensityScale = (DistanceToCellBack - LFVStartDistance) / max(0.0001, DistanceToCellBack - DistanceToCellFront);
+			}
+
+			const float FogInstanceSoftDensityScale = SoftDensityScale * (DistanceToCellBack - FogInstance.StartDistance) / max(0.0001, DistanceToCellBack - DistanceToCellFront);
+
 			float3 UnitSpacePos = mul(TranslatedWorldPosition - FogInstance.TranslatedWorlPos, FogInstance.PreTranslatedInvTransform);
 
 			float DistToCenter = length(UnitSpacePos - float3(0.0f, 0.0f, 0.0f));
diff --git a/Engine/Source/Runtime/Engine/Classes/Components/LocalFogVolumeComponent.h b/Engine/Source/Runtime/Engine/Classes/Components/LocalFogVolumeComponent.h
--- a/Engine/Source/Runtime/Engine/Classes/Components/LocalFogVolumeComponent.h
+++ b/Engine/Source/Runtime/Engine/Classes/Components/LocalFogVolumeComponent.h
@@ -40,10 +40,14 @@ class ULocalFogVolumeComponent : public USceneComponent
 	UPROPERTY(EditAnywhere, BlueprintReadOnly, interp, Category = "Shading", meta = (ClampMin = 0.0, ClampMax = 1.0))
 	FLinearColor FogAlbedo = FLinearColor::White;
 
-	/** Controls the emissive color of this fog volume. */
-	UPROPERTY(EditAnywhere, BlueprintReadOnly, interp, Category = "Shading", meta = (ClampMin = 0.0, ClampMax = 1.0))
+	/** Controls the emissive color of this fog volume. (Unit = Luminance) */
+	UPROPERTY(EditAnywhere, BlueprintReadOnly, interp, Category = "Shading")
 	FLinearColor FogEmissive = FLinearColor::Black;
 
+	/** Controls the distance after which this volume starts to appear. (Unit = Centimeter) */
+	UPROPERTY(EditAnywhere, BlueprintReadOnly, interp, Category = "Shading", meta = (DisplayName = "Start Distance", UIMin = "0.0", ClampMin = 0.0))
+	float FogStartDistance = 0.0f;
+
 	/** The priority can be used as a way to override the sorting by distance. A lower value means the volume will be considered further away, i.e. it will draw behind the one with a higher priority value. */
 	UPROPERTY(EditAnywhere, BlueprintReadOnly, Category = "Sorting", meta = (UIMin = "-127", UIMax = "127", ClampMin = -127, ClampMax = 127))
 	int32 FogSortPriority = 0;
@@ -59,6 +63,8 @@ class ULocalFogVolumeComponent : public USceneComponent
 	UFUNCTION(BlueprintCallable, Category = "Rendering", meta = (DisplayName = "Set Scattering Distribution"))
 	ENGINE_API void SetFogPhaseG(float NewValue);
 	UFUNCTION(BlueprintCallable, Category = "Rendering")
+	ENGINE_API void SetFogStartDistance(float NewValue);
+	UFUNCTION(BlueprintCallable, Category = "Rendering")
 	ENGINE_API void SetFogAlbedo(FLinearColor NewValue);
 	UFUNCTION(BlueprintCallable, Category = "Rendering")
 	ENGINE_API void SetFogEmissive(FLinearColor NewValue);
diff --git a/Engine/Source/Runtime/Engine/Private/Components/LocalFogVolumeComponent.cpp b/Engine/Source/Runtime/Engine/Private/Components/LocalFogVolumeComponent.cpp
--- a/Engine/Source/Runtime/Engine/Private/Components/LocalFogVolumeComponent.cpp
+++ b/Engine/Source/Runtime/Engine/Private/Components/LocalFogVolumeComponent.cpp
@@ -119,6 +119,7 @@ LFV_BLUEPRINT_SETFUNCTION(float, HeightFogExtinction);
 LFV_BLUEPRINT_SETFUNCTION(float, HeightFogFalloff);
 LFV_BLUEPRINT_SETFUNCTION(float, HeightFogOffset);
 LFV_BLUEPRINT_SETFUNCTION(float, FogPhaseG);
+LFV_BLUEPRINT_SETFUNCTION(float, FogStartDistance);
 LFV_BLUEPRINT_SETFUNCTION_LINEARCOLOR01(FogAlbedo);
 LFV_BLUEPRINT_SETFUNCTION_LINEARCOLOR(FogEmissive);
 
diff --git a/Engine/Source/Runtime/Engine/Private/LocalFogVolumeSceneProxy.cpp b/Engine/Source/Runtime/Engine/Private/LocalFogVolumeSceneProxy.cpp
--- a/Engine/Source/Runtime/Engine/Private/LocalFogVolumeSceneProxy.cpp
+++ b/Engine/Source/Runtime/Engine/Private/LocalFogVolumeSceneProxy.cpp
@@ -16,6 +16,7 @@ FLocalFogVolumeSceneProxy::FLocalFogVolumeSceneProxy(const ULocalFogVolumeCompon
 	, HeightFogOffset(InComponent->HeightFogOffset)
 	, FogUniformScale(1.0f)
 	, FogSortPriority(uint8(127 - int8(InComponent->FogSortPriority))) // FogSortPriority on the component is in [-127,127] and needs to be negated to match expected priority behavior.
+	, FogStartDistance(InComponent->FogStartDistance)
 	, FogPhaseG(InComponent->FogPhaseG)
 	, FogAlbedo(InComponent->FogAlbedo)
 	, FogEmissive(InComponent->FogEmissive)
diff --git a/Engine/Source/Runtime/Engine/Public/LocalFogVolumeSceneProxy.h b/Engine/Source/Runtime/Engine/Public/LocalFogVolumeSceneProxy.h
--- a/Engine/Source/Runtime/Engine/Public/LocalFogVolumeSceneProxy.h
+++ b/Engine/Source/Runtime/Engine/Public/LocalFogVolumeSceneProxy.h
@@ -27,6 +27,7 @@ public:
 	float FogUniformScale;
 
 	uint8 FogSortPriority;
+	float FogStartDistance;
 
 	float FogPhaseG;
 	FLinearColor FogAlbedo;
diff --git a/Engine/Source/Runtime/Renderer/Private/FogRendering.cpp b/Engine/Source/Runtime/Renderer/Private/FogRendering.cpp
--- a/Engine/Source/Runtime/Renderer/Private/FogRendering.cpp
+++ b/Engine/Source/Runtime/Renderer/Private/FogRendering.cpp
@@ -666,7 +666,8 @@ float GetViewFogCommonStartDistance(const FViewInfo& View, bool bShouldRenderVol
 
 	if (bShouldRenderLocalFogVolumes)
 	{
-		FogCommonStartDistance = FMath::Min(GetLocalFogVolumeGlobalStartDistance(), FogCommonStartDistance);
+		const float LocalFogVolumeStartDistance = GetLocalFogVolumeViewStartDistance(View);
+		FogCommonStartDistance = FMath::Min(LocalFogVolumeStartDistance, FogCommonStartDistance);
 	}
 
 	FogCommonStartDistance = FMath::Max(GetFogDefaultStartDistance(), FogCommonStartDistance);
diff --git a/Engine/Source/Runtime/Renderer/Private/LocalFogVolumeRendering.cpp b/Engine/Source/Runtime/Renderer/Private/LocalFogVolumeRendering.cpp
--- a/Engine/Source/Runtime/Renderer/Private/LocalFogVolumeRendering.cpp
+++ b/Engine/Source/Runtime/Renderer/Private/LocalFogVolumeRendering.cpp
@@ -63,8 +63,8 @@ static TAutoConsoleVariable<int32> CVarLocalFogVolumeTileDebug(
 	ECVF_RenderThreadSafe);
 
 static TAutoConsoleVariable<float> CVarLocalFogVolumeGlobalStartDistance(
-	TEXT("r.LocalFogVolume.GlobalStartDistance"), 2000.0f,
-	TEXT("The start distance in centimeter from which local fog volumes starts to appear."),
+	TEXT("r.LocalFogVolume.GlobalStartDistance"), 1000.0f,
+	TEXT("The start distance in centimeter from which local fog volumes starts to appear. Default to 10 meters. This will push away from the camera all the local fog volume visibility independenly of their own start distance."),
 	ECVF_RenderThreadSafe);
 
 static TAutoConsoleVariable<int32> CVarLocalFogVolumeUseHZB(
@@ -155,6 +155,12 @@ float GetLocalFogVolumeGlobalStartDistance()
 	return FMath::Max(10.0f, CVarLocalFogVolumeGlobalStartDistance.GetValueOnRenderThread());
 }
 
+float GetLocalFogVolumeViewStartDistance(const FViewInfo& View)
+{
+	const float FinalStartDistance = FMath::Max(GetLocalFogVolumeGlobalStartDistance(), View.LocalFogVolumeViewData.DepthBoundStartDistance);
+	return FinalStartDistance;
+}
+
 bool IsLocalFogVolumeHalfResolution()
 {
 	return CVarLocalFogVolumeHalfResolution.GetValueOnRenderThread() > 0;
@@ -191,6 +197,7 @@ void SetDummyLocalFogVolumeForView(FRDGBuilder& GraphBuilder, FViewInfo& View)
 	static_assert(sizeof(FLocalFogVolumeGPUInstanceData) == UintVec4CountInLocalFogVolumeGPUInstanceData * SizeOfUintVec4);
 
 	View.LocalFogVolumeViewData.GPUInstanceCount			= 0;
+	View.LocalFogVolumeViewData.DepthBoundStartDistance		= FLT_MAX;
 
 	static FLocalFogVolumeGPUInstanceData DummyData;	// Static data so ERDGInitialDataFlags::NoCopy can be used. 
 	View.LocalFogVolumeViewData.GPUInstanceDataBuffer		= CreateVertexBuffer(GraphBuilder, TEXT("DUMMYLocalFogVolumeGPUInstanceDataBuffer"),
@@ -406,6 +413,7 @@ void GetLocalFogVolumeViewSortingData(const FScene* Scene, const FViewInfo& View
 	Out.LocalFogVolumeInstanceCountFinal = 0;
 	Out.LocalFogVolumeGPUInstanceData = (FLocalFogVolumeGPUInstanceData*)GraphBuilder.Alloc(sizeof(FLocalFogVolumeGPUInstanceData) * Out.LocalFogVolumeInstanceCount, 16);
 	Out.LocalFogVolumeCenterPos = (FVector*)GraphBuilder.Alloc(sizeof(FVector) * Out.LocalFogVolumeInstanceCount, 16);
+	Out.LocalFogVolumeStartDistanceAndUniformScale = (FVector2f*)GraphBuilder.Alloc(sizeof(FVector2f) * Out.LocalFogVolumeInstanceCount, 16);
 	Out.LocalFogVolumeSortKeys.SetNumUninitialized(Out.LocalFogVolumeInstanceCount);
 	for (FLocalFogVolumeSceneProxy* LHF : Scene->LocalFogVolumes)
 	{
@@ -470,7 +478,7 @@ void GetLocalFogVolumeViewSortingData(const FScene* Scene, const FViewInfo& View
 		LocalFogVolumeGPUInstanceDataIt->Data1[0] = FVector2DHalf(XVec.X, XVec.Y).AsUInt32();
 		LocalFogVolumeGPUInstanceDataIt->Data1[1] = FVector2DHalf(XVec.Z, YVec.X).AsUInt32();
 		LocalFogVolumeGPUInstanceDataIt->Data1[2] = FVector2DHalf(YVec.Y, YVec.Z).AsUInt32();
-		LocalFogVolumeGPUInstanceDataIt->Data1[3] = 0; // FREE
+		LocalFogVolumeGPUInstanceDataIt->Data1[3] = AsUint32(LHF->FogStartDistance);
 
 		// All the remaining data are packed as small as possible w.r.t. their range of value.
 		LocalFogVolumeGPUInstanceDataIt->Data2[0] = AsFloat111110(	LHF->RadialFogExtinction,	LHF->HeightFogExtinction,	SafeFallOff);
@@ -479,7 +487,8 @@ void GetLocalFogVolumeViewSortingData(const FScene* Scene, const FViewInfo& View
 		LocalFogVolumeGPUInstanceDataIt->Data2[3] = AsUint32(		LHF->HeightFogOffset);
 
 		// Register the sorting data
-		Out.LocalFogVolumeCenterPos[Out.LocalFogVolumeInstanceCountFinal] = LHF->FogTransform.GetTranslation();
+		Out.LocalFogVolumeCenterPos[Out.LocalFogVolumeInstanceCountFinal]					= LHF->FogTransform.GetTranslation();
+		Out.LocalFogVolumeStartDistanceAndUniformScale[Out.LocalFogVolumeInstanceCountFinal]= FVector2f(LHF->FogStartDistance, LHF->FogUniformScale);
 
 		FLocalFogVolumeSortKey* LocalFogVolumeSortKeysIt= &Out.LocalFogVolumeSortKeys[Out.LocalFogVolumeInstanceCountFinal];
 		LocalFogVolumeSortKeysIt->FogVolume.Index		= Out.LocalFogVolumeInstanceCountFinal;
@@ -502,11 +511,21 @@ void CreateViewLocalFogVolumeBufferSRV(const FScene* Scene, FViewInfo& View, FRD
 
 	// 1. Sort all the volumes
 	const FVector ViewOrigin = View.ViewMatrices.GetViewOrigin();
+	float DepthBoundStartDistance = FLT_MAX;
 	for (uint32 i = 0; i < SortingData.LocalFogVolumeInstanceCountFinal; ++i)
 	{
-		FVector FogCenterPos = SortingData.LocalFogVolumeCenterPos[SortingData.LocalFogVolumeSortKeys[i].FogVolume.Index];	// Recovered form the original array via index because the sorting of the previous view might have changed the order.
-		float DistancetoView = float((FogCenterPos - ViewOrigin).Size());
-		SortingData.LocalFogVolumeSortKeys[i].FogVolume.Distance = *reinterpret_cast<uint32*>(&DistancetoView);
+		const uint64 Index = SortingData.LocalFogVolumeSortKeys[i].FogVolume.Index;
+		const FVector FogCenterPos = SortingData.LocalFogVolumeCenterPos[Index];	// Recovered form the original array via index because the sorting of the previous view might have changed the order.
+		const FVector2f FogStartDistanceAndUniformScale = SortingData.LocalFogVolumeStartDistanceAndUniformScale[Index];
+		const float FogStartDistance = float(FogStartDistanceAndUniformScale.X);
+		const float FogUniformScale = float(FogStartDistanceAndUniformScale.Y);
+
+		const float DistanceToView = float((FogCenterPos - ViewOrigin).Size());
+		SortingData.LocalFogVolumeSortKeys[i].FogVolume.Distance = BitCast<uint32>(DistanceToView);
+
+		// For a volume, if the camera is inside, the start closest start distance is FogStartDistance.
+		// Otherwise it is the distance to the edge of the volume.
+		DepthBoundStartDistance = FMath::Min(DepthBoundStartDistance, FMath::Max(FogStartDistance, DistanceToView - FogUniformScale));
 	}
 	SortingData.LocalFogVolumeSortKeys.Sort();
 
@@ -525,7 +544,7 @@ void CreateViewLocalFogVolumeBufferSRV(const FScene* Scene, FViewInfo& View, FRD
 		// We could also have an indirection buffer on GPU but choosing to go with the sorting + copy on CPU since it is expected to not have many local height fog volumes.
 		LocalFogVolumeGPUSortedInstanceData[i] = SortingData.LocalFogVolumeGPUInstanceData[LFVKey.FogVolume.Index];
 
-		FVector& LFVPosition = SortingData.LocalFogVolumeCenterPos[LFVKey.FogVolume.Index];
+		const FVector& LFVPosition = SortingData.LocalFogVolumeCenterPos[LFVKey.FogVolume.Index];
 		LocalFogVolumeGPUSortedInstanceCullingData[i] = FVector4f(LFVPosition.X, LFVPosition.Y, LFVPosition.Z, LocalFogVolumeGPUSortedInstanceData[i].GetUniformScale());
 	}
 
@@ -533,6 +552,8 @@ void CreateViewLocalFogVolumeBufferSRV(const FScene* Scene, FViewInfo& View, FRD
 	const uint32 AllLocalFogVolumeInstanceBytesFinal = sizeof(FLocalFogVolumeGPUInstanceData) * SortingData.LocalFogVolumeInstanceCountFinal;
 
 	View.LocalFogVolumeViewData.GPUInstanceCount			= SortingData.LocalFogVolumeInstanceCountFinal;
+	View.LocalFogVolumeViewData.DepthBoundStartDistance		= DepthBoundStartDistance;
+
 	View.LocalFogVolumeViewData.GPUInstanceDataBuffer		= CreateVertexBuffer(
 		GraphBuilder, TEXT("LocalFogVolume.GPUInstanceDataBuffer"),
 		FRDGBufferDesc::CreateBufferDesc(SizeOfUintVec4, SortingData.LocalFogVolumeInstanceCountFinal * UintVec4CountInLocalFogVolumeGPUInstanceData),
@@ -541,7 +562,7 @@ void CreateViewLocalFogVolumeBufferSRV(const FScene* Scene, FViewInfo& View, FRD
 
 	View.LocalFogVolumeViewData.GPUInstanceCullingDataBuffer = CreateVertexBuffer(
 		GraphBuilder, TEXT("LocalFogVolume.GPUInstanceCullingDataBuffer"),
-		FRDGBufferDesc::CreateBufferDesc(sizeof(FVector4f), SortingData.LocalFogVolumeInstanceCountFinal * sizeof(FVector4f)),
+		FRDGBufferDesc::CreateBufferDesc(sizeof(FVector4f), SortingData.LocalFogVolumeInstanceCountFinal),
 		LocalFogVolumeGPUSortedInstanceCullingData, SortingData.LocalFogVolumeInstanceCountFinal * sizeof(FVector4f), ERDGInitialDataFlags::NoCopy);
 	View.LocalFogVolumeViewData.GPUInstanceCullingDataBufferSRV = GraphBuilder.CreateSRV(View.LocalFogVolumeViewData.GPUInstanceCullingDataBuffer, PF_A32B32G32R32F); // LFV_TODO use byte buffer to leverage scalar pipe in the culling compute shader.
 
@@ -837,10 +858,13 @@ void RenderLocalFogVolume(
 				continue;
 			}
 
+			// This distance is the start distance account for global and all local fog volume start distance from the camera.
+			const float FinalStartDistance = GetLocalFogVolumeViewStartDistance(View);
+
 			const FIntRect ViewRect = View.ViewRect;
 			const FMatrix ViewProjectionMatrix = View.ViewMatrices.GetProjectionMatrix();
 			const FMatrix ViewInvProjectionMatrix = View.ViewMatrices.GetInvProjectionMatrix();
-			FDepthBoundSetup DepthBoundSetup = GetDepthBoundSetup(GetLocalFogVolumeGlobalStartDistance(), ViewProjectionMatrix, ViewInvProjectionMatrix);
+			FDepthBoundSetup DepthBoundSetup = GetDepthBoundSetup(FinalStartDistance, ViewProjectionMatrix, ViewInvProjectionMatrix);
 
 			FLocalFogVolumeTiledPassParameters* PassParameters = GraphBuilder.AllocParameters<FLocalFogVolumeTiledPassParameters>();
 
@@ -903,6 +927,31 @@ void RenderLocalFogVolume(
 				RHICmdList.SetStreamSource(0, nullptr, 0);
 				RHICmdList.DrawPrimitiveIndirect(PassParameters->TileDrawIndirectBuffer->GetIndirectRHICallBuffer(), 0);
 			});
+
+			if (bVisualizationPass)
+			{
+				const float LFVGlobalStartDistance = GetLocalFogVolumeGlobalStartDistance();
+				const float LFVInstancesStartDistance = View.LocalFogVolumeViewData.DepthBoundStartDistance;
+				// Now debug print
+				AddDrawCanvasPass(GraphBuilder, {}, View, FScreenPassRenderTarget(SceneColorTexture, ERenderTargetLoadAction::ELoad),
+					[&View, LFVGlobalStartDistance, LFVInstancesStartDistance, FinalStartDistance](FCanvas& Canvas)
+				{
+					FString Text;
+					const float DPIScale = Canvas.GetDPIScale();
+					Canvas.SetBaseTransform(FMatrix(FScaleMatrix(DPIScale) * Canvas.CalcBaseTransform2D(Canvas.GetViewRect().Width(), Canvas.GetViewRect().Height())));
+
+					Canvas.DrawShadowedString(50, 50.0f, TEXT("LOCAL FOG VOLUMES"), GEngine->GetLargeFont(), FLinearColor::White);
+
+					Text = FString::Printf(TEXT("Global Start Distance    = %.1f"), LFVGlobalStartDistance);
+					Canvas.DrawShadowedString(50, 70.0f, *Text, GEngine->GetLargeFont(), FLinearColor::White);
+
+					Text = FString::Printf(TEXT("Instances Start Distance = %.1f"), LFVInstancesStartDistance);
+					Canvas.DrawShadowedString(50, 85.0f, *Text, GEngine->GetLargeFont(), FLinearColor::White);
+
+					Text = FString::Printf(TEXT("Final Start Distance     = %.1f"), FinalStartDistance);
+					Canvas.DrawShadowedString(50, 100.0f, *Text, GEngine->GetLargeFont(), FLinearColor::White);
+				});
+			}
 		}
 	}
 }
@@ -1072,7 +1121,10 @@ void RenderLocalFogVolumeMobile(
 	GraphicsPSOInit.BoundShaderState.VertexShaderRHI = VertexShader.GetVertexShader();
 	GraphicsPSOInit.BoundShaderState.PixelShaderRHI = PixelShader.GetPixelShader();
 
-	FDepthBoundSetup DepthBoundSetup = GetDepthBoundSetup(GetLocalFogVolumeGlobalStartDistance(), View.ViewMatrices.GetProjectionMatrix(), View.ViewMatrices.GetInvProjectionMatrix());
+	// This distance is the start distance account for global and all local fog volume start distance from the camera.
+	const float FinalStartDistance = GetLocalFogVolumeViewStartDistance(View);
+
+	FDepthBoundSetup DepthBoundSetup = GetDepthBoundSetup(FinalStartDistance, View.ViewMatrices.GetProjectionMatrix(), View.ViewMatrices.GetInvProjectionMatrix());
 	GraphicsPSOInit.bDepthBounds = DepthBoundSetup.bEnabled;
 
 	SetGraphicsPipelineState(RHICmdList, GraphicsPSOInit, 0);
diff --git a/Engine/Source/Runtime/Renderer/Private/LocalFogVolumeRendering.h b/Engine/Source/Runtime/Renderer/Private/LocalFogVolumeRendering.h
--- a/Engine/Source/Runtime/Renderer/Private/LocalFogVolumeRendering.h
+++ b/Engine/Source/Runtime/Renderer/Private/LocalFogVolumeRendering.h
@@ -49,7 +49,7 @@ public:
 	uint32 Data1[4];
 	uint32 Data2[4];
 
-	float GetUniformScale() 
+	float GetUniformScale()
 	{
 		union { uint32 U; float F; } FU = { Data0[3] };
 		return FU.F;
@@ -85,16 +85,18 @@ public:
 // This can be generated once for a scene and then shared between all the views to generate the GPU buffer data needed for rendering while accounting for sorting.
 struct FLocalFogVolumeSortingData
 {
-	uint32 LocalFogVolumeInstanceCount;
-	uint32 LocalFogVolumeInstanceCountFinal;
+	uint32 LocalFogVolumeInstanceCount;			// Total count in the scene
+	uint32 LocalFogVolumeInstanceCountFinal;	// Count after culling on CPU has happened from GetLocalFogVolumeViewSortingData.
 	FLocalFogVolumeGPUInstanceData* LocalFogVolumeGPUInstanceData;
 	FVector* LocalFogVolumeCenterPos;
+	FVector2f* LocalFogVolumeStartDistanceAndUniformScale;
 	TArray<FLocalFogVolumeSortKey>	LocalFogVolumeSortKeys;
 };
 
 struct FLocalFogVolumeViewData
 {
 	uint32 GPUInstanceCount = 0;
+	float DepthBoundStartDistance = 0.0f;
 	FRDGBufferRef GPUInstanceDataBuffer = nullptr;
 	FRDGBufferSRVRef GPUInstanceDataBufferSRV = nullptr;
 
@@ -134,6 +136,7 @@ bool ShouldRenderLocalFogVolumeDuringHeightFogPass(const FScene* Scene, const FS
 bool ShouldRenderLocalFogVolumeInVolumetricFog(const FScene* Scene, const FSceneViewFamily& SceneViewFamily, bool bShouldRenderVolumetricFog);
 bool ShouldRenderLocalFogVolumeVisualizationPass(const FScene* Scene, const FSceneViewFamily& SceneViewFamily);
 float GetLocalFogVolumeGlobalStartDistance();
+float GetLocalFogVolumeViewStartDistance(const FViewInfo& View);	// Account for global and all local fog volumes start distances.
 bool IsLocalFogVolumeHalfResolution();
 
 void GetLocalFogVolumeViewSortingData(const FScene* Scene, const FViewInfo& View, FRDGBuilder& GraphBuilder, FLocalFogVolumeSortingData& Out);
