OIT - fixed naming of private enum OITPass_SeparateTranslucency.

#rnx
#tests PC editor

diff --git a/Engine/Source/Runtime/Renderer/Private/MeshDrawCommands.cpp b/Engine/Source/Runtime/Renderer/Private/MeshDrawCommands.cpp
--- a/Engine/Source/Runtime/Renderer/Private/MeshDrawCommands.cpp
+++ b/Engine/Source/Runtime/Renderer/Private/MeshDrawCommands.cpp
@@ -1129,7 +1129,7 @@ public:
 				}
 				else
 				{
-					bInverseSorting &= OIT::IsSortedPixelsEnabledForPass(OITPass_SeperateTranslucency);
+					bInverseSorting &= OIT::IsSortedPixelsEnabledForPass(OITPass_SeparateTranslucency);
 				}
 
 				UpdateTranslucentMeshSortKeys(
diff --git a/Engine/Source/Runtime/Renderer/Private/OIT/OIT.cpp b/Engine/Source/Runtime/Renderer/Private/OIT/OIT.cpp
--- a/Engine/Source/Runtime/Renderer/Private/OIT/OIT.cpp
+++ b/Engine/Source/Runtime/Renderer/Private/OIT/OIT.cpp
@@ -158,7 +158,7 @@ static void AddOITPixelDebugPass(
 
 	FComputeShaderUtils::AddPass(
 		GraphBuilder,
-		RDG_EVENT_NAME("Translucency::OITDebug(Pixel,%s)", !!(OITData.PassType & OITPass_SeperateTranslucency) ? TEXT("SeparateTransluency") : TEXT("Regular")),
+		RDG_EVENT_NAME("Translucency::OITDebug(Pixel,%s)", !!(OITData.PassType & OITPass_SeparateTranslucency) ? TEXT("SeparateTransluency") : TEXT("Regular")),
 		ComputeShader,
 		Parameters,
 		FIntVector(1, 1, 1));
@@ -248,7 +248,7 @@ static void AddInternalOITComposePass(
 	const FIntVector DispatchCount = FIntVector(FMath::DivideAndRoundUp(Resolution.X, GroupSize), FMath::DivideAndRoundUp(Resolution.Y, GroupSize), 1);
 	FComputeShaderUtils::AddPass(
 		GraphBuilder,
-		RDG_EVENT_NAME("Translucency::OITCombine(%s)", !!(OITData.PassType & OITPass_SeperateTranslucency) ? TEXT("SeparateTransluency") : TEXT("Regular")),
+		RDG_EVENT_NAME("Translucency::OITCombine(%s)", !!(OITData.PassType & OITPass_SeparateTranslucency) ? TEXT("SeparateTransluency") : TEXT("Regular")),
 		ComputeShader,
 		Parameters,
 		DispatchCount);
@@ -984,7 +984,7 @@ namespace OIT
 		// Allocate OIT data		
 		Out.PassType = PassTypeBits & PassType ? PassType : OITPass_None;
 		if (PassTypeBits & OITPass_RegularTranslucency)  Out.SupportedPass |= OITPass_RegularTranslucency;
-		if (PassTypeBits & OITPass_SeperateTranslucency) Out.SupportedPass |= OITPass_SeperateTranslucency;
+		if (PassTypeBits & OITPass_SeparateTranslucency) Out.SupportedPass |= OITPass_SeparateTranslucency;
 
 		if (Out.PassType != OITPass_None)
 		{
diff --git a/Engine/Source/Runtime/Renderer/Private/OIT/OIT.h b/Engine/Source/Runtime/Renderer/Private/OIT/OIT.h
--- a/Engine/Source/Runtime/Renderer/Private/OIT/OIT.h
+++ b/Engine/Source/Runtime/Renderer/Private/OIT/OIT.h
@@ -17,7 +17,7 @@ enum EOITPassType
 {
 	OITPass_None = 0,
 	OITPass_RegularTranslucency = 1,
-	OITPass_SeperateTranslucency = 2
+	OITPass_SeparateTranslucency = 2
 };
 
 enum class FTriangleSortingOrder
diff --git a/Engine/Source/Runtime/Renderer/Private/TranslucentRendering.cpp b/Engine/Source/Runtime/Renderer/Private/TranslucentRendering.cpp
--- a/Engine/Source/Runtime/Renderer/Private/TranslucentRendering.cpp
+++ b/Engine/Source/Runtime/Renderer/Private/TranslucentRendering.cpp
@@ -1546,7 +1546,7 @@ void FDeferredShadingSceneRenderer::RenderTranslucencyInner(
 				? ERenderTargetLoadAction::EClear
 				: ERenderTargetLoadAction::ELoad;
 
-			FOITData OITData = OIT::CreateOITData(GraphBuilder, View, OITPass_SeperateTranslucency);
+			FOITData OITData = OIT::CreateOITData(GraphBuilder, View, OITPass_SeparateTranslucency);
 
 			const bool bRenderTranslucencyViewInner = OIT::IsSortedPixelsEnabled(View) ? !bIsModulate : true;
 			if (bRenderTranslucencyViewInner)
@@ -1574,7 +1574,7 @@ void FDeferredShadingSceneRenderer::RenderTranslucencyInner(
 				}
 			}
 
-			if (OIT::IsSortedPixelsEnabled(View) && !!(OITData.PassType & OITPass_SeperateTranslucency))
+			if (OIT::IsSortedPixelsEnabled(View) && !!(OITData.PassType & OITPass_SeparateTranslucency))
 			{
 				// * If this pass it not modulate (i.e., TranslucencyPass != ETranslucencyPass::TPT_TranslucencyXXXModulate), then: 
 				//   * Either the result is directly composed to the the color target. 
@@ -1728,7 +1728,7 @@ void FDeferredShadingSceneRenderer::RenderTranslucencyInner(
 			// Simply load as we have already initialized the texture.
 			const ERenderTargetLoadAction SeparateTranslucencyColorLoadAction = ERenderTargetLoadAction::ELoad;
 
-			FOITData OITData = OIT::CreateOITData(GraphBuilder, View, OITPass_SeperateTranslucency);
+			FOITData OITData = OIT::CreateOITData(GraphBuilder, View, OITPass_SeparateTranslucency);
 
 			RenderTranslucencyViewInner(
 				GraphBuilder,
@@ -1752,7 +1752,7 @@ void FDeferredShadingSceneRenderer::RenderTranslucencyInner(
 				TranslucencyPassResources.DepthTexture = SharedDepthTexture;
 			}
 
-			if (OITData.PassType & OITPass_SeperateTranslucency)
+			if (OITData.PassType & OITPass_SeparateTranslucency)
 			{
 				OIT::AddOITComposePass(GraphBuilder, View, OITData, SeparateTranslucencyColorTexture.Target);
 			}
