LFV - improved debug visualisation pass to use post processing by default.

#rnx
#tests PC editor dx12

#Synced-CL 48139523

diff --git a/Engine/Source/Runtime/Renderer/Private/LocalFogVolumeRendering.cpp b/Engine/Source/Runtime/Renderer/Private/LocalFogVolumeRendering.cpp
--- a/Engine/Source/Runtime/Renderer/Private/LocalFogVolumeRendering.cpp
+++ b/Engine/Source/Runtime/Renderer/Private/LocalFogVolumeRendering.cpp
@@ -142,7 +142,8 @@ bool ShouldRenderLocalFogVolumeInVolumetricFog(const FScene* Scene, const FScene
 
 bool ShouldRenderLocalFogVolumeVisualizationPass(const FScene* Scene, const FSceneViewFamily& SceneViewFamily)
 {
-	if (ShouldRenderLocalFogVolume(Scene, SceneViewFamily))
+	if (ShouldRenderLocalFogVolume(Scene, SceneViewFamily) 
+		&& !SceneViewFamily.EngineShowFlags.VisualizeLocalFogVolumes)	// When this is set, the debug is rendered using the post process chain.
 	{
 		const FEngineShowFlags& EngineShowFlags = SceneViewFamily.EngineShowFlags;
 		return GetLocalFogVolumeTileDebugMode(EngineShowFlags) > 0;
@@ -833,125 +834,160 @@ static FDepthBoundSetup GetDepthBoundSetup(float FogStartDistance, FMatrix ViewP
 
 	return DepthBoundSetup;
 }
- 
-void RenderLocalFogVolume(
-	const FScene* Scene,
-	TArray<FViewInfo>& Views,
-	const FSceneViewFamily& SceneViewFamily,
+
+void RenderLocalFogVolumeForView(
+	const FViewInfo& View,
 	FRDGBuilder& GraphBuilder,
-	const FMinimalSceneTextures& SceneTextures,
+	int32 LocalFogVolumeTileDebugMode,
+	FRDGTextureRef SceneColorTexture,
+	FRDGTextureRef SceneDepthTexture,
 	FRDGTextureRef LightShaftOcclusionTexture,
+	TRDGUniformBufferRef<FSceneTextureUniformParameters> SceneTextures,
 	bool bVisualizationPass)
 {
-	uint32 LocalFogVolumeInstanceCount = Scene->LocalFogVolumes.Num();
-	if (LocalFogVolumeInstanceCount > 0 && ShouldRenderLocalFogVolume(Scene, SceneViewFamily))
+	const bool bRenderLocalFogVolumes = View.LocalFogVolumeViewData.GPUInstanceCount == 0;
+	if (bRenderLocalFogVolumes)
 	{
-		RDG_EVENT_SCOPE_STAT(GraphBuilder, LocalFogVolumeVolumes, "%s", bVisualizationPass ? TEXT("LocalFogVolume.Visualization") : TEXT("LocalFogVolume.Render"));
-		RDG_GPU_STAT_SCOPE(GraphBuilder, LocalFogVolumeVolumes);
+		return;
+	}
 
-		FRDGTextureRef SceneColorTexture = SceneTextures.Color.Resolve;
+	RDG_EVENT_SCOPE_CONDITIONAL_STAT(GraphBuilder, bRenderLocalFogVolumes, LocalFogVolumeVolumes, "%s", bVisualizationPass ? TEXT("LocalFogVolume.Visualization") : TEXT("LocalFogVolume.Render"));
+	RDG_GPU_STAT_SCOPE(GraphBuilder, LocalFogVolumeVolumes);
 
-		for (FViewInfo& View : Views)
-		{
-			if (View.LocalFogVolumeViewData.GPUInstanceCount == 0)
-			{
-				continue;
-			}
 
-			// This distance is the start distance account for global and all local fog volume start distance from the camera.
-			const float FinalStartDistance = GetLocalFogVolumeViewStartDistance(View);
+	// This distance is the start distance account for global and all local fog volume start distance from the camera.
+	const float FinalStartDistance = GetLocalFogVolumeViewStartDistance(View);
 
-			const FIntRect ViewRect = View.ViewRect;
-			const FMatrix ViewProjectionMatrix = View.ViewMatrices.GetProjectionMatrix();
-			const FMatrix ViewInvProjectionMatrix = View.ViewMatrices.GetInvProjectionMatrix();
-			FDepthBoundSetup DepthBoundSetup = GetDepthBoundSetup(FinalStartDistance, ViewProjectionMatrix, ViewInvProjectionMatrix);
+	const FIntRect ViewRect = View.ViewRect;
+	const FMatrix ViewProjectionMatrix = View.ViewMatrices.GetProjectionMatrix();
+	const FMatrix ViewInvProjectionMatrix = View.ViewMatrices.GetInvProjectionMatrix();
+	FDepthBoundSetup DepthBoundSetup = GetDepthBoundSetup(FinalStartDistance, ViewProjectionMatrix, ViewInvProjectionMatrix);
 
-			FLocalFogVolumeTiledPassParameters* PassParameters = GraphBuilder.AllocParameters<FLocalFogVolumeTiledPassParameters>();
+	FLocalFogVolumeTiledPassParameters* PassParameters = GraphBuilder.AllocParameters<FLocalFogVolumeTiledPassParameters>();
 
-			PassParameters->VS.View = GetShaderBinding(View.ViewUniformBuffer);
-			PassParameters->VS.LFV = View.LocalFogVolumeViewData.UniformParametersStruct;
-			PassParameters->VS.TileDataBuffer = View.LocalFogVolumeViewData.GPUTileDataBufferSRV;
-			PassParameters->VS.StartDepthZ = DepthBoundSetup.FogClipDeviceZ;
+	PassParameters->VS.View = GetShaderBinding(View.ViewUniformBuffer);
+	PassParameters->VS.LFV = View.LocalFogVolumeViewData.UniformParametersStruct;
+	PassParameters->VS.TileDataBuffer = View.LocalFogVolumeViewData.GPUTileDataBufferSRV;
+	PassParameters->VS.StartDepthZ = DepthBoundSetup.FogClipDeviceZ;
 
-			PassParameters->PS.View = GetShaderBinding(View.ViewUniformBuffer);
-			PassParameters->PS.LFV = View.LocalFogVolumeViewData.UniformParametersStruct;
-			PassParameters->PS.LocalFogVolumeTileDebug = GetLocalFogVolumeTileDebugMode(SceneViewFamily.EngineShowFlags);
+	PassParameters->PS.View = GetShaderBinding(View.ViewUniformBuffer);
+	PassParameters->PS.LFV = View.LocalFogVolumeViewData.UniformParametersStruct;
+	PassParameters->PS.LocalFogVolumeTileDebug = LocalFogVolumeTileDebugMode;
 
-			PassParameters->SceneTextures = SceneTextures.UniformBuffer;
-			PassParameters->TileDrawIndirectBuffer = View.LocalFogVolumeViewData.GPUTileDrawIndirectBuffer;
-			PassParameters->RenderTargets[0] = FRenderTargetBinding(SceneColorTexture, ERenderTargetLoadAction::ENoAction);
-			PassParameters->RenderTargets.DepthStencil = FDepthStencilBinding(SceneTextures.Depth.Target, ERenderTargetLoadAction::ELoad, ERenderTargetLoadAction::ELoad, FExclusiveDepthStencil::DepthRead_StencilWrite);
+	PassParameters->SceneTextures = SceneTextures;
+	PassParameters->TileDrawIndirectBuffer = View.LocalFogVolumeViewData.GPUTileDrawIndirectBuffer;
+	PassParameters->RenderTargets[0] = FRenderTargetBinding(SceneColorTexture, ERenderTargetLoadAction::ENoAction);
+	PassParameters->RenderTargets.DepthStencil = FDepthStencilBinding(SceneDepthTexture, ERenderTargetLoadAction::ELoad, ERenderTargetLoadAction::ELoad, FExclusiveDepthStencil::DepthRead_StencilWrite);
 
-			FLocalFogVolumeTiledRenderVS::FPermutationDomain VSPermutationVector;
-			auto VertexShader = View.ShaderMap->GetShader< FLocalFogVolumeTiledRenderVS >(VSPermutationVector);
+	FLocalFogVolumeTiledRenderVS::FPermutationDomain VSPermutationVector;
+	auto VertexShader = View.ShaderMap->GetShader< FLocalFogVolumeTiledRenderVS >(VSPermutationVector);
 
-			FLocalFogVolumeTiledRenderPS::FPermutationDomain PsPermutationVector;
-			PsPermutationVector.Set<FLocalFogVolumeTiledRenderPS::FVisualizationPass>(bVisualizationPass);
-			auto PixelShader = View.ShaderMap->GetShader< FLocalFogVolumeTiledRenderPS >(PsPermutationVector);
+	FLocalFogVolumeTiledRenderPS::FPermutationDomain PsPermutationVector;
+	PsPermutationVector.Set<FLocalFogVolumeTiledRenderPS::FVisualizationPass>(bVisualizationPass);
+	auto PixelShader = View.ShaderMap->GetShader< FLocalFogVolumeTiledRenderPS >(PsPermutationVector);
 
-			ClearUnusedGraphResources(VertexShader, &PassParameters->VS);
-			ClearUnusedGraphResources(PixelShader, &PassParameters->PS);
+	ClearUnusedGraphResources(VertexShader, &PassParameters->VS);
+	ClearUnusedGraphResources(PixelShader, &PassParameters->PS);
 
-			FUintVector2& LocalFogVolumeTileDataTextureResolution = View.LocalFogVolumeViewData.UniformParametersStruct.LocalFogVolumeCommon.LocalFogVolumeTileDataTextureResolution;
+	FUintVector2& LocalFogVolumeTileDataTextureResolution = View.LocalFogVolumeViewData.UniformParametersStruct.LocalFogVolumeCommon.LocalFogVolumeTileDataTextureResolution;
 
-			GraphBuilder.AddPass(
-				RDG_EVENT_NAME("LocalFogVolume.Tiled (%u x %u)", LocalFogVolumeTileDataTextureResolution.X, LocalFogVolumeTileDataTextureResolution.Y),
-				PassParameters,
-				ERDGPassFlags::Raster,
-				[VertexShader, PixelShader, PassParameters, ViewRect, DepthBoundSetup](FRDGAsyncTask, FRHICommandList& RHICmdList)
-			{
-				FGraphicsPipelineStateInitializer GraphicsPSOInit;
-				RHICmdList.ApplyCachedRenderTargets(GraphicsPSOInit);
+	GraphBuilder.AddPass(
+		RDG_EVENT_NAME("LocalFogVolume.Tiled (%u x %u)", LocalFogVolumeTileDataTextureResolution.X, LocalFogVolumeTileDataTextureResolution.Y),
+		PassParameters,
+		ERDGPassFlags::Raster,
+		[VertexShader, PixelShader, PassParameters, ViewRect, DepthBoundSetup](FRDGAsyncTask, FRHICommandList& RHICmdList)
+		{
+			FGraphicsPipelineStateInitializer GraphicsPSOInit;
+			RHICmdList.ApplyCachedRenderTargets(GraphicsPSOInit);
+
+			RHICmdList.SetViewport(ViewRect.Min.X, ViewRect.Min.Y, 0.0f, ViewRect.Max.X, ViewRect.Max.Y, 1.0f);
+
+			// Render back faces only since camera may intersect
+			GraphicsPSOInit.RasterizerState = TStaticRasterizerState<>::GetRHI();
+			GraphicsPSOInit.DepthStencilState = TStaticDepthStencilState<false, CF_DepthNearOrEqual>::GetRHI();
+			GraphicsPSOInit.BlendState = TStaticBlendState<CW_RGBA, BO_Add, BF_One, BF_SourceAlpha, BO_Add, BF_Zero, BF_SourceAlpha>::GetRHI();
+			GraphicsPSOInit.PrimitiveType = PT_TriangleList; // LFV_TODO check if rects are supported and use them if so
 
-				RHICmdList.SetViewport(ViewRect.Min.X, ViewRect.Min.Y, 0.0f, ViewRect.Max.X, ViewRect.Max.Y, 1.0f);
+			GraphicsPSOInit.BoundShaderState.VertexDeclarationRHI = GFilterVertexDeclaration.VertexDeclarationRHI;
+			GraphicsPSOInit.BoundShaderState.VertexShaderRHI = VertexShader.GetVertexShader();
+			GraphicsPSOInit.BoundShaderState.PixelShaderRHI = PixelShader.GetPixelShader();
 
-				// Render back faces only since camera may intersect
-				GraphicsPSOInit.RasterizerState = TStaticRasterizerState<>::GetRHI();
-				GraphicsPSOInit.DepthStencilState = TStaticDepthStencilState<false, CF_DepthNearOrEqual>::GetRHI();
-				GraphicsPSOInit.BlendState = TStaticBlendState<CW_RGBA, BO_Add, BF_One, BF_SourceAlpha, BO_Add, BF_Zero, BF_SourceAlpha>::GetRHI();
-				GraphicsPSOInit.PrimitiveType = PT_TriangleList; // LFV_TODO check if rects are supported and use them if so
+			GraphicsPSOInit.bDepthBounds = DepthBoundSetup.bEnabled;
 
-				GraphicsPSOInit.BoundShaderState.VertexDeclarationRHI = GFilterVertexDeclaration.VertexDeclarationRHI;
-				GraphicsPSOInit.BoundShaderState.VertexShaderRHI = VertexShader.GetVertexShader();
-				GraphicsPSOInit.BoundShaderState.PixelShaderRHI = PixelShader.GetPixelShader();
+			SetGraphicsPipelineState(RHICmdList, GraphicsPSOInit, 0);
 
-				GraphicsPSOInit.bDepthBounds = DepthBoundSetup.bEnabled;
+			RHICmdList.SetDepthBounds(DepthBoundSetup.MinDeviceZ, DepthBoundSetup.MaxDeviceZ);
 
-				SetGraphicsPipelineState(RHICmdList, GraphicsPSOInit, 0);
+			SetShaderParameters(RHICmdList, VertexShader, VertexShader.GetVertexShader(), PassParameters->VS);
+			SetShaderParameters(RHICmdList, PixelShader, PixelShader.GetPixelShader(), PassParameters->PS);
 
-				RHICmdList.SetDepthBounds(DepthBoundSetup.MinDeviceZ, DepthBoundSetup.MaxDeviceZ);
+			RHICmdList.SetStreamSource(0, nullptr, 0);
+			RHICmdList.DrawPrimitiveIndirect(PassParameters->TileDrawIndirectBuffer->GetIndirectRHICallBuffer(), 0);
+		});
 
-				SetShaderParameters(RHICmdList, VertexShader, VertexShader.GetVertexShader(), PassParameters->VS);
-				SetShaderParameters(RHICmdList, PixelShader, PixelShader.GetPixelShader(), PassParameters->PS);
+	if (bVisualizationPass)
+	{
+		const float LFVGlobalStartDistance = GetLocalFogVolumeGlobalStartDistance();
+		const float LFVInstancesStartDistance = View.LocalFogVolumeViewData.DepthBoundStartDistance;
+		// Now debug print
+		AddDrawCanvasPass(GraphBuilder, {}, View, FScreenPassRenderTarget(SceneColorTexture, ERenderTargetLoadAction::ELoad),
+			[&View, LFVGlobalStartDistance, LFVInstancesStartDistance, FinalStartDistance](FCanvas& Canvas)
+			{
+				FString Text;
+				const float DPIScale = Canvas.GetDPIScale();
+				Canvas.SetBaseTransform(FMatrix(FScaleMatrix(DPIScale) * Canvas.CalcBaseTransform2D(Canvas.GetViewRect().Width(), Canvas.GetViewRect().Height())));
 
-				RHICmdList.SetStreamSource(0, nullptr, 0);
-				RHICmdList.DrawPrimitiveIndirect(PassParameters->TileDrawIndirectBuffer->GetIndirectRHICallBuffer(), 0);
+				Canvas.DrawShadowedString(50.0, 50.0, TEXT("LOCAL FOG VOLUME DEBUG VIEW"), GEngine->GetLargeFont(), FLinearColor::White);
+
+				Canvas.DrawShadowedString(50.0, 70.0, TEXT("Global Start Distance"), GEngine->GetLargeFont(), FLinearColor::White);
+				Text = FString::Printf(TEXT("= %.1f"), LFVGlobalStartDistance);
+				Canvas.DrawShadowedString(200.0, 70.0, *Text, GEngine->GetLargeFont(), FLinearColor::White);
+
+				Canvas.DrawShadowedString(50.0, 85.0, TEXT("Instances Start Distance"), GEngine->GetLargeFont(), FLinearColor::White);
+				Text = FString::Printf(TEXT("= %.1f"), LFVInstancesStartDistance);
+				Canvas.DrawShadowedString(200.0, 85.0, *Text, GEngine->GetLargeFont(), FLinearColor::White);
+
+				Canvas.DrawShadowedString(50.0, 100.0, TEXT("Final Start Distance"), GEngine->GetLargeFont(), FLinearColor::White);
+				Text = FString::Printf(TEXT("= %.1f"), FinalStartDistance);
+				Canvas.DrawShadowedString(200.0, 100.0, *Text, GEngine->GetLargeFont(), FLinearColor::White);
 			});
+	}
+}
+ 
+void RenderLocalFogVolume(
+	const FScene* Scene,
+	TArray<FViewInfo>& Views,
+	const FSceneViewFamily& SceneViewFamily,
+	FRDGBuilder& GraphBuilder,
+	const FMinimalSceneTextures& SceneTextures,
+	FRDGTextureRef LightShaftOcclusionTexture,
+	bool bVisualizationPass)
+{
+	uint32 LocalFogVolumeInstanceCount = Scene->LocalFogVolumes.Num();
+	if (LocalFogVolumeInstanceCount > 0 && ShouldRenderLocalFogVolume(Scene, SceneViewFamily))
+	{
+		RDG_EVENT_SCOPE_STAT(GraphBuilder, LocalFogVolumeVolumes, "%s", bVisualizationPass ? TEXT("LocalFogVolume.Visualization") : TEXT("LocalFogVolume.Render"));
+		RDG_GPU_STAT_SCOPE(GraphBuilder, LocalFogVolumeVolumes);
+
+		FRDGTextureRef SceneColorTexture = SceneTextures.Color.Resolve;
+		FRDGTextureRef SceneDepthTexture = SceneTextures.Depth.Target;
 
-			if (bVisualizationPass)
+		for (FViewInfo& View : Views)
+		{
+			if (View.LocalFogVolumeViewData.GPUInstanceCount == 0)
 			{
-				const float LFVGlobalStartDistance = GetLocalFogVolumeGlobalStartDistance();
-				const float LFVInstancesStartDistance = View.LocalFogVolumeViewData.DepthBoundStartDistance;
-				// Now debug print
-				AddDrawCanvasPass(GraphBuilder, {}, View, FScreenPassRenderTarget(SceneColorTexture, ERenderTargetLoadAction::ELoad),
-					[&View, LFVGlobalStartDistance, LFVInstancesStartDistance, FinalStartDistance](FCanvas& Canvas)
-				{
-					FString Text;
-					const float DPIScale = Canvas.GetDPIScale();
-					Canvas.SetBaseTransform(FMatrix(FScaleMatrix(DPIScale) * Canvas.CalcBaseTransform2D(Canvas.GetViewRect().Width(), Canvas.GetViewRect().Height())));
-
-					Canvas.DrawShadowedString(50, 50.0f, TEXT("LOCAL FOG VOLUMES"), GEngine->GetLargeFont(), FLinearColor::White);
-
-					Text = FString::Printf(TEXT("Global Start Distance    = %.1f"), LFVGlobalStartDistance);
-					Canvas.DrawShadowedString(50, 70.0f, *Text, GEngine->GetLargeFont(), FLinearColor::White);
-
-					Text = FString::Printf(TEXT("Instances Start Distance = %.1f"), LFVInstancesStartDistance);
-					Canvas.DrawShadowedString(50, 85.0f, *Text, GEngine->GetLargeFont(), FLinearColor::White);
-
-					Text = FString::Printf(TEXT("Final Start Distance     = %.1f"), FinalStartDistance);
-					Canvas.DrawShadowedString(50, 100.0f, *Text, GEngine->GetLargeFont(), FLinearColor::White);
-				});
+				continue;
 			}
+
+			RenderLocalFogVolumeForView(
+				View,
+				GraphBuilder,
+				GetLocalFogVolumeTileDebugMode(SceneViewFamily.EngineShowFlags),
+				SceneColorTexture,
+				SceneDepthTexture,
+				LightShaftOcclusionTexture,
+				SceneTextures.UniformBuffer,
+				bVisualizationPass);
 		}
 	}
 }
@@ -966,10 +1002,62 @@ void RenderLocalFogVolumeVisualization(
 #if !UE_BUILD_SHIPPING
 	const bool bVisualizationPass = true;
 	FRDGTexture* LightShaftOcclusionTexture = GSystemTextures.GetBlackDummy(GraphBuilder);
-	RenderLocalFogVolume(Scene, Views, SceneViewFamily, GraphBuilder, SceneTextures, LightShaftOcclusionTexture, bVisualizationPass);
+	if (ShouldRenderLocalFogVolume(Scene, SceneViewFamily))
+	{
+		RenderLocalFogVolume(Scene, Views, SceneViewFamily, GraphBuilder, SceneTextures, LightShaftOcclusionTexture, bVisualizationPass);
+	}
 #endif
 }
 
+
+FScreenPassTexture RenderLocalFogVolumeVisualizationPostProcess(
+	FRDGBuilder& GraphBuilder,
+	const FViewInfo& View,
+	const FScreenPassTexture& ScreenPassSceneColor,
+	const FScreenPassTexture& SceneDepth,
+	const FScreenPassRenderTarget& OverrideOutput)
+{
+	FScreenPassRenderTarget OutputRT = OverrideOutput;
+
+#if WITH_EDITOR
+
+	if (!OutputRT.IsValid())
+	{
+		OutputRT = FScreenPassRenderTarget::CreateFromInput(GraphBuilder, ScreenPassSceneColor, View.GetOverwriteLoadAction(), TEXT("VisualizeLightFunctionAtlas"));
+	}
+
+	if (ScreenPassSceneColor.Texture != OutputRT.GetRenderTargetBinding().GetTexture())
+	{
+		// Copy the scene color because below we do not write and copy all the view pixel. We only stamp the atlas and its tiles on top.
+		AddDrawTexturePass(
+			GraphBuilder,
+			View,
+			ScreenPassSceneColor.Texture,
+			OutputRT.GetRenderTargetBinding().GetTexture(),
+			FIntPoint::ZeroValue,
+			OutputRT.GetRenderTargetBinding().GetTexture()->Desc.Extent,
+			FIntPoint::ZeroValue,
+			OutputRT.GetRenderTargetBinding().GetTexture()->Desc.Extent);
+	}
+
+	check(View.Family);
+	FRDGTexture* LightShaftOcclusionTexture = GSystemTextures.GetBlackDummy(GraphBuilder);
+	const bool bVisualizationPass = true;
+	RenderLocalFogVolumeForView(
+		View,
+		GraphBuilder,
+		GetLocalFogVolumeTileDebugMode(View.Family->EngineShowFlags),
+		OutputRT.Texture,
+		SceneDepth.Texture,
+		LightShaftOcclusionTexture,
+		View.GetSceneTextures().UniformBuffer,
+		bVisualizationPass);
+
+#endif
+
+	return OutputRT;
+}
+
 /*=============================================================================
 	Local height fog rendering - mobile
 =============================================================================*/
diff --git a/Engine/Source/Runtime/Renderer/Private/LocalFogVolumeRendering.h b/Engine/Source/Runtime/Renderer/Private/LocalFogVolumeRendering.h
--- a/Engine/Source/Runtime/Renderer/Private/LocalFogVolumeRendering.h
+++ b/Engine/Source/Runtime/Renderer/Private/LocalFogVolumeRendering.h
@@ -177,3 +177,13 @@ void RenderLocalFogVolumeVisualization(
 	FRDGBuilder& GraphBuilder,
 	const FMinimalSceneTextures& SceneTextures);
 
+struct FScreenPassTexture;
+struct FScreenPassRenderTarget;
+
+FScreenPassTexture RenderLocalFogVolumeVisualizationPostProcess(
+	FRDGBuilder& GraphBuilder,
+	const FViewInfo& View,
+	const FScreenPassTexture& ScreenPassSceneColor,
+	const FScreenPassTexture& SceneDepth,
+	const FScreenPassRenderTarget& OverrideOutput);
+
diff --git a/Engine/Source/Runtime/Renderer/Private/PostProcess/PostProcessing.cpp b/Engine/Source/Runtime/Renderer/Private/PostProcess/PostProcessing.cpp
--- a/Engine/Source/Runtime/Renderer/Private/PostProcess/PostProcessing.cpp
+++ b/Engine/Source/Runtime/Renderer/Private/PostProcess/PostProcessing.cpp
@@ -83,6 +83,9 @@
 #include "VT/VirtualTextureFeedbackResource.h"
 #include "VT/VirtualTextureVisualizationData.h"
 #include "RenderViewportFeedback.h"
+#if WITH_EDITOR
+#include "LightFunctionAtlas.h"
+#endif
 
 bool IsMobileEyeAdaptationEnabled(const FViewInfo& View);
 
@@ -459,6 +462,7 @@ void AddPostProcessingPasses(
 		VisualizeSkyAtmosphere,
 		VisualizeSkyLightIlluminanceMeter,
 		VisualizeLightFunctionAtlas,
+		VisualizeLocalFogVolumes,
 		VisualizeLevelInstance,
 		VisualizeVirtualShadowMaps_PreEditorPrimitives,
 		ChannelMask,
@@ -521,6 +525,7 @@ void AddPostProcessingPasses(
 		TEXT("VisualizeSkyAtmosphere"),
 		TEXT("VisualizeSkyLightIlluminanceMeter"),
 		TEXT("VisualizeLightFunctionAtlas"),
+		TEXT("VisualizeLocalFogVolumes"),
 		TEXT("VisualizeLevelInstance"),
 		TEXT("VisualizeVirtualShadowMaps_PreEditorPrimitives"),
 		TEXT("ChannelMask"),
@@ -567,6 +572,7 @@ void AddPostProcessingPasses(
 	PassSequence.SetEnabled(EPass::VisualizeSkyAtmosphere, GIsEditor && Scene && View.Family && View.Family->EngineShowFlags.VisualizeSkyAtmosphere && ShouldRenderSkyAtmosphereDebugPasses(Scene, View.Family->EngineShowFlags));
 	PassSequence.SetEnabled(EPass::VisualizeSkyLightIlluminanceMeter, GIsEditor && Scene && Scene->SkyLight && View.Family && View.Family->EngineShowFlags.VisualizeSkyLightIlluminance);
 	PassSequence.SetEnabled(EPass::VisualizeLightFunctionAtlas, GIsEditor && Scene && Scene->LightFunctionAtlasSceneData.GetLightFunctionAtlasEnabled() && View.Family && View.Family->EngineShowFlags.VisualizeLightFunctionAtlas);
+	PassSequence.SetEnabled(EPass::VisualizeLocalFogVolumes, GIsEditor && Scene && View.Family && ShouldRenderLocalFogVolume(Scene, *View.Family) && View.Family->EngineShowFlags.VisualizeLocalFogVolumes);
 	PassSequence.SetEnabled(EPass::VisualizeLevelInstance, GIsEditor && EngineShowFlags.EditingLevelInstance && EngineShowFlags.VisualizeLevelInstanceEditing && !bVisualizeHDR);
 	PassSequence.SetEnabled(EPass::ChannelMask, View.Family->ChannelMaskParams.ColorChannelMask != EColorChannelMask::All || View.Family->ChannelMaskParams.bDrawAlphaBlendedCheckerboard);
 	PassSequence.SetEnabled(EPass::SelectionOutline, GIsEditor && EngineShowFlags.Selection && EngineShowFlags.SelectionOutline && !EngineShowFlags.Wireframe && !bVisualizeHDR);
@@ -575,6 +581,7 @@ void AddPostProcessingPasses(
 	PassSequence.SetEnabled(EPass::VisualizeSkyAtmosphere, false);
 	PassSequence.SetEnabled(EPass::VisualizeSkyLightIlluminanceMeter, false);
 	PassSequence.SetEnabled(EPass::VisualizeLightFunctionAtlas, false);
+	PassSequence.SetEnabled(EPass::VisualizeLocalFogVolumes, false);
 	PassSequence.SetEnabled(EPass::VisualizeLevelInstance, false);
 	PassSequence.SetEnabled(EPass::ChannelMask, false);
 	PassSequence.SetEnabled(EPass::SelectionOutline, false);
@@ -1770,6 +1777,13 @@ void AddPostProcessingPasses(
 		}
 	}
 
+	if (PassSequence.IsEnabled(EPass::VisualizeLocalFogVolumes))
+	{
+		FScreenPassRenderTarget OverrideOutput;
+		PassSequence.AcceptOverrideIfLastPass(EPass::VisualizeLocalFogVolumes, OverrideOutput);
+		SceneColor = RenderLocalFogVolumeVisualizationPostProcess(GraphBuilder, View, SceneColor, SceneDepth, OverrideOutput);
+	}
+
 	if (PassSequence.IsEnabled(EPass::VisualizeLevelInstance))
 	{
 		FVisualizeLevelInstanceInputs PassInputs;
