Don't precache velocity PSOs when not needed during mobile base pass rendering (only needed when AA mode is set)

#Synced-CL 48092843

diff --git a/Engine/Source/Runtime/Renderer/Private/VelocityRendering.cpp b/Engine/Source/Runtime/Renderer/Private/VelocityRendering.cpp
--- a/Engine/Source/Runtime/Renderer/Private/VelocityRendering.cpp
+++ b/Engine/Source/Runtime/Renderer/Private/VelocityRendering.cpp
@@ -1289,6 +1289,15 @@ bool FVelocityMeshProcessor::CollectPSOInitializersInternal(
 		FVelocityPS> VelocityPassShaders;
 
 	const bool bIsTranslucentClippedDepthPass = (MeshPassType == EMeshPass::TranslucentVelocityClippedDepth);
+			
+	const EShaderPlatform ShaderPlatform = GetFeatureLevelShaderPlatform(FeatureLevel);
+
+	// Only enabled on mobile when any AA mode is enabled
+	EAntiAliasingMethod DefaultAntiAliasingMethod = GetDefaultAntiAliasingMethod(FeatureLevel);
+	if (IsMobilePlatform(ShaderPlatform) && DefaultAntiAliasingMethod == AAM_None)
+	{
+		return false;
+	}
 
 	EVelocityPassMode PassMode = EVelocityPassMode::Velocity_Standard;
 	if (PlatformSupportsOpenXRMotionVectors(GetFeatureLevelShaderPlatform(FeatureLevel)))
@@ -1312,8 +1321,6 @@ bool FVelocityMeshProcessor::CollectPSOInitializersInternal(
 		return false;
 	}
 
-	EShaderPlatform ShaderPlatform = GetFeatureLevelShaderPlatform(FeatureLevel);
-
 	FGraphicsPipelineRenderTargetsInfo RenderTargetsInfo;
 	RenderTargetsInfo.NumSamples = 1;
 	AddRenderTargetInfo(FVelocityRendering::GetFormat(ShaderPlatform), FVelocityRendering::GetCreateFlags(ShaderPlatform), RenderTargetsInfo);
