For specific RHIs, keep HZB texture resolution at the fixed viewport size (View.GetSceneTexturesConfig().Extent) in order to avoid the RT reallocation when dynamic resolution is enabled.

#Synced-CL 48072297

diff --git a/Engine/Source/Runtime/D3D12RHI/Private/Windows/WindowsD3D12Device.cpp b/Engine/Source/Runtime/D3D12RHI/Private/Windows/WindowsD3D12Device.cpp
--- a/Engine/Source/Runtime/D3D12RHI/Private/Windows/WindowsD3D12Device.cpp
+++ b/Engine/Source/Runtime/D3D12RHI/Private/Windows/WindowsD3D12Device.cpp
@@ -1892,6 +1892,8 @@ void FD3D12DynamicRHI::Init()
 		GRHIGlobals.RayTracing.SupportsPersistentSBTs = true;
 	}
 
+	GRHISupportsFastRTAllocation = false;
+
 	// Indicate that the RHI needs to use the engine's deferred deletion queue.
 	GRHINeedsExtraDeletionLatency = true;
 
diff --git a/Engine/Source/Runtime/RHI/Public/RHIGlobals.h b/Engine/Source/Runtime/RHI/Public/RHIGlobals.h
--- a/Engine/Source/Runtime/RHI/Public/RHIGlobals.h
+++ b/Engine/Source/Runtime/RHI/Public/RHIGlobals.h
@@ -752,6 +752,9 @@ PRAGMA_ENABLE_DEPRECATION_WARNINGS
 	/** True if the platform supports split barriers. */
 	bool SupportsSplitBarriers = false;
 
+	/** True if the platform has low overhead when allocating non pooled RT resources. */
+	bool SupportsFastRTAllocation = true;
+
 	/** Preferred format of FShaderCode (binary code or microcode). Only used for preview platforms with more than one shader code being supported such as D3D12 (supports DXIL and DXBC). */
 	FName PreferredPreviewShaderCodeFormat = NAME_None;
 };
@@ -934,6 +937,7 @@ extern RHI_API uint32 GGPUFrameTime;
 #define GRHISupportsLossyFramebufferCompression                GRHIGlobals.SupportsLossyFramebufferCompression
 #define GRHISupportsBindingTexArrayPerSlice					   GRHIGlobals.SupportsBindingTexArrayPerSlice
 #define GRHINeedsSRVGraphicsNonPixelWorkaround 				   GRHIGlobals.NeedsSRVGraphicsNonPixelWorkaround
+#define GRHISupportsFastRTAllocation						   GRHIGlobals.SupportsFastRTAllocation
 
 // Utility Getters
 
diff --git a/Engine/Source/Runtime/Renderer/Private/CustomDepthRendering.cpp b/Engine/Source/Runtime/Renderer/Private/CustomDepthRendering.cpp
--- a/Engine/Source/Runtime/Renderer/Private/CustomDepthRendering.cpp
+++ b/Engine/Source/Runtime/Renderer/Private/CustomDepthRendering.cpp
@@ -349,7 +349,9 @@ bool FSceneRenderer::RenderCustomDepthPass(
 			RDG_EVENT_SCOPE_CONDITIONAL(GraphBuilder, Views.Num() > 1 && bDrawSceneViewsInOneNanitePass, "View%d (together with %d more)", ViewIndex, Views.Num() - 1);
 			
 			FViewInfo& View = Views[ViewIndex];
-			FIntRect ViewRect = bDrawSceneViewsInOneNanitePass ? FIntRect(0, 0, FamilySize.X, FamilySize.Y) : View.ViewRect;
+
+			const FIntRect HZBViewRect = bDrawSceneViewsInOneNanitePass ? FIntRect(0, 0, FamilySize.X, FamilySize.Y) : GetHZBViewRect(View);
+
 			auto NaniteRenderer = Nanite::IRenderer::Create(
 				GraphBuilder,
 				*Scene,
@@ -358,7 +360,7 @@ bool FSceneRenderer::RenderCustomDepthPass(
 				SharedContext,
 				RasterContext,
 				CullingConfig,
-				ViewRect,
+				HZBViewRect,
 				/* PrevHZB = */ nullptr
 			);
 
diff --git a/Engine/Source/Runtime/Renderer/Private/DeferredShadingRenderer.cpp b/Engine/Source/Runtime/Renderer/Private/DeferredShadingRenderer.cpp
--- a/Engine/Source/Runtime/Renderer/Private/DeferredShadingRenderer.cpp
+++ b/Engine/Source/Runtime/Renderer/Private/DeferredShadingRenderer.cpp
@@ -500,7 +500,7 @@ bool FDeferredShadingSceneRenderer::RenderHzb(FRDGBuilder& GraphBuilder, FRDGTex
 				GraphBuilder,
 				SceneDepthTexture,
 				/* VisBufferTexture = */ nullptr,
-				View.ViewRect,
+				GetHZBViewRect(View),
 				View.GetFeatureLevel(),
 				View.GetShaderPlatform(),
 				TEXT("HZBClosest"),
@@ -1509,7 +1509,7 @@ void FDeferredShadingSceneRenderer::RenderNanite(FRDGBuilder& GraphBuilder, cons
 			RDG_EVENT_SCOPE_CONDITIONAL(GraphBuilder, InViews.Num() > 1 && !bDrawSceneViewsInOneNanitePass, "View%u", ViewIndex);
 			RDG_EVENT_SCOPE_CONDITIONAL(GraphBuilder, InViews.Num() > 1 && bDrawSceneViewsInOneNanitePass, "View%u (together with %d more)", ViewIndex, InViews.Num() - 1);
 
-			FIntRect ViewRect = bDrawSceneViewsInOneNanitePass ? FIntRect(0, 0, FamilySize.X, FamilySize.Y) : View.ViewRect;
+			FIntRect ViewRect = bDrawSceneViewsInOneNanitePass ? FIntRect(0, 0, FamilySize.X, FamilySize.Y) : GetHZBViewRect(View);
 			CullingConfig.SetViewFlags(View);
 
 			float LODScaleFactor = 1.0f;
diff --git a/Engine/Source/Runtime/Renderer/Private/MobileShadingRenderer.cpp b/Engine/Source/Runtime/Renderer/Private/MobileShadingRenderer.cpp
--- a/Engine/Source/Runtime/Renderer/Private/MobileShadingRenderer.cpp
+++ b/Engine/Source/Runtime/Renderer/Private/MobileShadingRenderer.cpp
@@ -3016,7 +3016,7 @@ void FMobileSceneRenderer::RenderHZB(FRDGBuilder& GraphBuilder, FRDGTextureRef S
 					GraphBuilder,
 					SceneDepthTexture,
 					/* VisBufferTexture = */ nullptr,
-					View.ViewRect,
+					GetHZBViewRect(View),
 					View.GetFeatureLevel(),
 					View.GetShaderPlatform(),
 					TEXT("MobileHZBFurthest"),
diff --git a/Engine/Source/Runtime/Renderer/Private/SceneRendering.h b/Engine/Source/Runtime/Renderer/Private/SceneRendering.h
--- a/Engine/Source/Runtime/Renderer/Private/SceneRendering.h
+++ b/Engine/Source/Runtime/Renderer/Private/SceneRendering.h
@@ -2073,6 +2073,16 @@ public:
 	 */
 	virtual FDeferredShadingSceneRenderer* GetDeferredShadingSceneRenderer() { return nullptr; }
 
+protected:
+
+	/**
+	* Keep the HZB resolution to the max view size on platforms where RT reallocation could impact performance
+	*/
+	inline FIntRect GetHZBViewRect ( const FViewInfo& View )
+	{
+		return GRHISupportsFastRTAllocation ? View.ViewRect : FIntRect(0, 0, View.GetSceneTexturesConfig().Extent.X, View.GetSceneTexturesConfig().Extent.Y);
+	}
+
 private:
 	FSceneUniformBuffer SceneUniforms;
 	FSceneExtensionsRenderers SceneExtensionsRenderers;
