[Volumetric Fog] Don't compile VIRTUAL_SHADOW_MAP permutations based on DoesPlatformSupportVirtualShadowMaps(...).

#Synced-CL 48046669

diff --git a/Engine/Source/Runtime/Renderer/Private/VolumetricFog.cpp b/Engine/Source/Runtime/Renderer/Private/VolumetricFog.cpp
--- a/Engine/Source/Runtime/Renderer/Private/VolumetricFog.cpp
+++ b/Engine/Source/Runtime/Renderer/Private/VolumetricFog.cpp
@@ -1128,6 +1128,11 @@ class FVolumetricFogLightScatteringCS : public FGlobalShader
 			PermutationVector.Set<FDistanceFieldSkyOcclusion>(false);
 		}
 
+		if (!DoesPlatformSupportVirtualShadowMaps(ShaderPlatform))
+		{
+			PermutationVector.Set<FVirtualShadowMap>(false);
+		}
+
 		return PermutationVector;
 	}
 
