[Volumetric Fog] Support ApplyDepthConstraintsToOffset(...) to avoid lighting on fog leaking through walls.

#Synced-CL 48158727

diff --git a/Engine/Shaders/Private/SceneTexturesCommon.ush b/Engine/Shaders/Private/SceneTexturesCommon.ush
--- a/Engine/Shaders/Private/SceneTexturesCommon.ush
+++ b/Engine/Shaders/Private/SceneTexturesCommon.ush
@@ -246,6 +246,16 @@ float LookupDeviceZ( float2 ScreenUV )
 	return LookupDeviceZ(ScreenUV, 0);
 }
 
+/** Returns DeviceZ which is the z value stored in the depth buffer. */
+float LookupDeviceZ(uint2 PixelPos)
+{
+#if SCENE_TEXTURES_DISABLED
+	return SCENE_TEXTURES_DISABLED_SCENE_DEPTH_VALUE;
+#else
+	return MobileSceneTextures.SceneDepthTexture.Load(int3(PixelPos, 0)).r;
+#endif
+}
+
 /** Returns 4 DeviceZ values, which is the z value stored in the depth buffer. */
 float4 GatherDeviceZ( float2 ScreenUV )
 {
@@ -283,4 +293,14 @@ float CalcSceneDepth(float2 ScreenUV)
 	return CalcSceneDepth(ScreenUV, 0);
 }
 
+/** Returns clip space W, which is world space distance along the View Z axis. */
+float CalcSceneDepth(uint2 PixelPos)
+{
+#if SCENE_TEXTURES_DISABLED
+	return SCENE_TEXTURES_DISABLED_SCENE_DEPTH_VALUE;
+#else
+	return ConvertFromDeviceZ(LookupDeviceZ(PixelPos));
+#endif
+}
+
 #endif // SHADING_PATH_MOBILE
diff --git a/Engine/Shaders/Private/VolumetricFog.usf b/Engine/Shaders/Private/VolumetricFog.usf
--- a/Engine/Shaders/Private/VolumetricFog.usf
+++ b/Engine/Shaders/Private/VolumetricFog.usf
@@ -56,8 +56,12 @@ RWTexture3D<float4> RWVBufferB;
 
 float ComputeDepthFromZSlice(float ZSlice)
 {
-	float SliceDepth = (exp2(ZSlice / VolumetricFog.GridZParams.z) - VolumetricFog.GridZParams.y) / VolumetricFog.GridZParams.x;
-	return SliceDepth;
+	return ComputeDepthFromZSlice(VolumetricFog.GridZParams, ZSlice);
+}
+
+float ComputeZSliceFromDepth(float SliceDepth)
+{
+	return ComputeZSliceFromDepth(VolumetricFog.GridZParams, SliceDepth);
 }
 
 float4x4 UnjitteredClipToTranslatedWorld;    
@@ -81,6 +85,35 @@ float3 ComputeCellTranslatedWorldPosition(uint3 GridCoordinate, float3 CellOffse
 	return ComputeCellTranslatedWorldPosition(GridCoordinate, CellOffset, Unused);
 }
 
+float GridCenterOffsetFromDepthBuffer;
+float GridCenterOffsetThresholdToAcceptDepthBufferOffset;
+
+// Applying an offset from the depth buffer helps with reducing lighting on fog leaking through walls
+void ApplyDepthConstraintsToOffset(float3 GridCoordinate, inout float3 LocalFrameJitterOffset)
+{
+	if (GridCenterOffsetFromDepthBuffer >= 0.0)
+	{
+		float Unused;
+		float4 GridSampleClip = mul(float4(ComputeCellTranslatedWorldPosition(GridCoordinate, LocalFrameJitterOffset, Unused), 1), PrimaryView.TranslatedWorldToClip);
+		float3 GridSampleScreen = GridSampleClip.xyz / GridSampleClip.w;
+		float3 GridSampleUVz = float3(GridSampleScreen.xy * View.ScreenPositionScaleBias.xy + View.ScreenPositionScaleBias.wz, GridSampleScreen.z);
+		float GridSampleSceneDepth = CalcSceneDepth(uint2(GridSampleUVz.xy * View.BufferSizeAndInvSize.xy));
+
+		float DepthGridCoordinateWithZOffset = ComputeZSliceFromDepth(GridSampleSceneDepth) - GridCenterOffsetFromDepthBuffer;
+
+		float3 GridCoordinateOffset = float3(GridCoordinate) + LocalFrameJitterOffset;
+		
+		const float DepthToSampleGridOffset = DepthGridCoordinateWithZOffset - GridCoordinateOffset.z; // <0 if samples is in front of the depth buffer
+		if (DepthToSampleGridOffset < 0.0 // Move the sample forward only if the depth buffer grid coordinate is closer,
+			&& (-DepthToSampleGridOffset < GridCenterOffsetThresholdToAcceptDepthBufferOffset) // But only if it is less than N slice away (to not bring all the sample behind forward)
+			)
+		{
+			// If the grid sample is further than the grid z for the depth buffer, adjust the offset 
+			LocalFrameJitterOffset.z += DepthToSampleGridOffset;
+		}
+	}
+}
+
 float3 RaleighScattering()
 {
 	float3 Wavelengths = float3(650.0f, 510.0f, 475.0f);
@@ -128,7 +161,9 @@ void MaterialSetupCS(
 	uint3 GridCoordinate = DispatchThreadId;
 
 	// Center of the voxel
-	float VoxelOffset = .5f;
+	float3 VoxelOffset = .5f;
+
+	ApplyDepthConstraintsToOffset(GridCoordinate, VoxelOffset);
 	 
 	float3 WorldPosition = ComputeCellTranslatedWorldPosition(GridCoordinate, VoxelOffset) - DFHackToFloat(PrimaryView.PreViewTranslation);
 
@@ -148,14 +183,17 @@ void MaterialSetupCS(
 	float3 AdditionalEmissive = 0.0f;
 
 #if USE_LOCAL_FOG_VOLUMES
-	float CellOffset = .5f;
+	float3 CellOffset = .5f;
+
+	ApplyDepthConstraintsToOffset(GridCoordinate, CellOffset);
+
 	float SceneDepth;
 	const float3 TranslatedWorldPosition			= ComputeCellTranslatedWorldPosition(GridCoordinate, CellOffset, SceneDepth);
-	const float3 BackCellTranslatedWorldPosition	= ComputeCellTranslatedWorldPosition(GridCoordinate, float3(CellOffset, CellOffset, 1.0f), SceneDepth);
+	const float3 BackCellTranslatedWorldPosition	= ComputeCellTranslatedWorldPosition(GridCoordinate, float3(CellOffset.x, CellOffset.y, 1.0f), SceneDepth);
 	const float  DistanceToCellBack					= GetDistanceToCameraFromViewVector(BackCellTranslatedWorldPosition - PrimaryView.TranslatedWorldCameraOrigin);
 	if(DistanceToCellBack > LFVGlobalStartDistance)
 	{
-		const float3 FrontCellTranslatedWorldPosition	= ComputeCellTranslatedWorldPosition(GridCoordinate, float3(CellOffset, CellOffset, 0.0f), SceneDepth);
+		const float3 FrontCellTranslatedWorldPosition	= ComputeCellTranslatedWorldPosition(GridCoordinate, float3(CellOffset.x, CellOffset.y, 0.0f), SceneDepth);
 		const float DistanceToCellFront					= GetDistanceToCameraFromViewVector(FrontCellTranslatedWorldPosition - PrimaryView.TranslatedWorldCameraOrigin);
 
 		float4 VertexClipSpacePosition = mul(float4(TranslatedWorldPosition.xyz, 1), PrimaryView.TranslatedWorldToClip);
@@ -407,7 +445,9 @@ float4 InjectShadowedLocalLightCommon(uint3 GridCoordinate)
 		for (uint SampleIndex = 0; SampleIndex < NumSuperSamples; SampleIndex++)
 		{
 			float3 CellOffset = FrameJitterOffsets[SampleIndex].xyz;
-			//float CellOffset = .5f;
+			//float3 CellOffset = .5f;
+
+			ApplyDepthConstraintsToOffset(GridCoordinate, CellOffset);
 
 			float SceneDepth;
 			float3 TranslatedWorldPosition = ComputeCellTranslatedWorldPosition(GridCoordinate, CellOffset, SceneDepth);
@@ -543,6 +583,8 @@ RAY_TRACING_ENTRY_RAYGEN(InjectShadowedDirectionalLightRGS)
 	float3 Rand3D = (float3(Rand32Bits) / float(uint(0xffffffff))) * 2.0f - 1.0f;
 	float3 CellOffset = FrameJitterOffsets[SampleIndex].xyz + LightScatteringSampleJitterMultiplier * Rand3D;
 
+	ApplyDepthConstraintsToOffset(Position, CellOffset);
+
 	const float3 TranslatedWorldPosition = ComputeCellTranslatedWorldPosition(Position, CellOffset);
 	OutShadowVolumeTexture[Position] = ComputeRayTracedShadowFactor(TranslatedWorldPosition, ForwardLightStruct.DirectionalLightDirection, 1.0e27);
 }
@@ -839,6 +881,8 @@ void LightScatteringCS(
 		float3 CellOffset = FrameJitterOffsets[SampleIndex].xyz + LightScatteringSampleJitterMultiplier * Rand3D;
 		//CellOffset = 0.5f;
 
+		ApplyDepthConstraintsToOffset(GridCoordinate, CellOffset);
+
 		float SceneDepth;
 		float3 TranslatedWorldPosition = ComputeCellTranslatedWorldPosition(GridCoordinate, CellOffset, SceneDepth);
 		float3 WorldPosition = TranslatedWorldPosition - DFHackToFloat(PrimaryView.PreViewTranslation); // LWC_TODO
diff --git a/Engine/Source/Runtime/Renderer/Private/VolumetricFog.cpp b/Engine/Source/Runtime/Renderer/Private/VolumetricFog.cpp
--- a/Engine/Source/Runtime/Renderer/Private/VolumetricFog.cpp
+++ b/Engine/Source/Runtime/Renderer/Private/VolumetricFog.cpp
@@ -164,6 +164,20 @@ static TAutoConsoleVariable<float> CVarVolumetricFogLightSoftFading(
 	TEXT("Enabled when >0, controls the soft fading of spot and rect light edges in order to make them appear smoother, resulting in less flickering. A value of 1 is a good starting point and it means that fading will be applied over the size of 1 froxel on the screen."),
 	ECVF_RenderThreadSafe | ECVF_Scalability);
 
+static TAutoConsoleVariable<float> CVarVolumetricFogGridCenterOffsetFromDepthBuffer(
+	TEXT("r.VolumetricFog.GridCenterOffsetFromDepthBuffer"),
+	0.5f,
+	TEXT("Offset in grid units to move grid center sample out form the depth buffer along the Z direction. -1 means disabled. This reduces sample self intersection with geometry when tracing, and thus reduces flickering in those areas, as well as results in less leaking sometimes. Set to -1 to disable."),
+	ECVF_RenderThreadSafe
+);
+
+static TAutoConsoleVariable<float> CVarVolumetricFogOffsetThresholdToAcceptDepthBufferOffset(
+	TEXT("r.VolumetricFog.OffsetThresholdToAcceptDepthBufferOffset"),
+	1.0f,
+	TEXT("Offset in grid units to accept a sample to be moved forward in front of the depth buffer. This is to avoid moving all samples behind the depth buffer forward which would affect the lighting of translucent and volumetric at edges of mesh. Default to 1.0 to only allow moving the first layer of froxel intersecting depth."),
+	ECVF_RenderThreadSafe
+);
+
 static const TCHAR* VolumetricFogGlobalPSOCollectorName = TEXT("VolumetricFogGlobalPSOCollector");
 
 int32 GetVolumetricFogGridPixelSize()
@@ -241,6 +255,9 @@ void SetupVolumetricFogIntegrationParameters(
 
 	extern int32 GVolumetricFogHistoryMissSupersampleCount;
 	Out.HistoryMissSuperSampleCount = FMath::Clamp(GVolumetricFogHistoryMissSupersampleCount, 1, 16);
+
+	Out.GridCenterOffsetFromDepthBuffer = CVarVolumetricFogGridCenterOffsetFromDepthBuffer.GetValueOnRenderThread();
+	Out.GridCenterOffsetThresholdToAcceptDepthBufferOffset = CVarVolumetricFogOffsetThresholdToAcceptDepthBufferOffset.GetValueOnRenderThread();
 }
 
 static const uint32 VolumetricFogGridInjectionGroupSize = 4;
@@ -265,7 +282,7 @@ class FVolumetricFogMaterialSetupCS : public FGlobalShader
 
 		SHADER_PARAMETER_RDG_UNIFORM_BUFFER(FFogUniformParameters, Fog)
 		SHADER_PARAMETER_STRUCT_REF(FViewUniformShaderParameters, View)
-
+		SHADER_PARAMETER_STRUCT_INCLUDE(FSceneTextureShaderParameters, SceneTextures)
 		SHADER_PARAMETER_STRUCT_INCLUDE(FVolumetricFogIntegrationParameters, VolumetricFogParameters)
 
 		SHADER_PARAMETER_STRUCT(FLocalFogVolumeUniformParameters, LFV)
@@ -313,6 +330,7 @@ BEGIN_SHADER_PARAMETER_STRUCT(FInjectShadowedLocalLightCommonParameters, )
 	SHADER_PARAMETER_RDG_UNIFORM_BUFFER(FDeferredLightUniformStruct, DeferredLight)
 	SHADER_PARAMETER_RDG_UNIFORM_BUFFER(FLightFunctionAtlasGlobalParameters, LightFunctionAtlas)
 	SHADER_PARAMETER_RDG_TEXTURE(Texture2D, WhiteDummyTexture)
+	SHADER_PARAMETER_STRUCT_INCLUDE(FSceneTextureShaderParameters, SceneTextures)
 	SHADER_PARAMETER_STRUCT_INCLUDE(FVolumetricFogIntegrationParameters, VolumetricFogParameters)
 	SHADER_PARAMETER(float, PhaseG)
 	SHADER_PARAMETER(float, InverseSquaredLightDistanceBiasScale)
@@ -332,6 +350,7 @@ static bool SetupInjectShadowedLocalLightCommonParameters(
 	SetupVolumetricFogIntegrationParameters(OutCommonParameters.VolumetricFogParameters, View, IntegrationData);
 
 	OutCommonParameters.ViewUniformBuffer = View.ViewUniformBuffer;
+	OutCommonParameters.SceneTextures = GetSceneTextureShaderParameters(View);
 	OutCommonParameters.PhaseG = FogInfo.VolumetricFogScatteringDistribution;
 	OutCommonParameters.InverseSquaredLightDistanceBiasScale = GInverseSquaredLightDistanceBiasScale;
 
@@ -465,6 +484,7 @@ class FRayTraceDirectionalLightVolumeShadowMapRGS : public FGlobalShader
 		SHADER_PARAMETER_RDG_UNIFORM_BUFFER(FSceneUniformParameters, Scene)
 		SHADER_PARAMETER_RDG_UNIFORM_BUFFER(FNaniteRayTracingUniformParameters, NaniteRayTracing)
 		SHADER_PARAMETER_RDG_UNIFORM_BUFFER(FForwardLightUniformParameters, ForwardLightStruct)
+		SHADER_PARAMETER_STRUCT_INCLUDE(FSceneTextureShaderParameters, SceneTextures)
 		SHADER_PARAMETER_STRUCT_INCLUDE(FVolumetricFogIntegrationParameters, VolumetricFogParameters)
 		SHADER_PARAMETER_RDG_TEXTURE_UAV(RWTexture3D, OutShadowVolumeTexture)
 		SHADER_PARAMETER_RDG_BUFFER_SRV(RaytracingAccelerationStructure, TLAS)
@@ -550,6 +570,7 @@ static void RenderRaytracedDirectionalShadowVolume(
 		PassParameters->TLAS = View.GetRayTracingSceneLayerViewChecked(ERayTracingSceneLayer::Base);
 		PassParameters->View = View.ViewUniformBuffer;
 		PassParameters->Scene = GetSceneUniformBufferRef(GraphBuilder, View);
+		PassParameters->SceneTextures = GetSceneTextureShaderParameters(View);
 		PassParameters->NaniteRayTracing = Nanite::GRayTracingManager.GetUniformBuffer();
 		PassParameters->ForwardLightStruct = View.ForwardLightingResources.ForwardLightUniformBuffer;
 		PassParameters->LightScatteringSampleJitterMultiplier = GVolumetricFogJitter ? GLightScatteringSampleJitterMultiplier : 0;
@@ -1032,6 +1053,7 @@ class FVolumetricFogLightScatteringCS : public FGlobalShader
 		SHADER_PARAMETER_RDG_UNIFORM_BUFFER(FForwardLightUniformParameters, ForwardLightStruct)
 		SHADER_PARAMETER_RDG_UNIFORM_BUFFER(FFogUniformParameters, Fog)
 		SHADER_PARAMETER_RDG_UNIFORM_BUFFER(FLightFunctionAtlasGlobalParameters, LightFunctionAtlas)
+		SHADER_PARAMETER_STRUCT_INCLUDE(FSceneTextureShaderParameters, SceneTextures)
 		SHADER_PARAMETER_STRUCT_INCLUDE(FVolumetricFogIntegrationParameters, VolumetricFogParameters)
 
 		SHADER_PARAMETER_RDG_TEXTURE(Texture3D, MegaLightsVolume)
@@ -1635,6 +1657,7 @@ void FSceneRenderer::ComputeVolumetricFog(FRDGBuilder& GraphBuilder,
 
 			PassParameters->Fog = FogUniformBuffer; 
 			PassParameters->View = View.ViewUniformBuffer;
+			PassParameters->SceneTextures = GetSceneTextureShaderParameters(View);
 			SetupVolumetricFogIntegrationParameters(PassParameters->VolumetricFogParameters, View, IntegrationData);
 
 			FVolumetricFogMaterialSetupCS::FPermutationDomain PermutationVector;
@@ -1677,6 +1700,7 @@ void FSceneRenderer::ComputeVolumetricFog(FRDGBuilder& GraphBuilder,
 			FVolumetricFogLightScatteringCS::FParameters* PassParameters = GraphBuilder.AllocParameters<FVolumetricFogLightScatteringCS::FParameters>();
 
 			PassParameters->View = View.ViewUniformBuffer;
+			PassParameters->SceneTextures = GetSceneTextureShaderParameters(View);
 			PassParameters->ForwardLightStruct = View.ForwardLightingResources.ForwardLightUniformBuffer;
 			PassParameters->Fog = FogUniformBuffer;
 			SetupVolumetricFogIntegrationParameters(PassParameters->VolumetricFogParameters, View, IntegrationData);
diff --git a/Engine/Source/Runtime/Renderer/Private/VolumetricFogShared.h b/Engine/Source/Runtime/Renderer/Private/VolumetricFogShared.h
--- a/Engine/Source/Runtime/Renderer/Private/VolumetricFogShared.h
+++ b/Engine/Source/Runtime/Renderer/Private/VolumetricFogShared.h
@@ -46,6 +46,8 @@ BEGIN_SHADER_PARAMETER_STRUCT(FVolumetricFogIntegrationParameters, )
 	SHADER_PARAMETER_ARRAY(FVector4f, FrameJitterOffsets, [16])
 	SHADER_PARAMETER(float, HistoryWeight)
 	SHADER_PARAMETER(uint32, HistoryMissSuperSampleCount)
+	SHADER_PARAMETER(float, GridCenterOffsetFromDepthBuffer)
+	SHADER_PARAMETER(float, GridCenterOffsetThresholdToAcceptDepthBufferOffset)
 END_SHADER_PARAMETER_STRUCT()
 
 void SetupVolumetricFogIntegrationParameters(
