Stop CSM and Inset Shadow from rendering simultaneously if capsule shadows are enabled on mobile

#Synced-CL 47305274

diff --git a/Engine/Source/Runtime/Renderer/Private/ShadowSetup.cpp b/Engine/Source/Runtime/Renderer/Private/ShadowSetup.cpp
--- a/Engine/Source/Runtime/Renderer/Private/ShadowSetup.cpp
+++ b/Engine/Source/Runtime/Renderer/Private/ShadowSetup.cpp
@@ -3227,17 +3227,22 @@ void FSceneRenderer::SetupInteractionShadows(
 	if (!bShadowHandledByParent)
 	{
 		const bool bCreateTranslucentObjectShadow = GUseTranslucencyShadowDepths && Interaction->HasTranslucentObjectShadow();
-		const bool bCreateInsetObjectShadow = Interaction->HasInsetObjectShadow();
 		const bool bCreateObjectShadowForStationaryLight = ShouldCreateObjectShadowForStationaryLight(Interaction->GetLight(), PrimitiveSceneInfo->Proxy, Interaction->IsShadowMapped());
 
-		if (Interaction->HasShadow() 
-			// TODO: Handle inset shadows, especially when an object is only casting a self-shadow.
+		bool bCreateOpaqueObjectShadow = Interaction->HasInsetObjectShadow();
+		if (bCreateOpaqueObjectShadow && TaskData.FeatureLevel < ERHIFeatureLevel::SM5)
+		{
+			bCreateOpaqueObjectShadow = (CVarMobileInsetShadows.GetValueOnRenderThread() && Interaction->HasInsetObjectShadow())
+				|| (IsCapsuleDirectShadowsEnabled(TaskData.ShaderPlatform) && PrimitiveSceneInfo->Proxy->CastsCapsuleDirectShadow());
+		}
+
+		if (Interaction->HasShadow()
 			// Only render shadows from objects that use static lighting during a reflection capture, since the reflection capture doesn't update at runtime
 			&& (!bStaticSceneOnly || PrimitiveSceneInfo->Proxy->HasStaticLighting())
-			&& (bCreateTranslucentObjectShadow || bCreateInsetObjectShadow || bCreateObjectShadowForStationaryLight))
+			&& (bCreateTranslucentObjectShadow || bCreateOpaqueObjectShadow || bCreateObjectShadowForStationaryLight))
 		{
 			// Create projected shadow infos
-			CreatePerObjectProjectedShadow(TaskData, Interaction, bCreateTranslucentObjectShadow, bCreateInsetObjectShadow || bCreateObjectShadowForStationaryLight, ViewDependentWholeSceneShadows, PreShadows);
+			CreatePerObjectProjectedShadow(TaskData, Interaction, bCreateTranslucentObjectShadow, bCreateOpaqueObjectShadow || bCreateObjectShadowForStationaryLight, ViewDependentWholeSceneShadows, PreShadows);
 		}
 	}
 }
@@ -4793,7 +4798,6 @@ struct FGatherShadowPrimitivesPacket
 	TArray<FAddSubjectPrimitiveOverflowedIndices> ViewDependentWholeSceneShadowOverflowedIndices;
 	TArray<FShadowSubjectPrimitives> PreShadowSubjectPrimitives;
 	TArray<FShadowSubjectPrimitives> ViewDependentWholeSceneShadowSubjectPrimitives;
-	bool bSupportMobileInsetShadows = false;
 
 	FGatherShadowPrimitivesPacket(
 		FScenePrimitiveOctree::FNodeIndex InNodeIndex,
@@ -4827,12 +4831,6 @@ struct FGatherShadowPrimitivesPacket
 
 		ViewDependentWholeSceneShadowSubjectPrimitives.Empty(NumVDWSShadows);
 		ViewDependentWholeSceneShadowSubjectPrimitives.AddDefaulted(NumVDWSShadows);
-
-		if (TaskData.FeatureLevel < ERHIFeatureLevel::SM5)
-		{
-			bSupportMobileInsetShadows = CVarMobileInsetShadows.GetValueOnRenderThread() 
-				&& (MobileUsesFullDepthPrepass(TaskData.ShaderPlatform) || (IsMobileDeferredShadingEnabled(TaskData.ShaderPlatform) && !MobileAllowFramebufferFetch(TaskData.ShaderPlatform)));
-		}
 		
 		if (NodeIndex != INDEX_NONE)
 		{
@@ -4875,23 +4873,63 @@ struct FGatherShadowPrimitivesPacket
 		ViewDependentWholeSceneShadowStats.Empty();
 	}
 
-	bool DoesPrimitiveCastInsetShadow(FDynamicShadowsTaskData& TaskData, const FPrimitiveSceneInfo* PrimitiveSceneInfo, const FPrimitiveSceneProxy* PrimitiveProxy) const
+	bool FilterForHigherPriorityShadow(FDynamicShadowsTaskData& TaskData, const FPrimitiveSceneInfo* PrimitiveSceneInfo, const FPrimitiveSceneProxy* PrimitiveProxy) const
 	{
 		if (UseVirtualShadowMaps(TaskData.ShaderPlatform, TaskData.FeatureLevel))
 		{
 			return false;
 		}
 
+		bool bCastsHigherPriorityShadow = false;
+
 		// If light attachment root is valid, we're in a group and need to get the flag from the root.
+		// If not in mobile, it is sufficient to only check inset shadow because Inset Shadow = Inset Shadow || CapsuleDirectShadow
 		if (PrimitiveSceneInfo->LightingAttachmentRoot.IsValid())
 		{
 			const FAttachmentGroupSceneInfo& AttachmentGroup = PrimitiveSceneInfo->Scene->AttachmentGroups.FindChecked(PrimitiveSceneInfo->LightingAttachmentRoot);
-			return AttachmentGroup.ParentSceneInfo && AttachmentGroup.ParentSceneInfo->Proxy->CastsInsetShadow();
+			if (AttachmentGroup.ParentSceneInfo)
+			{
+				bCastsHigherPriorityShadow = AttachmentGroup.ParentSceneInfo->Proxy->CastsInsetShadow();
+			}
 		}
 		else
 		{
-			return PrimitiveProxy->CastsInsetShadow();
+			bCastsHigherPriorityShadow = PrimitiveProxy->CastsInsetShadow();
 		}
+
+		// if we are on mobile check if inset/capsule shadows are actually cast
+		if (bCastsHigherPriorityShadow && TaskData.FeatureLevel < ERHIFeatureLevel::SM5)
+		{
+			const bool MobileAllow = MobileUsesFullDepthPrepass(TaskData.ShaderPlatform) || (IsMobileDeferredShadingEnabled(TaskData.ShaderPlatform) && !MobileAllowFramebufferFetch(TaskData.ShaderPlatform));
+			const bool bSupportsMobileInsetShadows = CVarMobileInsetShadows.GetValueOnRenderThread() && MobileAllow;
+			const bool bSupportsMobileCapsuleDirectShadows = IsCapsuleDirectShadowsEnabled(TaskData.ShaderPlatform) && MobileAllow;
+
+			// primitive casts inset shadow and inset shadow is turned on
+			const bool bCastsMobileInsetShadow = bCastsHigherPriorityShadow && bSupportsMobileInsetShadows;
+
+			// Primitive casts capsule direct shadow and capsule direct shadow is turned on
+			// We may need to check against light attachment root
+			bool bCastsMobileCapsuleDirectShadow = false;
+			if (bSupportsMobileCapsuleDirectShadows)
+			{
+				if (PrimitiveSceneInfo->LightingAttachmentRoot.IsValid())
+				{
+					const FAttachmentGroupSceneInfo& AttachmentGroup = PrimitiveSceneInfo->Scene->AttachmentGroups.FindChecked(PrimitiveSceneInfo->LightingAttachmentRoot);
+					if (AttachmentGroup.ParentSceneInfo)
+					{
+						bCastsMobileCapsuleDirectShadow = AttachmentGroup.ParentSceneInfo->Proxy->CastsCapsuleDirectShadow();
+					}
+				}
+				else
+				{
+					bCastsMobileCapsuleDirectShadow = PrimitiveProxy->CastsCapsuleDirectShadow();
+				}
+			}
+
+			bCastsHigherPriorityShadow = bCastsMobileInsetShadow || bCastsMobileCapsuleDirectShadow;
+		}
+
+		return bCastsHigherPriorityShadow;
 	}
 
 	void FilterPrimitiveForShadows(FDynamicShadowsTaskData& TaskData, const FPrimitiveSceneInfoCompact& PrimitiveSceneInfoCompact)
@@ -4996,13 +5034,8 @@ struct FGatherShadowPrimitivesPacket
 					}
 				}
 
-				bool bCastsInsetShadow = DoesPrimitiveCastInsetShadow(TaskData, PrimitiveSceneInfo, PrimitiveProxy);
-
-				// if we are on mobile also check if inset shadows are supported
-				if (bCastsInsetShadow && TaskData.FeatureLevel < ERHIFeatureLevel::SM5)
-				{
-					bCastsInsetShadow &= bSupportMobileInsetShadows;
-				}
+				// Filter out when inset or capsule shadows are drawn
+				bool bCastsHigherPriorityShadow = FilterForHigherPriorityShadow(TaskData, PrimitiveSceneInfo, PrimitiveProxy);
 
 				if (!bScreenSpaceSizeCulled
 					&& bShadowRelevance
@@ -5011,7 +5044,7 @@ struct FGatherShadowPrimitivesPacket
 					// Since lights with static shadowing still create per-object shadows for primitives without static shadowing.
 					&& (!LightProxy.HasStaticLighting() || (!LightSceneInfo.IsPrecomputedLightingValid() || LightProxy.UseCSMForDynamicObjects()))
 					// Exclude primitives that will create their own per-object shadow, except when rendering RSMs
-					&& (!bCastsInsetShadow)
+					&& (!bCastsHigherPriorityShadow)
 					// Exclude primitives that will create a per-object shadow from a stationary light
 					&& !ShouldCreateObjectShadowForStationaryLight(&LightSceneInfo, PrimitiveProxy, true)
 					// Only render shadows from objects that use static lighting during a reflection capture, since the reflection capture doesn't update at runtime
@@ -6120,6 +6153,7 @@ void FSceneRenderer::CreateDynamicShadows(FDynamicShadowsTaskData& TaskData)
 	const bool bProjectEnablePointLightShadows = FReadOnlyCVARCache::EnablePointLightShadows(ShaderPlatform);
 	const bool bProjectEnableMovableSpotLightShadows = !bMobile || IsMobileMovableSpotlightShadowsEnabled(ShaderPlatform);
 	const bool bMobileInsetShadows = CVarMobileInsetShadows.GetValueOnRenderThread();
+	const bool bCapsuleDirectShadows = IsCapsuleDirectShadowsEnabled(ShaderPlatform);
 
 	uint32 NumPointShadowCachesUpdatedThisFrame = 0;
 	uint32 NumSpotShadowCachesUpdatedThisFrame = 0;
@@ -6271,8 +6305,8 @@ void FSceneRenderer::CreateDynamicShadows(FDynamicShadowsTaskData& TaskData)
 							{
 								AddViewDependentWholeSceneShadowsForView(ViewDependentWholeSceneShadows, ViewDependentWholeSceneShadowsThatNeedCulling, VisibleLightInfo, *LightSceneInfo, GetCachedShadowMapsSize(), NumCSMCachesUpdatedThisFrame);
 							}
-							
-							if (!bMobile || bMobileInsetShadows || IsMobileCapsuleShadowsEnabled(ShaderPlatform) || (LightSceneInfo->Proxy->CastsModulatedShadows() && !LightSceneInfo->Proxy->UseCSMForDynamicObjects() && LightSceneInfo->Proxy->HasStaticShadowing()))
+
+							if (!bMobile || bMobileInsetShadows || bCapsuleDirectShadows || (LightSceneInfo->Proxy->CastsModulatedShadows() && !LightSceneInfo->Proxy->UseCSMForDynamicObjects() && LightSceneInfo->Proxy->HasStaticShadowing()))
 							{
 								if (!LightSceneInfo->Proxy->UseVirtualShadowMaps())
 								{
