[Distance Fields] Disable DFAO when Lumen is enabled.
- when r.AOApplyToStaticIndirect is enabled the renderer was running RenderDFAOAsIndirectShadowing(...) regardless of whether Lumen is enabled or not.

#Synced-CL 48158727

diff --git a/Engine/Source/Runtime/Renderer/Private/DistanceFieldAmbientOcclusion.cpp b/Engine/Source/Runtime/Renderer/Private/DistanceFieldAmbientOcclusion.cpp
--- a/Engine/Source/Runtime/Renderer/Private/DistanceFieldAmbientOcclusion.cpp
+++ b/Engine/Source/Runtime/Renderer/Private/DistanceFieldAmbientOcclusion.cpp
@@ -866,7 +866,7 @@ void FDeferredShadingSceneRenderer::RenderDFAOAsIndirectShadowing(
 	const FSceneTextures& SceneTextures,
 	TArray<FRDGTextureRef>& DynamicBentNormalAOTextures)
 {
-	if (GDistanceFieldAOApplyToStaticIndirect && ShouldRenderDistanceFieldAO(Views, ViewFamily.EngineShowFlags) && ShouldRenderDistanceFieldLighting(Scene->DistanceFieldSceneData, Views))
+	if (GDistanceFieldAOApplyToStaticIndirect && ShouldRenderDistanceFieldAO(Views, ViewFamily.EngineShowFlags, AnyViewHasGIMethodSupportingDFAO()) && ShouldRenderDistanceFieldLighting(Scene->DistanceFieldSceneData, Views))
 	{
 		// Use the skylight's max distance if there is one, to be consistent with DFAO shadowing on the skylight
 		const float OcclusionMaxDistance = Scene->SkyLight && !Scene->SkyLight->bWantsStaticShadowing ? Scene->SkyLight->OcclusionMaxDistance : Scene->DefaultMaxDistanceFieldOcclusionDistance;
@@ -974,7 +974,7 @@ void FSceneRenderer::RenderDistanceFieldLighting(
 	}
 }
 
-bool ShouldRenderDistanceFieldAO(TConstArrayView<FViewInfo> Views, const FEngineShowFlags& EngineShowFlags)
+bool ShouldRenderDistanceFieldAO(TConstArrayView<FViewInfo> Views, const FEngineShowFlags& EngineShowFlags, bool bAnyViewHasGIMethodSupportingDFAO)
 {
 	bool bShouldRenderRTAO = false;
 	for (int ViewIndex = 0; ViewIndex < Views.Num(); ++ViewIndex)
@@ -983,6 +983,7 @@ bool ShouldRenderDistanceFieldAO(TConstArrayView<FViewInfo> Views, const FEngine
 	}
 
 	return EngineShowFlags.DistanceFieldAO
+		&& bAnyViewHasGIMethodSupportingDFAO
 		&& !bShouldRenderRTAO
 		&& !EngineShowFlags.VisualizeDistanceFieldAO
 		&& !EngineShowFlags.VisualizeMeshDistanceFields
diff --git a/Engine/Source/Runtime/Renderer/Private/DistanceFieldLightingShared.h b/Engine/Source/Runtime/Renderer/Private/DistanceFieldLightingShared.h
--- a/Engine/Source/Runtime/Renderer/Private/DistanceFieldLightingShared.h
+++ b/Engine/Source/Runtime/Renderer/Private/DistanceFieldLightingShared.h
@@ -156,7 +156,7 @@ extern void CullDistanceFieldObjectsForLight(
 	FLightTileIntersectionParameters& LightTileIntersectionParameters);
 
 extern bool ShouldPrepareForDistanceFieldAO(const FScene* Scene, const FSceneViewFamily& ViewFamily, bool bAnyViewHasGIMethodSupportingDFAO);
-extern bool ShouldRenderDistanceFieldAO(TConstArrayView<FViewInfo> Views, const FEngineShowFlags& EngineShowFlags);
+extern bool ShouldRenderDistanceFieldAO(TConstArrayView<FViewInfo> Views, const FEngineShowFlags& EngineShowFlags, bool bAnyViewHasGIMethodSupportingDFAO);
 
 extern bool SupportsDistanceFieldAO(ERHIFeatureLevel::Type FeatureLevel, EShaderPlatform ShaderPlatform);
 
diff --git a/Engine/Source/Runtime/Renderer/Private/IndirectLightRendering.cpp b/Engine/Source/Runtime/Renderer/Private/IndirectLightRendering.cpp
--- a/Engine/Source/Runtime/Renderer/Private/IndirectLightRendering.cpp
+++ b/Engine/Source/Runtime/Renderer/Private/IndirectLightRendering.cpp
@@ -2068,7 +2068,7 @@ void FDeferredShadingSceneRenderer::RenderDeferredReflectionsAndSkyLighting(
 		extern int32 GDistanceFieldAOApplyToStaticIndirect;
 		if (Scene->SkyLight->bCastShadows
 			&& !GDistanceFieldAOApplyToStaticIndirect
-			&& ShouldRenderDistanceFieldAO(Views, ViewFamily.EngineShowFlags)
+			&& ShouldRenderDistanceFieldAO(Views, ViewFamily.EngineShowFlags, AnyViewHasGIMethodSupportingDFAO())
 			&& ShouldRenderDistanceFieldLighting(Scene->DistanceFieldSceneData, Views)
 			&& ViewFamily.EngineShowFlags.AmbientOcclusion
 			&& !bReflectionCapture)
diff --git a/Engine/Source/Runtime/Renderer/Private/MobileShadingRenderer.cpp b/Engine/Source/Runtime/Renderer/Private/MobileShadingRenderer.cpp
--- a/Engine/Source/Runtime/Renderer/Private/MobileShadingRenderer.cpp
+++ b/Engine/Source/Runtime/Renderer/Private/MobileShadingRenderer.cpp
@@ -616,7 +616,7 @@ void FMobileSceneRenderer::InitViews(
 		&& AnyViewHasGIMethodSupportingDFAO()
 		&& (Views[0].GlobalDistanceFieldInfo.Clipmaps.Num() > 0)
 		&& !GDistanceFieldAOApplyToStaticIndirect
-		&& ShouldRenderDistanceFieldAO(Views, ViewFamily.EngineShowFlags)
+		&& ShouldRenderDistanceFieldAO(Views, ViewFamily.EngineShowFlags, AnyViewHasGIMethodSupportingDFAO())
 		&& ShouldRenderDistanceFieldLighting(Scene->DistanceFieldSceneData, Views)
 		&& ViewFamily.EngineShowFlags.AmbientOcclusion
 		&& !Views[0].bIsReflectionCapture;
diff --git a/Engine/Source/Runtime/Renderer/Private/VolumetricFog.cpp b/Engine/Source/Runtime/Renderer/Private/VolumetricFog.cpp
--- a/Engine/Source/Runtime/Renderer/Private/VolumetricFog.cpp
+++ b/Engine/Source/Runtime/Renderer/Private/VolumetricFog.cpp
@@ -1847,7 +1847,7 @@ void FSceneRenderer::ComputeVolumetricFog(FRDGBuilder& GraphBuilder,
 				&& Scene->SkyLight
 				&& Scene->SkyLight->bCastShadows
 				&& Scene->SkyLight->bCastVolumetricShadow
-				&& ShouldRenderDistanceFieldAO(Views, ViewFamily.EngineShowFlags)
+				&& ShouldRenderDistanceFieldAO(Views, ViewFamily.EngineShowFlags, AnyViewHasGIMethodSupportingDFAO())
 				&& SupportsDistanceFieldAO(View.GetFeatureLevel(), View.GetShaderPlatform())
 				&& bUseGlobalDistanceField
 				&& Views.Num() == 1
