Fix race condition accessing GlobalDistanceFieldData->CameraVelocityOffset.
- Need to update global distance field view origin before calling BeginUpdateLumenSceneTasks(...).
- Decoupled UpdateGlobalDistanceFieldViewOrigin(...) from PrepareDistanceFieldScene(...) and call it earlier in the frame.

#Synced-CL 48158727

diff --git a/Engine/Source/Runtime/Renderer/Private/DeferredShadingRenderer.cpp b/Engine/Source/Runtime/Renderer/Private/DeferredShadingRenderer.cpp
--- a/Engine/Source/Runtime/Renderer/Private/DeferredShadingRenderer.cpp
+++ b/Engine/Source/Runtime/Renderer/Private/DeferredShadingRenderer.cpp
@@ -1825,7 +1825,14 @@ void FDeferredShadingSceneRenderer::Render(FRDGBuilder& GraphBuilder, const FSce
 				}
 			}
 
-			// 2. Update lumen scene
+			// 2. Update global distance field view origin
+			// This needs to be done prior to start Lumen scene lighting to ensure GlobalDistanceFieldData->CameraVelocityOffset is updated
+			if (SceneUpdateInputs)
+			{
+				UpdateGlobalDistanceFieldViewOrigin(*SceneUpdateInputs);
+			}
+
+			// 3. Update lumen scene
 			{
 				InitViewTaskDatas.LumenFrameTemporaries = &LumenFrameTemporaries;
 	
diff --git a/Engine/Source/Runtime/Renderer/Private/GlobalDistanceField.cpp b/Engine/Source/Runtime/Renderer/Private/GlobalDistanceField.cpp
--- a/Engine/Source/Runtime/Renderer/Private/GlobalDistanceField.cpp
+++ b/Engine/Source/Runtime/Renderer/Private/GlobalDistanceField.cpp
@@ -725,9 +725,10 @@ static void UpdateGlobalDistanceFieldViewOrigin(const FViewInfo& View, bool bLum
 	// Don't update origin if it has already been updated this frame
 	if (View.ViewState && (View.ViewState->GlobalDistanceFieldData->bFirstFrame || View.ViewState->GlobalDistanceFieldData->UpdateFrame != View.Family->FrameNumber))
 	{
+		FVector& CameraVelocityOffset = View.ViewState->GlobalDistanceFieldData->CameraVelocityOffset;
+
 		if (GAOGlobalDistanceFieldFastCameraMode != 0)
 		{
-			FVector& CameraVelocityOffset = View.ViewState->GlobalDistanceFieldData->CameraVelocityOffset;
 			const FVector CameraVelocity = View.ViewMatrices.GetViewOrigin() - View.PrevViewInfo.ViewMatrices.GetViewOrigin();
 			// Framerate independent decay
 			CameraVelocityOffset = CameraVelocityOffset * FMath::Pow(GAOGlobalDistanceFieldCameraPositionVelocityOffsetDecay, View.Family->Time.GetDeltaWorldTimeSeconds()) + CameraVelocity;
@@ -743,11 +744,21 @@ static void UpdateGlobalDistanceFieldViewOrigin(const FViewInfo& View, bool bLum
 		}
 		else
 		{
-			View.ViewState->GlobalDistanceFieldData->CameraVelocityOffset = FVector(0.0f, 0.0f, 0.0f);
+			CameraVelocityOffset = FVector(0.0f, 0.0f, 0.0f);
 		}
 	}
 }
 
+void UpdateGlobalDistanceFieldViewOrigin(const FSceneRenderUpdateInputs& SceneUpdateInputs)
+{
+	SceneUpdateInputs.ForEachView([&](const FSceneRenderer* Renderer, FViewInfo& View)
+		{
+			const bool bUseLumenGlobalDistanceFieldSettings = Renderer->IsLumenEnabled(View) && Lumen::IsUsingGlobalSDF(*View.Family);
+			UpdateGlobalDistanceFieldViewOrigin(View, bUseLumenGlobalDistanceFieldSettings);
+			return true;
+		});
+}
+
 FVector GetGlobalDistanceFieldViewOrigin(const FViewInfo& View, int32 ClipmapIndex, bool bLumenEnabled)
 {
 	FVector CameraOrigin = View.ViewMatrices.GetViewOrigin();
@@ -3144,8 +3155,6 @@ void UpdateGlobalDistanceFieldVolume(
 	const FDistanceFieldSceneData& DistanceFieldSceneData = Scene->DistanceFieldSceneData;
 	bool bNeedsFinalizeAccess = true;
 
-	UpdateGlobalDistanceFieldViewOrigin(View, bLumenEnabled);
-
 	if (DistanceFieldSceneData.NumObjectsInBuffer > 0 || DistanceFieldSceneData.HeightfieldPrimitives.Num() > 0)
 	{
 		const int32 NumClipmaps = FMath::Clamp<int32>(GlobalDistanceField::GetNumGlobalDistanceFieldClipmaps(bLumenEnabled, View.FinalPostProcessSettings.LumenSceneViewDistance), 0, GlobalDistanceField::MaxClipmaps);
diff --git a/Engine/Source/Runtime/Renderer/Private/MobileShadingRenderer.cpp b/Engine/Source/Runtime/Renderer/Private/MobileShadingRenderer.cpp
--- a/Engine/Source/Runtime/Renderer/Private/MobileShadingRenderer.cpp
+++ b/Engine/Source/Runtime/Renderer/Private/MobileShadingRenderer.cpp
@@ -1060,6 +1060,7 @@ void FMobileSceneRenderer::Render(FRDGBuilder& GraphBuilder, const FSceneRenderU
 	
 				if (EngineShowFlags.Lighting && !EngineShowFlags.VisualizeLightCulling && !Family->UseDebugViewPS() && !View.bIsReflectionCapture && !View.bIsPlanarReflection)
 				{
+					UpdateGlobalDistanceFieldViewOrigin(*SceneUpdateInputs);
 					PrepareDistanceFieldScene(GraphBuilder, ExternalAccessQueue, *SceneUpdateInputs);
 					break;
 				}
diff --git a/Engine/Source/Runtime/Renderer/Private/ScenePrivate.h b/Engine/Source/Runtime/Renderer/Private/ScenePrivate.h
--- a/Engine/Source/Runtime/Renderer/Private/ScenePrivate.h
+++ b/Engine/Source/Runtime/Renderer/Private/ScenePrivate.h
@@ -2230,8 +2230,15 @@ private:
 	friend class FDistanceFieldStreamingUpdateTask;
 };
 
+/**
+ * Updates the global distance field view origin for all scene renderers.
+ * Must be called before PrepareDistanceFieldScene(...).
+ */
+extern void UpdateGlobalDistanceFieldViewOrigin(const FSceneRenderUpdateInputs& SceneUpdateInputs);
+
 /**
  * Prepares the distance field scene for all scene renderers.
+ * Must be called after UpdateGlobalDistanceFieldViewOrigin(...).
  */
 extern void PrepareDistanceFieldScene(FRDGBuilder& GraphBuilder, FRDGExternalAccessQueue& ExternalAccessQueue, const FSceneRenderUpdateInputs& SceneUpdateInputs);
 
