Add better hair composition with SLW.

When Water is enabled with depth prepass, hair samples composition is split into two stages:
* Composition of pixel below water level
* Composition of pixels above water level

#Synced-CL 48203640

diff --git a/Engine/Shaders/Private/HairStrands/HairStrandsComposition.usf b/Engine/Shaders/Private/HairStrands/HairStrandsComposition.usf
--- a/Engine/Shaders/Private/HairStrands/HairStrandsComposition.usf
+++ b/Engine/Shaders/Private/HairStrands/HairStrandsComposition.usf
@@ -77,6 +77,14 @@ int2 OutputResolution;
 uint bComposeDofDepth;
 uint bHasHoldout;
 
+#if PERMUTATION_SLW
+#define SLW_COMPOSE_NONE 0
+#define SLW_COMPOSE_BEHIND_WATER 1
+#define SLW_COMPOSE_INFRONTOF_WATER 2
+uint SLWComposeType;
+Texture2D<float> SLWDepthTexture;
+#endif
+
 Texture2D<float>	HairDOFDepthTexture;
 Texture2D<float4>	HairLightingSampleBuffer;
 
@@ -103,12 +111,16 @@ void ComposeSamplePS(
 		discard;
 	}
 
+	#if PERMUTATION_SLW
+	const float SLWDepth = SLWDepthTexture.Load(PixelCoord);
+	#endif
 	const uint TotalNodeCount = HairStrands.HairSampleCount[0];
 	const uint2 Resolution = GetHairSampleResolution(TotalNodeCount);
 
 	const float3 ClosetPointWorldPosition = DFHackToFloat(SvPositionToWorld(float4(Input.Position.xy, HairDeviceZ, 1.0)));
 	const float4 Volumetric = EvaluateVolumetric(ClosetPointWorldPosition, ConvertFromDeviceZ(HairDeviceZ));
 
+	bool bSkipComposition = false;
 	bool bHoldout = false;
 	float3 LocalAccColor = 0;
 	LOOP
@@ -121,12 +133,34 @@ void ComposeSamplePS(
 			bHoldout = bHoldout || HasHairFlags(Sample.Flags, HAIR_FLAGS_HOLDOUT);
 		}
 		#endif
+		
+		#if PERMUTATION_SLW
+		{
+			const FHairSample Sample = UnpackHairSample(HairStrands.HairSampleData[NodeDesc.Offset + SampleIt]);
+			if (SLWComposeType == SLW_COMPOSE_BEHIND_WATER)
+			{
+				bSkipComposition = bSkipComposition || Sample.Depth >= SLWDepth;
+			}
+			else if (SLWComposeType == SLW_COMPOSE_INFRONTOF_WATER)
+			{
+				bSkipComposition = bSkipComposition || Sample.Depth < SLWDepth;
+			}
+		}
+		#endif
 
 		const uint LinearCoord = NodeDesc.Offset + SampleIt;
 		const uint2 Coord = GetHairSampleCoord(LinearCoord, Resolution);
 		LocalAccColor += HairLightingSampleBuffer.Load(uint3(Coord, 0)).xyz;
 	}
 
+	#if PERMUTATION_SLW
+	if (bSkipComposition)
+	{
+		discard;
+	}
+
+	#endif
+
 	// If any sample support holdout, then the lighting contribution is set to 0 to mask out the current hair pixel 
 	// but fog from camera to hair pixel needs to be accounted for
 	if (bHoldout)
diff --git a/Engine/Source/Runtime/Renderer/Private/DeferredShadingRenderer.cpp b/Engine/Source/Runtime/Renderer/Private/DeferredShadingRenderer.cpp
--- a/Engine/Source/Runtime/Renderer/Private/DeferredShadingRenderer.cpp
+++ b/Engine/Source/Runtime/Renderer/Private/DeferredShadingRenderer.cpp
@@ -2495,7 +2495,7 @@ void FDeferredShadingSceneRenderer::Render(FRDGBuilder& GraphBuilder, const FSce
 							View.ReflectionCaptureUniformBuffer = EmptyReflectionCaptureUniformBuffer;
 						}
 
-						RenderTranslucency(*this, GraphBuilder, *CustomRenderPassSceneTextures, TranslucencyLightingVolumeTextures, &TranslucencyResourceMap, CustomRenderPassViews, ETranslucencyView::AboveWater, CustomTranslucencyDimensions, InstanceCullingManager, bStandardTranslucentCanRenderSeparate, TranslucencySharedDepthTexture);
+						RenderTranslucency(*this, GraphBuilder, *CustomRenderPassSceneTextures, TranslucencyLightingVolumeTextures, &TranslucencyResourceMap, CustomRenderPassViews, ETranslucencyView::AboveWater, CustomTranslucencyDimensions, InstanceCullingManager, bStandardTranslucentCanRenderSeparate, SingleLayerWaterPrePassResult, TranslucencySharedDepthTexture);
 					}
 				}
 
@@ -3488,7 +3488,7 @@ void FDeferredShadingSceneRenderer::Render(FRDGBuilder& GraphBuilder, const FSce
 				SCOPE_CYCLE_COUNTER(STAT_TranslucencyDrawTime);
 				const bool bStandardTranslucentCanRenderSeparate = false;
 				FRDGTextureMSAA SharedDepthTexture;
-				RenderTranslucency(*this, GraphBuilder, SceneTextures, TranslucencyLightingVolumeTextures, &TranslucencyResourceMap, Views, ETranslucencyView::UnderWater, SeparateTranslucencyDimensions, InstanceCullingManager, bStandardTranslucentCanRenderSeparate, SharedDepthTexture);
+				RenderTranslucency(*this, GraphBuilder, SceneTextures, TranslucencyLightingVolumeTextures, &TranslucencyResourceMap, Views, ETranslucencyView::UnderWater, SeparateTranslucencyDimensions, InstanceCullingManager, bStandardTranslucentCanRenderSeparate, SingleLayerWaterPrePassResult, SharedDepthTexture);
 				EnumRemoveFlags(TranslucencyViewsToRender, ETranslucencyView::UnderWater);
 			}
 
@@ -3547,7 +3547,7 @@ void FDeferredShadingSceneRenderer::Render(FRDGBuilder& GraphBuilder, const FSce
 		{
 			RDG_EVENT_SCOPE_STAT(GraphBuilder, HairRendering, "HairRendering");
 			RDG_GPU_STAT_SCOPE(GraphBuilder, HairRendering);
-			RenderHairComposition(GraphBuilder, Views, SceneTextures.Color.Target, SceneTextures.Depth.Target, SceneTextures.Velocity, TranslucencyResourceMap);
+			RenderHairComposition(GraphBuilder, Views, SceneTextures.Color.Target, SceneTextures.Depth.Target, SceneTextures.Velocity, TranslucencyResourceMap, SingleLayerWaterPrePassResult);
 		}
 
 #if DEBUG_ALPHA_CHANNEL
@@ -3622,7 +3622,7 @@ void FDeferredShadingSceneRenderer::Render(FRDGBuilder& GraphBuilder, const FSce
 			{
 				// Render all remaining translucency views.
 				const bool bStandardTranslucentCanRenderSeparate = bShouldRenderDistortion; // It is only needed to render standard translucent as separate when there is distortion (non self distortion of transmittance/specular/etc.)
-				RenderTranslucency(*this, GraphBuilder, SceneTextures, TranslucencyLightingVolumeTextures, &TranslucencyResourceMap, Views, TranslucencyViewsToRender, SeparateTranslucencyDimensions, InstanceCullingManager, bStandardTranslucentCanRenderSeparate, TranslucencySharedDepthTexture);
+				RenderTranslucency(*this, GraphBuilder, SceneTextures, TranslucencyLightingVolumeTextures, &TranslucencyResourceMap, Views, TranslucencyViewsToRender, SeparateTranslucencyDimensions, InstanceCullingManager, bStandardTranslucentCanRenderSeparate, SingleLayerWaterPrePassResult, TranslucencySharedDepthTexture);
 			}
 
 			// Compose hair before velocity/distortion pass since these pass write depth value, 
@@ -3632,7 +3632,7 @@ void FDeferredShadingSceneRenderer::Render(FRDGBuilder& GraphBuilder, const FSce
 				RDG_EVENT_SCOPE_STAT(GraphBuilder, HairRendering, "HairRendering");
 				RDG_GPU_STAT_SCOPE(GraphBuilder, HairRendering);
 
-				RenderHairComposition(GraphBuilder, Views, SceneTextures.Color.Target, SceneTextures.Depth.Target, SceneTextures.Velocity, TranslucencyResourceMap);
+				RenderHairComposition(GraphBuilder, Views, SceneTextures.Color.Target, SceneTextures.Depth.Target, SceneTextures.Velocity, TranslucencyResourceMap, SingleLayerWaterPrePassResult);
 			}
 
 			if (bShouldRenderDistortion)
@@ -3660,7 +3660,7 @@ void FDeferredShadingSceneRenderer::Render(FRDGBuilder& GraphBuilder, const FSce
 			RDG_EVENT_SCOPE_STAT(GraphBuilder, HairRendering, "HairRendering");
 			RDG_GPU_STAT_SCOPE(GraphBuilder, HairRendering);
 
-			RenderHairComposition(GraphBuilder, Views, SceneTextures.Color.Target, SceneTextures.Depth.Target, SceneTextures.Velocity, TranslucencyResourceMap);
+			RenderHairComposition(GraphBuilder, Views, SceneTextures.Color.Target, SceneTextures.Depth.Target, SceneTextures.Velocity, TranslucencyResourceMap, SingleLayerWaterPrePassResult);
 		}
 
 #if !UE_BUILD_SHIPPING
diff --git a/Engine/Source/Runtime/Renderer/Private/DeferredShadingRenderer.h b/Engine/Source/Runtime/Renderer/Private/DeferredShadingRenderer.h
--- a/Engine/Source/Runtime/Renderer/Private/DeferredShadingRenderer.h
+++ b/Engine/Source/Runtime/Renderer/Private/DeferredShadingRenderer.h
@@ -880,6 +880,7 @@ private:
 		const FSeparateTranslucencyDimensions& SeparateTranslucencyDimensions,
 		FInstanceCullingManager& InstanceCullingManager,
 		bool bStandardTranslucentCanRenderSeparate,
+		const FSingleLayerWaterPrePassResult* SingleLayerWaterPrePassResult,
 		FRDGTextureMSAA& OutSharedDepthTexture);
 
 	/** Renders the scene's translucency given a specific pass. */
diff --git a/Engine/Source/Runtime/Renderer/Private/HairStrands/HairStrandsComposition.cpp b/Engine/Source/Runtime/Renderer/Private/HairStrands/HairStrandsComposition.cpp
--- a/Engine/Source/Runtime/Renderer/Private/HairStrands/HairStrandsComposition.cpp
+++ b/Engine/Source/Runtime/Renderer/Private/HairStrands/HairStrandsComposition.cpp
@@ -15,6 +15,7 @@
 #include "FogRendering.h"
 #include "PostProcess/TemporalAA.h"
 #include "HairStrandsForwardRaster.h"
+#include "SingleLayerWaterRendering.h"
 
 /////////////////////////////////////////////////////////////////////////////////////////
 
@@ -49,6 +50,13 @@ enum EHairStrandsCommonPassType
 	HoldoutHair,
 };
 
+enum EHairSLWCompositionType
+{
+	All,
+	BehindWater,
+	InFrontOfWater,
+};
+
 template<typename TPassParameter, typename TPixelShader>
 void InternalCommonDrawPass(
 	FRDGBuilder& GraphBuilder, 
@@ -285,7 +293,8 @@ class FHairVisibilityComposeSamplePS : public FGlobalShader
 	SHADER_USE_PARAMETER_STRUCT(FHairVisibilityComposeSamplePS, FGlobalShader);
 
 	class FDebug : SHADER_PERMUTATION_BOOL("PERMUTATION_DEBUG");
-	using FPermutationDomain = TShaderPermutationDomain<FDebug>;
+	class FSlw : SHADER_PERMUTATION_BOOL("PERMUTATION_SLW");
+	using FPermutationDomain = TShaderPermutationDomain<FDebug, FSlw>;
 	
 	BEGIN_SHADER_PARAMETER_STRUCT(FParameters, )
 		SHADER_PARAMETER_STRUCT_REF(FViewUniformShaderParameters, ViewUniformBuffer)
@@ -293,9 +302,11 @@ class FHairVisibilityComposeSamplePS : public FGlobalShader
 		SHADER_PARAMETER(FIntPoint, OutputResolution)
 		SHADER_PARAMETER(uint32, bComposeDofDepth)
 		SHADER_PARAMETER(uint32, bHasHoldout)
+		SHADER_PARAMETER(uint32, SLWComposeType)
 		SHADER_PARAMETER_RDG_UNIFORM_BUFFER(FHairStrandsViewUniformParameters, HairStrands)
 		SHADER_PARAMETER_RDG_TEXTURE(Texture2D, HairLightingSampleBuffer)
 		SHADER_PARAMETER_RDG_TEXTURE(Texture2D, HairDOFDepthTexture)
+		SHADER_PARAMETER_RDG_TEXTURE(Texture2D, SLWDepthTexture)
 		SHADER_PARAMETER_RDG_TEXTURE_UAV(RWTexture2D<uint>, OutMetaTexture)
 		SHADER_PARAMETER_RDG_UNIFORM_BUFFER(FFogUniformParameters, FogStruct)
 		RENDER_TARGET_BINDING_SLOTS()
@@ -317,8 +328,9 @@ static void AddHairVisibilityComposeSamplePass(
 	const FViewInfo& View,
 	const FHairStrandsVisibilityData& VisibilityData,
 	const bool bHasHoldout,
+	const EHairSLWCompositionType SLWComposeType,
 	const FRDGTextureRef& HairDOFDepthTexture,
-	const FRDGTextureRef& HairTemporalAccumulationTexture,
+	const FRDGTextureRef SLWDepthTexture,
 	FRDGTextureRef& OutColorTexture,
 	FRDGTextureRef& OutDepthTexture)
 {
@@ -329,6 +341,7 @@ static void AddHairVisibilityComposeSamplePass(
 
 	FHairVisibilityComposeSamplePS::FParameters* Parameters = GraphBuilder.AllocParameters<FHairVisibilityComposeSamplePS::FParameters>();
 	Parameters->bComposeDofDepth = bDOFEnable ? 1 : 0;
+	Parameters->SLWComposeType = 0;
 	Parameters->HairLightingSampleBuffer = VisibilityData.SampleLightingTexture;
 	Parameters->HairDOFDepthTexture = bDOFEnable ? HairDOFDepthTexture : GSystemTextures.GetBlackDummy(GraphBuilder);
 	Parameters->OutputResolution = OutColorTexture->Desc.Extent;
@@ -339,9 +352,17 @@ static void AddHairVisibilityComposeSamplePass(
 	Parameters->RenderTargets[0] = FRenderTargetBinding(OutColorTexture, ERenderTargetLoadAction::ELoad);
 	Parameters->RenderTargets.DepthStencil = FDepthStencilBinding(OutDepthTexture, ERenderTargetLoadAction::ELoad, ERenderTargetLoadAction::ELoad, FExclusiveDepthStencil::DepthWrite_StencilRead);
 
+	const bool bSWL = SLWDepthTexture && SLWComposeType != EHairSLWCompositionType::All;
+	if (bSWL)
+	{
+		Parameters->SLWComposeType = SLWComposeType;
+		Parameters->SLWDepthTexture = SLWDepthTexture;
+	}
+
 	const bool bDebugComposition = View.Family->EngineShowFlags.LODColoration;
 	FHairVisibilityComposeSamplePS::FPermutationDomain PermutationVector;
 	PermutationVector.Set<FHairVisibilityComposeSamplePS::FDebug>(bDebugComposition);
+	PermutationVector.Set<FHairVisibilityComposeSamplePS::FSlw>(bSWL);
 	TShaderMapRef<FHairVisibilityComposeSamplePS> PixelShader(View.ShaderMap, PermutationVector);
 	InternalCommonDrawPass(
 		GraphBuilder,
@@ -658,7 +679,8 @@ static void InternalRenderHairComposition(
 	const FViewInfo& View,
 	FRDGTextureRef SceneColorTexture,
 	FRDGTextureRef SceneDepthTexture,
-	FRDGTextureRef SceneVelocityTexture)
+	FRDGTextureRef SceneVelocityTexture,
+	FRDGTextureRef SLWDepthTexture)
 {
 	RDG_EVENT_SCOPE_STAT(GraphBuilder, HairStrandsComposition, "HairStrandsComposition");
 	RDG_GPU_STAT_SCOPE(GraphBuilder, HairStrandsComposition);
@@ -696,8 +718,6 @@ static void InternalRenderHairComposition(
 				SceneDepthTexture);
 		}
 
-		FRDGTextureRef AccumulatedColor = nullptr;
-
 		const bool bHasHoldout = HasHairFlags(View.HairStrandsViewData.Flags, HAIR_FLAGS_HOLDOUT) && View.Family->EngineShowFlags.AllowPrimitiveAlphaHoldout;
 		const bool bTemporalLayeringEnabled = false; // HAIR_TODO: remove
 
@@ -706,8 +726,9 @@ static void InternalRenderHairComposition(
 			View,
 			VisibilityData,
 			bHasHoldout,
+			EHairSLWCompositionType::InFrontOfWater,
 			DOFDepth,
-			AccumulatedColor,
+			SLWDepthTexture,
 			SceneColorTexture,
 			SceneDepthTexture);
 
@@ -787,13 +808,48 @@ static void InternalRenderHairComposition(
 	}
 }
 
+void RenderHairCompositionForSLW(
+	FRDGBuilder& GraphBuilder,
+	TArrayView<const FViewInfo> Views,
+	FRDGTextureRef SLWDepthTexture,
+	FRDGTextureRef SceneColorTexture,
+	FRDGTextureRef SceneDepthTexture)
+{
+	for (const FViewInfo& View : Views)
+	{		
+		if (View.Family && HairStrands::HasViewHairStrandsData(View))
+		{
+			const FHairStrandsVisibilityData& VisibilityData = View.HairStrandsViewData.VisibilityData;
+			const bool bHasHoldout = HasHairFlags(View.HairStrandsViewData.Flags, HAIR_FLAGS_HOLDOUT) && View.Family->EngineShowFlags.AllowPrimitiveAlphaHoldout;
+			if (!VisibilityData.CoverageTexture || bHasHoldout)
+			{
+				continue; // Automatically skip for any view not rendering hair
+			}
+
+			// When rendering hair for SLW prior to the scene color copy, compose all hair samples (in front & behind) to accomodate 
+			// UV distortion which can require more data than what is stricly needed.
+			AddHairVisibilityComposeSamplePass(
+				GraphBuilder,
+				View,
+				VisibilityData,
+				false /*bHasHoldout*/,
+				EHairSLWCompositionType::All,
+				nullptr /*DOFDepth*/,
+				SLWDepthTexture,
+				SceneColorTexture,
+				SceneDepthTexture);
+		}
+	}
+}
+
 void RenderHairComposition(
 	FRDGBuilder& GraphBuilder,
 	const TArray<FViewInfo>& Views,
 	FRDGTextureRef SceneColorTexture,
 	FRDGTextureRef SceneDepthTexture,
 	FRDGTextureRef SceneVelocityTexture,
-	FTranslucencyPassResourcesMap& TranslucencyResourceMap)
+	FTranslucencyPassResourcesMap& TranslucencyResourceMap,
+	const FSingleLayerWaterPrePassResult* SingleLayerWaterPrePassResult)
 {
 	uint32 ViewIndex = 0;
 	for (const FViewInfo& View : Views)
@@ -808,7 +864,8 @@ void RenderHairComposition(
 				View,
 				SceneColorTarget,
 				SceneDepthTexture,
-				SceneVelocityTexture);
+				SceneVelocityTexture,
+				SingleLayerWaterPrePassResult ? SingleLayerWaterPrePassResult->DepthPrepassTexture.Resolve : nullptr);
 		}
 		ViewIndex++;
 	}
@@ -819,7 +876,8 @@ void RenderHairComposition(
 	const FViewInfo& View,
 	FRDGTextureRef SceneColorTexture,
 	FRDGTextureRef SceneDepthTexture,
-	FRDGTextureRef SceneVelocityTexture)
+	FRDGTextureRef SceneVelocityTexture,
+	const FSingleLayerWaterPrePassResult* SingleLayerWaterPrePassResult)
 {
 	if (View.Family && HairStrands::HasViewHairStrandsData(View))
 	{
@@ -828,6 +886,7 @@ void RenderHairComposition(
 			View,
 			SceneColorTexture,
 			SceneDepthTexture,
-			SceneVelocityTexture);
+			SceneVelocityTexture,
+			SingleLayerWaterPrePassResult ? SingleLayerWaterPrePassResult->DepthPrepassTexture.Resolve : nullptr);
 	}
 }
diff --git a/Engine/Source/Runtime/Renderer/Private/HairStrands/HairStrandsComposition.h b/Engine/Source/Runtime/Renderer/Private/HairStrands/HairStrandsComposition.h
--- a/Engine/Source/Runtime/Renderer/Private/HairStrands/HairStrandsComposition.h
+++ b/Engine/Source/Runtime/Renderer/Private/HairStrands/HairStrandsComposition.h
@@ -10,13 +10,15 @@
 #include "RenderGraphFwd.h"
 
 class FViewInfo;
+struct FSingleLayerWaterPrePassResult;
 
 void RenderHairComposition(
 	FRDGBuilder& GraphBuilder,
 	const FViewInfo& View,
 	FRDGTextureRef SceneColorTexture,
 	FRDGTextureRef SceneDepthTexture, 
-	FRDGTextureRef VelocityTexture);
+	FRDGTextureRef VelocityTexture,
+	const FSingleLayerWaterPrePassResult* SingleLayerWaterPrePassResult = nullptr);
 
 void RenderHairComposition(
 	FRDGBuilder& GraphBuilder, 
@@ -24,4 +26,12 @@ void RenderHairComposition(
 	FRDGTextureRef SceneColorTexture,
 	FRDGTextureRef SceneDepthTexture,
 	FRDGTextureRef VelocityTexture,
-	struct FTranslucencyPassResourcesMap& TranslucencyResourceMap);
+	struct FTranslucencyPassResourcesMap& TranslucencyResourceMap,
+	const FSingleLayerWaterPrePassResult* SingleLayerWaterPrePassResult = nullptr);
+
+void RenderHairCompositionForSLW(
+	FRDGBuilder& GraphBuilder,
+	TArrayView<const FViewInfo> Views,
+	FRDGTextureRef SLWDepthTexture,
+	FRDGTextureRef SceneColorTexture,
+	FRDGTextureRef SceneDepthTexture);
diff --git a/Engine/Source/Runtime/Renderer/Private/HairStrands/HairStrandsData.h b/Engine/Source/Runtime/Renderer/Private/HairStrands/HairStrandsData.h
--- a/Engine/Source/Runtime/Renderer/Private/HairStrands/HairStrandsData.h
+++ b/Engine/Source/Runtime/Renderer/Private/HairStrands/HairStrandsData.h
@@ -496,7 +496,7 @@ namespace HairStrands
 	// Hair helpers
 	bool HasHairInstanceInScene(const FScene& Scene);
 	bool HasHairCardsVisible(const TArray<FViewInfo>& Views);
-	bool HasHairStrandsVisible(const TArray<FViewInfo>& Views);
+	bool HasHairStrandsVisible(TArrayView<const FViewInfo> Views);
 
 	void AddVisibleShadowCastingLight(const FScene& Scene, TArray<FViewInfo>& Views, const FLightSceneInfo* LightSceneInfo);
 
diff --git a/Engine/Source/Runtime/Renderer/Private/HairStrands/HairStrandsRendering.cpp b/Engine/Source/Runtime/Renderer/Private/HairStrands/HairStrandsRendering.cpp
--- a/Engine/Source/Runtime/Renderer/Private/HairStrands/HairStrandsRendering.cpp
+++ b/Engine/Source/Runtime/Renderer/Private/HairStrands/HairStrandsRendering.cpp
@@ -340,7 +340,7 @@ bool HasViewHairStrandsData(const TArray<FViewInfo>& Views)
 	return false;
 }
 
-bool HasHairStrandsVisible(const TArray<FViewInfo>& Views)
+bool HasHairStrandsVisible(TArrayView<const FViewInfo> Views)
 {
 	for (const FViewInfo& View : Views)
 	{
diff --git a/Engine/Source/Runtime/Renderer/Private/SingleLayerWaterRendering.cpp b/Engine/Source/Runtime/Renderer/Private/SingleLayerWaterRendering.cpp
--- a/Engine/Source/Runtime/Renderer/Private/SingleLayerWaterRendering.cpp
+++ b/Engine/Source/Runtime/Renderer/Private/SingleLayerWaterRendering.cpp
@@ -23,6 +23,7 @@
 #include "RenderCore.h"
 #include "UnrealEngine.h"
 #include "CustomDepthRendering.h"
+#include "HairStrands/HairStrandsData.h"
 
 #include "DepthCopy.h"
 using namespace DepthCopy;
@@ -1303,6 +1304,17 @@ static FSceneWithoutWaterTextures AddCopySceneWithoutWaterPass(
 	Textures.RefractionDownsampleFactor = float(RefractionDownsampleFactor);
 	Textures.Views.SetNum(Views.Num());
 
+	// Compose hair pixels below water into the main scene color in order to be affected by water transmittance
+	if (SingleLayerWaterPrePassResult && HairStrands::HasHairStrandsVisible(Views))
+	{
+		RenderHairCompositionForSLW(
+			GraphBuilder,
+			Views,
+			SingleLayerWaterPrePassResult->DepthPrepassTexture.Target,
+			SceneColorTexture,
+			SceneDepthTexture);
+	}
+
 	AddCopySceneWithoutWaterPass_Internal(GraphBuilder, ViewFamily, Views, 
 		SceneDepthWithoutSingleLayerWaterTexture, SceneDepthTexture, 
 		SceneColorWithoutSingleLayerWaterTexture, SceneColorTexture, 
diff --git a/Engine/Source/Runtime/Renderer/Private/TranslucentRendering.cpp b/Engine/Source/Runtime/Renderer/Private/TranslucentRendering.cpp
--- a/Engine/Source/Runtime/Renderer/Private/TranslucentRendering.cpp
+++ b/Engine/Source/Runtime/Renderer/Private/TranslucentRendering.cpp
@@ -1848,6 +1848,7 @@ void FDeferredShadingSceneRenderer::RenderTranslucency(
 	const FSeparateTranslucencyDimensions& SeparateTranslucencyDimensions,
 	FInstanceCullingManager& InstanceCullingManager,
 	bool bStandardTranslucentCanRenderSeparate,
+	const FSingleLayerWaterPrePassResult* InSingleLayerWaterPrePassResult,
 	FRDGTextureMSAA& OutSharedDepthTexture)
 {
 	if (!EnumHasAnyFlags(ViewsToRender, ETranslucencyView::UnderWater | ETranslucencyView::AboveWater))
@@ -1922,7 +1923,7 @@ void FDeferredShadingSceneRenderer::RenderTranslucency(
 
 		if (GetHairStrandsComposition() == EHairStrandsCompositionType::AfterTranslucentBeforeTranslucentAfterDOF)
 		{
-			RenderHairComposition(GraphBuilder, InViews, SceneTextures.Color.Target, SceneTextures.Depth.Target, SceneTextures.Velocity, *OutTranslucencyResourceMap);
+			RenderHairComposition(GraphBuilder, InViews, SceneTextures.Color.Target, SceneTextures.Depth.Target, SceneTextures.Velocity, *OutTranslucencyResourceMap, InSingleLayerWaterPrePassResult);
 		}
 		RenderTranslucencyInner(Renderer, GraphBuilder, SceneTextures, TranslucentLightingVolumeTextures, OutTranslucencyResourceMap, OutSharedDepthTexture, InViews, ViewsToRender, SeparateTranslucencyDimensions, SceneColorCopyTexture, ETranslucencyPass::TPT_TranslucencyAfterDOF, InstanceCullingManager, bStandardTranslucentCanRenderSeparate);
 		RenderTranslucencyInner(Renderer, GraphBuilder, SceneTextures, TranslucentLightingVolumeTextures, OutTranslucencyResourceMap, OutSharedDepthTexture, InViews, ViewsToRender, SeparateTranslucencyDimensions, SceneColorCopyTexture, ETranslucencyPass::TPT_TranslucencyAfterDOFModulate, InstanceCullingManager, bStandardTranslucentCanRenderSeparate);
