MegaLights Translucency Volume - merge translucency MegaLights injection with the translucency temporal filtering pass in order to save on bandwidth (both passes were reading and writing translucency volume).

Saves 0.05ms on console with default translucency settings and 1080p

#Synced-CL 48217666

diff --git a/Engine/Shaders/Private/TranslucentLightInjectionShaders.usf b/Engine/Shaders/Private/TranslucentLightInjectionShaders.usf
--- a/Engine/Shaders/Private/TranslucentLightInjectionShaders.usf
+++ b/Engine/Shaders/Private/TranslucentLightInjectionShaders.usf
@@ -431,65 +431,6 @@ void InjectBatchMainCS(
 
 #endif // InjectBatchMainCS
 
-#ifdef InjectMegaLightsCS
-
-uint3 VolumeSize;
-
-RWTexture3D<float4> RWTranslucencyLightingVolumeAmbient;
-RWTexture3D<float4> RWTranslucencyLightingVolumeDirectional;
-
-#if !USE_UAV_TYPED_LOAD
-Texture3D<float4> TranslucencyLightingVolumeAmbient;
-Texture3D<float4> TranslucencyLightingVolumeDirectional;
-#endif
-
-Texture3D<float4> MegaLightsAmbient;
-Texture3D<float4> MegaLightsDirectional;
-
-#if INDIRECT_VOXEL_DISPATCH
-StructuredBuffer<uint> VoxelAllocator;
-StructuredBuffer<uint> VoxelData;
-#endif
-
-[numthreads(THREADGROUP_SIZE_X, THREADGROUP_SIZE_Y, THREADGROUP_SIZE_Z)]
-void InjectMegaLightsCS(
-	uint3 DispatchThreadId : SV_DispatchThreadID,
-	uint3 GroupId : SV_GroupID,
-	uint GroupIndex : SV_GroupIndex)
-{
-#if INDIRECT_VOXEL_DISPATCH
-	const uint VoxelIndex = GetUnWrappedDispatchThreadId(GroupId, GroupIndex, THREADGROUP_SIZE_X);
-
-	if (VoxelIndex >= VoxelAllocator[0])
-	{
-		return;
-	}
-
-	const uint PackedVolumeCoord = VoxelData[VoxelIndex];
-	const uint3 VolumeCoord = uint3(PackedVolumeCoord & 0x3FF, (PackedVolumeCoord >> 10) & 0x3FF, (PackedVolumeCoord >> 20) & 0x3FF);
-#else
-	const uint3 VolumeCoord = DispatchThreadId.xyz;
-
-	if (any(VolumeCoord >= VolumeSize))
-	{
-		return;
-	}
-#endif
-
-	float4 Ambient = MegaLightsAmbient[VolumeCoord];
-	float4 Directional = MegaLightsDirectional[VolumeCoord];
-
-#if USE_UAV_TYPED_LOAD
-	RWTranslucencyLightingVolumeAmbient[VolumeCoord] = RWTranslucencyLightingVolumeAmbient[VolumeCoord] + Ambient;
-	RWTranslucencyLightingVolumeDirectional[VolumeCoord] = RWTranslucencyLightingVolumeDirectional[VolumeCoord] + Directional;
-#else
-	RWTranslucencyLightingVolumeAmbient[VolumeCoord] = TranslucencyLightingVolumeAmbient[VolumeCoord] + Ambient;
-	RWTranslucencyLightingVolumeDirectional[VolumeCoord] = TranslucencyLightingVolumeDirectional[VolumeCoord] + Directional;
-#endif
-}
-
-#endif // InjectMegaLightsCS
-
 #if CLEAR_COMPUTE_SHADER
 RWTexture3D<float4> RWAmbient0;
 RWTexture3D<float4> RWDirectional0;
diff --git a/Engine/Shaders/Private/TranslucentLightingShaders.usf b/Engine/Shaders/Private/TranslucentLightingShaders.usf
--- a/Engine/Shaders/Private/TranslucentLightingShaders.usf
+++ b/Engine/Shaders/Private/TranslucentLightingShaders.usf
@@ -149,6 +149,7 @@ Texture3D TranslucencyLightingVolumeAmbient;
 SamplerState TranslucencyLightingVolumeAmbientSampler;
 Texture3D TranslucencyLightingVolumeDirectional;
 SamplerState TranslucencyLightingVolumeDirectionalSampler;
+Texture3D<uint> TranslucencyLightingVolumeMarkTexture;
 
 Texture3D HistoryAmbient;
 Texture3D HistoryDirectional;
@@ -158,6 +159,10 @@ SamplerState HistoryDirectionalSampler;
 
 Texture3D<uint> HistoryMark;
 
+// MegaLights injection
+Texture3D<float4> MegaLightsAmbient;
+Texture3D<float4> MegaLightsDirectional;
+
 uint3 VolumeSize;
 float TexelSize;
 
@@ -180,88 +185,101 @@ void FilterTranslucentVolumeCS(uint3 DispatchThreadId : SV_DispatchThreadID)
 		return;
 	}
 
-	const float3 VolumeUV = (VolumeCoord + 0.5f) * TexelSize;
-
-	float4 TextureValue0 = 0;
-	float4 TextureValue1 = 0;
-
-#define USE_FILTER 1
-
-#if USE_FILTER
-
-#if USE_TEMPORAL_REPROJECTION
-
-	TextureValue0 = TranslucencyLightingVolumeAmbient[VolumeCoord];
-	TextureValue1 = TranslucencyLightingVolumeDirectional[VolumeCoord];
-
-	float3 TranslatedWorldPosition = GetTranslatedWorldPositionFromVolumeCoord(VolumeCoord);
-	int3 HistoryCoord = (TranslatedWorldPosition - PrevTranslucencyLightingVolumeMin.xyz) / PrevTranslucencyLightingVolumeInvSize.w;
-
-	float4 HistoryTextureValue0 = HistoryAmbient[HistoryCoord];
-	float4 HistoryTextureValue1 = HistoryDirectional[HistoryCoord];
+	#if !USE_SPATIAL_FILTER && USE_MARK_TEXTURE
+	{
+		// Skip empty voxels
+		if (TranslucencyLightingVolumeMarkTexture[VolumeCoord] == 0)
+		{
+			RWTranslucencyLightingVolumeAmbient[VolumeCoord] = 0.0f;
+			RWTranslucencyLightingVolumeDirectional[VolumeCoord] = 0.0f;
+			return;
+		}
+	}
+	#endif
 
-	float HistoryAlpha = HistoryWeight;
+	const float3 VolumeUV = (VolumeCoord + 0.5f) * TexelSize;
+	float4 TextureValue0 = 0.0f;
+	float4 TextureValue1 = 0.0f;
+	float InvWeight = 1.0f;
 
-	if (any(HistoryCoord < 0) || any(HistoryCoord >= VolumeSize))
+	#if !USE_SPATIAL_FILTER
 	{
-		HistoryAlpha = 0;
+		TextureValue0 = TranslucencyLightingVolumeAmbient[VolumeCoord];
+		TextureValue1 = TranslucencyLightingVolumeDirectional[VolumeCoord];
+
+		#if INJECT_MEGALIGHTS
+		{
+			TextureValue0 += MegaLightsAmbient[VolumeCoord];
+			TextureValue1 += MegaLightsDirectional[VolumeCoord];
+		}
+		#endif
 	}
+	#endif
 
-#if CHECK_HISTORY_MARK
-	// don't use history if it was not computed
-	if (HistoryMark[HistoryCoord] == 0)
+	#if USE_TEMPORAL_FILTER
 	{
-		HistoryAlpha = 0;
+		float3 TranslatedWorldPosition = GetTranslatedWorldPositionFromVolumeCoord(VolumeCoord);
+		int3 HistoryCoord = (TranslatedWorldPosition - PrevTranslucencyLightingVolumeMin.xyz) / PrevTranslucencyLightingVolumeInvSize.w;
+
+		float4 HistoryTextureValue0 = HistoryAmbient[HistoryCoord];
+		float4 HistoryTextureValue1 = HistoryDirectional[HistoryCoord];
+
+		float HistoryAlpha = HistoryWeight;
+
+		if (any(HistoryCoord < 0) || any(HistoryCoord >= VolumeSize))
+		{
+			HistoryAlpha = 0;
+		}
+
+		#if USE_MARK_TEXTURE
+		{
+			// don't use history if it was not computed
+			if (HistoryMark[HistoryCoord] == 0)
+			{
+				HistoryAlpha = 0;
+			}
+		}
+		#endif
+
+		if (VolumeCascadeIndex == 1)
+		{
+			// fade out at border to prevent issues with usage of clamped sampler during shading passes
+			HistoryAlpha *= GetBorderLerpFactorFromVolumeUVs(VolumeUV, VolumeCascadeIndex);
+		}
+
+		TextureValue0 = lerp(TextureValue0, HistoryTextureValue0, HistoryAlpha);
+		TextureValue1 = lerp(TextureValue1, HistoryTextureValue1, HistoryAlpha);
 	}
-#endif
-
-	if (VolumeCascadeIndex == 1)
+	#elif USE_SPATIAL_FILTER
 	{
-		// fade out at border to prevent issues with usage of clamped sampler during shading passes
-		HistoryAlpha *= GetBorderLerpFactorFromVolumeUVs(VolumeUV, VolumeCascadeIndex);
+		TextureValue0 = 0.0f;
+		TextureValue1 = 0.0f;
+
+		// Use trilinear filtering to filter neighbors to the current voxel with minimal texture fetches
+		TextureValue0 += Texture3DSample(TranslucencyLightingVolumeAmbient, TranslucencyLightingVolumeAmbientSampler, VolumeUV + .5f * TexelSize * float3(1, 1, 1));
+		TextureValue0 += Texture3DSample(TranslucencyLightingVolumeAmbient, TranslucencyLightingVolumeAmbientSampler, VolumeUV + .5f * TexelSize * float3(-1, 1, 1));
+		TextureValue0 += Texture3DSample(TranslucencyLightingVolumeAmbient, TranslucencyLightingVolumeAmbientSampler, VolumeUV + .5f * TexelSize * float3(1, -1, 1));
+		TextureValue0 += Texture3DSample(TranslucencyLightingVolumeAmbient, TranslucencyLightingVolumeAmbientSampler, VolumeUV + .5f * TexelSize * float3(-1, -1, 1));
+		TextureValue0 += Texture3DSample(TranslucencyLightingVolumeAmbient, TranslucencyLightingVolumeAmbientSampler, VolumeUV + .5f * TexelSize * float3(1, 1, -1));
+		TextureValue0 += Texture3DSample(TranslucencyLightingVolumeAmbient, TranslucencyLightingVolumeAmbientSampler, VolumeUV + .5f * TexelSize * float3(-1, 1, -1));
+		TextureValue0 += Texture3DSample(TranslucencyLightingVolumeAmbient, TranslucencyLightingVolumeAmbientSampler, VolumeUV + .5f * TexelSize * float3(1, -1, -1));
+		TextureValue0 += Texture3DSample(TranslucencyLightingVolumeAmbient, TranslucencyLightingVolumeAmbientSampler, VolumeUV + .5f * TexelSize * float3(-1, -1, -1));
+
+		TextureValue1 += Texture3DSample(TranslucencyLightingVolumeDirectional, TranslucencyLightingVolumeDirectionalSampler, VolumeUV + .5f * TexelSize * float3(1, 1, 1));
+		TextureValue1 += Texture3DSample(TranslucencyLightingVolumeDirectional, TranslucencyLightingVolumeDirectionalSampler, VolumeUV + .5f * TexelSize * float3(-1, 1, 1));
+		TextureValue1 += Texture3DSample(TranslucencyLightingVolumeDirectional, TranslucencyLightingVolumeDirectionalSampler, VolumeUV + .5f * TexelSize * float3(1, -1, 1));
+		TextureValue1 += Texture3DSample(TranslucencyLightingVolumeDirectional, TranslucencyLightingVolumeDirectionalSampler, VolumeUV + .5f * TexelSize * float3(-1, -1, 1));
+		TextureValue1 += Texture3DSample(TranslucencyLightingVolumeDirectional, TranslucencyLightingVolumeDirectionalSampler, VolumeUV + .5f * TexelSize * float3(1, 1, -1));
+		TextureValue1 += Texture3DSample(TranslucencyLightingVolumeDirectional, TranslucencyLightingVolumeDirectionalSampler, VolumeUV + .5f * TexelSize * float3(-1, 1, -1));
+		TextureValue1 += Texture3DSample(TranslucencyLightingVolumeDirectional, TranslucencyLightingVolumeDirectionalSampler, VolumeUV + .5f * TexelSize * float3(1, -1, -1));
+		TextureValue1 += Texture3DSample(TranslucencyLightingVolumeDirectional, TranslucencyLightingVolumeDirectionalSampler, VolumeUV + .5f * TexelSize * float3(-1, -1, -1));
+
+		InvWeight = 1.0f / 8;
 	}
-
-
-	TextureValue0 = lerp(TextureValue0, HistoryTextureValue0, HistoryAlpha);
-	TextureValue1 = lerp(TextureValue1, HistoryTextureValue1, HistoryAlpha);
-
-	float InvWeight = 1;
-
-#else // !USE_TEMPORAL_REPROJECTION
-
-	// Use trilinear filtering to filter neighbors to the current voxel with minimal texture fetches
-	TextureValue0 += Texture3DSample(TranslucencyLightingVolumeAmbient, TranslucencyLightingVolumeAmbientSampler, VolumeUV + .5f * TexelSize * float3(1, 1, 1));
-	TextureValue0 += Texture3DSample(TranslucencyLightingVolumeAmbient, TranslucencyLightingVolumeAmbientSampler, VolumeUV + .5f * TexelSize * float3(-1, 1, 1));
-	TextureValue0 += Texture3DSample(TranslucencyLightingVolumeAmbient, TranslucencyLightingVolumeAmbientSampler, VolumeUV + .5f * TexelSize * float3(1, -1, 1));
-	TextureValue0 += Texture3DSample(TranslucencyLightingVolumeAmbient, TranslucencyLightingVolumeAmbientSampler, VolumeUV + .5f * TexelSize * float3(-1, -1, 1));
-	TextureValue0 += Texture3DSample(TranslucencyLightingVolumeAmbient, TranslucencyLightingVolumeAmbientSampler, VolumeUV + .5f * TexelSize * float3(1, 1, -1));
-	TextureValue0 += Texture3DSample(TranslucencyLightingVolumeAmbient, TranslucencyLightingVolumeAmbientSampler, VolumeUV + .5f * TexelSize * float3(-1, 1, -1));
-	TextureValue0 += Texture3DSample(TranslucencyLightingVolumeAmbient, TranslucencyLightingVolumeAmbientSampler, VolumeUV + .5f * TexelSize * float3(1, -1, -1));
-	TextureValue0 += Texture3DSample(TranslucencyLightingVolumeAmbient, TranslucencyLightingVolumeAmbientSampler, VolumeUV + .5f * TexelSize * float3(-1, -1, -1));
-
-	TextureValue1 += Texture3DSample(TranslucencyLightingVolumeDirectional, TranslucencyLightingVolumeDirectionalSampler, VolumeUV + .5f * TexelSize * float3(1, 1, 1));
-	TextureValue1 += Texture3DSample(TranslucencyLightingVolumeDirectional, TranslucencyLightingVolumeDirectionalSampler, VolumeUV + .5f * TexelSize * float3(-1, 1, 1));
-	TextureValue1 += Texture3DSample(TranslucencyLightingVolumeDirectional, TranslucencyLightingVolumeDirectionalSampler, VolumeUV + .5f * TexelSize * float3(1, -1, 1));
-	TextureValue1 += Texture3DSample(TranslucencyLightingVolumeDirectional, TranslucencyLightingVolumeDirectionalSampler, VolumeUV + .5f * TexelSize * float3(-1, -1, 1));
-	TextureValue1 += Texture3DSample(TranslucencyLightingVolumeDirectional, TranslucencyLightingVolumeDirectionalSampler, VolumeUV + .5f * TexelSize * float3(1, 1, -1));
-	TextureValue1 += Texture3DSample(TranslucencyLightingVolumeDirectional, TranslucencyLightingVolumeDirectionalSampler, VolumeUV + .5f * TexelSize * float3(-1, 1, -1));
-	TextureValue1 += Texture3DSample(TranslucencyLightingVolumeDirectional, TranslucencyLightingVolumeDirectionalSampler, VolumeUV + .5f * TexelSize * float3(1, -1, -1));
-	TextureValue1 += Texture3DSample(TranslucencyLightingVolumeDirectional, TranslucencyLightingVolumeDirectionalSampler, VolumeUV + .5f * TexelSize * float3(-1, -1, -1));
-
-	float InvWeight = 1.0f / 8;
-
-#endif // !USE_TEMPORAL_REPROJECTION
-
-#else
-
-	TextureValue0 = Texture3DSample(TranslucencyLightingVolumeAmbient, TranslucencyLightingVolumeAmbientSampler, VolumeUV);
-	TextureValue1 = Texture3DSample(TranslucencyLightingVolumeDirectional, TranslucencyLightingVolumeDirectionalSampler, VolumeUV);
-	float InvWeight = 1;
-
-#endif
+	#endif
 	
-	RWTranslucencyLightingVolumeAmbient[DispatchThreadId] = TextureValue0 * InvWeight;
-	RWTranslucencyLightingVolumeDirectional[DispatchThreadId] = TextureValue1 * InvWeight;
+	RWTranslucencyLightingVolumeAmbient[VolumeCoord] = TextureValue0 * InvWeight;
+	RWTranslucencyLightingVolumeDirectional[VolumeCoord] = TextureValue1 * InvWeight;
 }
 
 #endif // FilterTranslucentVolumeCS
diff --git a/Engine/Source/Runtime/Renderer/Private/TranslucentLighting.cpp b/Engine/Source/Runtime/Renderer/Private/TranslucentLighting.cpp
--- a/Engine/Source/Runtime/Renderer/Private/TranslucentLighting.cpp
+++ b/Engine/Source/Runtime/Renderer/Private/TranslucentLighting.cpp
@@ -869,6 +869,7 @@ public:
 		SHADER_PARAMETER_RDG_TEXTURE(Texture3D, TranslucencyLightingVolumeDirectional)
 		SHADER_PARAMETER_SAMPLER(SamplerState, TranslucencyLightingVolumeAmbientSampler)
 		SHADER_PARAMETER_SAMPLER(SamplerState, TranslucencyLightingVolumeDirectionalSampler)
+		SHADER_PARAMETER_RDG_TEXTURE(Texture3D, TranslucencyLightingVolumeMarkTexture)
 
 		SHADER_PARAMETER_RDG_TEXTURE(Texture3D, HistoryAmbient)
 		SHADER_PARAMETER_RDG_TEXTURE(Texture3D, HistoryDirectional)
@@ -877,6 +878,9 @@ public:
 
 		SHADER_PARAMETER_RDG_TEXTURE(Texture3D, HistoryMark)
 
+		SHADER_PARAMETER_RDG_TEXTURE_SRV(Texture3D, MegaLightsAmbient)
+		SHADER_PARAMETER_RDG_TEXTURE_SRV(Texture3D, MegaLightsDirectional)
+
 		SHADER_PARAMETER(FIntVector, VolumeSize)
 		SHADER_PARAMETER(float, TexelSize)
 		SHADER_PARAMETER(uint32, VolumeCascadeIndex)
@@ -890,12 +894,32 @@ public:
 		SHADER_PARAMETER(float, HistoryWeight)
 	END_SHADER_PARAMETER_STRUCT()
 
-	class FUseTemporalReprojection : SHADER_PERMUTATION_BOOL("USE_TEMPORAL_REPROJECTION");
-	class FCheckHistoryMark : SHADER_PERMUTATION_BOOL("CHECK_HISTORY_MARK");
-	using FPermutationDomain = TShaderPermutationDomain<FUseTemporalReprojection, FCheckHistoryMark>;
+	class FUseTemporalFilter : SHADER_PERMUTATION_BOOL("USE_TEMPORAL_FILTER");
+	class FUseSpatialFilter : SHADER_PERMUTATION_BOOL("USE_SPATIAL_FILTER");
+	class FUseMarkTexture : SHADER_PERMUTATION_BOOL("USE_MARK_TEXTURE");
+	class FInjectMegaLights : SHADER_PERMUTATION_BOOL("INJECT_MEGALIGHTS");
+	using FPermutationDomain = TShaderPermutationDomain<FUseTemporalFilter, FUseSpatialFilter, FUseMarkTexture, FInjectMegaLights>;
+
+	static FPermutationDomain RemapPermutation(FPermutationDomain PermutationVector)
+	{
+		if (PermutationVector.Get<FUseSpatialFilter>())
+		{
+			PermutationVector.Set<FUseTemporalFilter>(false);
+			PermutationVector.Set<FInjectMegaLights>(false);
+			PermutationVector.Set<FUseMarkTexture>(false);
+		}
+
+		return PermutationVector;
+	}
 
 	static bool ShouldCompilePermutation(const FGlobalShaderPermutationParameters& Parameters) 
 	{ 
+		FPermutationDomain PermutationVector(Parameters.PermutationId);
+		if (RemapPermutation(PermutationVector) != PermutationVector)
+		{
+			return false;
+		}
+
 		return IsFeatureLevelSupported(Parameters.Platform, ERHIFeatureLevel::SM5);
 	}
 
@@ -1373,200 +1397,6 @@ void InjectTranslucencyLightingVolumeAmbientCubemap(
 	}
 }
 
-class FInjectMegaLightsCS : public FGlobalShader
-{
-public:
-	DECLARE_GLOBAL_SHADER(FInjectMegaLightsCS)
-	SHADER_USE_PARAMETER_STRUCT(FInjectMegaLightsCS, FGlobalShader)
-
-	BEGIN_SHADER_PARAMETER_STRUCT(FParameters, )
-		SHADER_PARAMETER_STRUCT_REF(FViewUniformShaderParameters, View)
-
-		SHADER_PARAMETER_RDG_TEXTURE_UAV(RWTexture3D, RWTranslucencyLightingVolumeAmbient)
-		SHADER_PARAMETER_RDG_TEXTURE_UAV(RWTexture3D, RWTranslucencyLightingVolumeDirectional)
-
-		// used when UAV typed loads are not supported
-		SHADER_PARAMETER_RDG_TEXTURE(Texture3D, TranslucencyLightingVolumeAmbient)
-		SHADER_PARAMETER_RDG_TEXTURE(Texture3D, TranslucencyLightingVolumeDirectional)
-
-		SHADER_PARAMETER_RDG_TEXTURE_SRV(Texture3D, MegaLightsAmbient)
-		SHADER_PARAMETER_RDG_TEXTURE_SRV(Texture3D, MegaLightsDirectional)
-
-		SHADER_PARAMETER(FIntVector, VolumeSize)
-		SHADER_PARAMETER(uint32, VolumeCascadeIndex)
-
-		SHADER_PARAMETER_RDG_BUFFER_SRV(StructuredBuffer, VoxelAllocator)
-		SHADER_PARAMETER_RDG_BUFFER_SRV(StructuredBuffer, VoxelData)
-
-		RDG_BUFFER_ACCESS(IndirectArgs, ERHIAccess::IndirectArgs)
-	END_SHADER_PARAMETER_STRUCT()
-
-	class FIndirectVoxelDispatch : SHADER_PERMUTATION_BOOL("INDIRECT_VOXEL_DISPATCH");
-	class FUseUAVTypedLoad : SHADER_PERMUTATION_BOOL("USE_UAV_TYPED_LOAD");
-	using FPermutationDomain = TShaderPermutationDomain<FIndirectVoxelDispatch, FUseUAVTypedLoad>;
-
-	static bool ShouldCompilePermutation(const FGlobalShaderPermutationParameters& Parameters)
-	{
-		return IsFeatureLevelSupported(Parameters.Platform, ERHIFeatureLevel::SM5);
-	}
-
-	static FIntVector GetGroupSize()
-	{
-		return FIntVector(4, 4, 4);
-	}
-
-	static void ModifyCompilationEnvironment(const FGlobalShaderPermutationParameters& Parameters, FShaderCompilerEnvironment& OutEnvironment)
-	{
-		FGlobalShader::ModifyCompilationEnvironment(Parameters, OutEnvironment);
-
-		const FIntVector GroupSize = GetGroupSize();
-
-		const FPermutationDomain PermutationVector(Parameters.PermutationId);
-
-		if (PermutationVector.Get<FIndirectVoxelDispatch>())
-		{
-			OutEnvironment.SetDefine(TEXT("THREADGROUP_SIZE_X"), GroupSize.X * GroupSize.Y * GroupSize.Z);
-			OutEnvironment.SetDefine(TEXT("THREADGROUP_SIZE_Y"), 1);
-			OutEnvironment.SetDefine(TEXT("THREADGROUP_SIZE_Z"), 1);
-		}
-		else
-		{
-			OutEnvironment.SetDefine(TEXT("THREADGROUP_SIZE_X"), GroupSize.X);
-			OutEnvironment.SetDefine(TEXT("THREADGROUP_SIZE_Y"), GroupSize.Y);
-			OutEnvironment.SetDefine(TEXT("THREADGROUP_SIZE_Z"), GroupSize.Z);
-		}
-
-		// This shader must support typed UAV load and we are testing if it is supported at runtime
-		// using UE::PixelFormat::HasCapabilities(..., EPixelFormatCapabilities::TypedUAVLoad)
-		OutEnvironment.CompilerFlags.Add(CFLAG_AllowTypedUAVLoads);
-	}
-};
-
-IMPLEMENT_GLOBAL_SHADER(FInjectMegaLightsCS, "/Engine/Private/TranslucentLightInjectionShaders.usf", "InjectMegaLightsCS", SF_Compute);
-
-void InjectTranslucencyLightingVolumeMegaLights(
-	FRDGBuilder& GraphBuilder,
-	const TArrayView<const FViewInfo> Views,
-	const FTranslucencyLightingVolumeTextures& Textures)
-{
-	if (!GUseTranslucencyLightingVolumes || Views.Num() == 0 || !MegaLights::IsEnabled(*Views[0].Family) || !MegaLights::UseTranslucencyVolume())
-	{
-		return;
-	}
-
-	RDG_EVENT_SCOPE_STAT(GraphBuilder, TranslucentLighting, "InjectTranslucencyLightingVolumeMegaLights");
-	RDG_GPU_STAT_SCOPE(GraphBuilder, TranslucentLighting);
-
-	const FIntVector VolumeSize = FIntVector(GetTranslucencyLightingVolumeDim());
-
-	EPixelFormat TranslucencyPixelFormat;
-	ETextureCreateFlags TranslucencyTargetFlags;
-	FTranslucencyLightingVolumeTextures::GetTextureFormatAndCreationFlags(TranslucencyPixelFormat, TranslucencyTargetFlags);
-
-	const bool bUseUAVTypedLoad = UE::PixelFormat::HasCapabilities(TranslucencyPixelFormat, EPixelFormatCapabilities::TypedUAVLoad);
-
-	for (int32 ViewIndex = 0; ViewIndex < Views.Num(); ++ViewIndex)
-	{
-		const FViewInfo& View = Views[ViewIndex];
-		RDG_GPU_MASK_SCOPE(GraphBuilder, View.GPUMask);
-		RDG_EVENT_SCOPE_CONDITIONAL(GraphBuilder, Views.Num() > 1, "View%d", ViewIndex);
-
-		{
-			FRDGTextureAccessArray BatchBarrierTextures;
-
-			for (int32 VolumeCascadeIndex = 0; VolumeCascadeIndex < TVC_MAX; ++VolumeCascadeIndex)
-			{
-				FRDGTextureRef MegaLightsAmbient = View.GetMegaLightsVolume().TranslucencyAmbient[VolumeCascadeIndex];
-				FRDGTextureRef MegaLightsDirectional = View.GetMegaLightsVolume().TranslucencyDirectional[VolumeCascadeIndex];
-
-				if (MegaLightsAmbient)
-				{
-					BatchBarrierTextures.Emplace(MegaLightsAmbient, ERHIAccess::SRVCompute);
-					BatchBarrierTextures.Emplace(MegaLightsDirectional, ERHIAccess::SRVCompute);
-				}
-			}
-
-			AddBarrierPass(GraphBuilder, ERDGPassFlags::Compute, BatchBarrierTextures);
-		}
-
-		for (int32 VolumeCascadeIndex = 0; VolumeCascadeIndex < TVC_MAX; ++VolumeCascadeIndex)
-		{
-			const bool bUseVolumeMarkTexture = IsTranslucencyLightingVolumeUsingVoxelMarking() && View.TranslucencyVolumeMarkData[VolumeCascadeIndex].MarkTexture != nullptr;
-
-			// for stereo case, using PrimaryViewIndex essentially shares the lighting volume textures
-			const uint32 TextureIndex = Textures.GetIndex(View, VolumeCascadeIndex);
-			FRDGTextureRef VolumeAmbientTexture = Textures.Ambient[TextureIndex];
-			FRDGTextureRef VolumeDirectionalTexture = Textures.Directional[TextureIndex];
-
-			FRDGTextureRef MegaLightsAmbient = View.GetMegaLightsVolume().TranslucencyAmbient[VolumeCascadeIndex];
-			FRDGTextureRef MegaLightsDirectional = View.GetMegaLightsVolume().TranslucencyDirectional[VolumeCascadeIndex];
-
-			auto* PassParameters = GraphBuilder.AllocParameters<FInjectMegaLightsCS::FParameters>();
-
-			PassParameters->View = View.ViewUniformBuffer;
-			PassParameters->VolumeSize = VolumeSize;
-			PassParameters->VolumeCascadeIndex = VolumeCascadeIndex;
-
-			PassParameters->MegaLightsAmbient = MegaLightsAmbient ? GraphBuilder.CreateSRV(MegaLightsAmbient) : nullptr;
-			PassParameters->MegaLightsDirectional = MegaLightsDirectional ? GraphBuilder.CreateSRV(MegaLightsDirectional) : nullptr;
-
-			if (bUseUAVTypedLoad)
-			{
-				PassParameters->RWTranslucencyLightingVolumeAmbient = GraphBuilder.CreateUAV(VolumeAmbientTexture);
-				PassParameters->RWTranslucencyLightingVolumeDirectional = GraphBuilder.CreateUAV(VolumeDirectionalTexture);
-			}
-			else
-			{
-				// need to output to auxiliary textures to be able to read from existing volumes
-
-				FRDGTextureRef OutputVolumeAmbientTexture = GraphBuilder.CreateTexture(VolumeAmbientTexture->Desc, VolumeAmbientTexture->Name);
-				FRDGTextureRef OutputVolumeDirectionalTexture = GraphBuilder.CreateTexture(VolumeDirectionalTexture->Desc, VolumeDirectionalTexture->Name);
-
-				PassParameters->RWTranslucencyLightingVolumeAmbient = GraphBuilder.CreateUAV(OutputVolumeAmbientTexture);
-				PassParameters->RWTranslucencyLightingVolumeDirectional = GraphBuilder.CreateUAV(OutputVolumeDirectionalTexture);
-
-				PassParameters->TranslucencyLightingVolumeAmbient = VolumeAmbientTexture;
-				PassParameters->TranslucencyLightingVolumeDirectional = VolumeDirectionalTexture;
-
-				VolumeAmbientTexture = OutputVolumeAmbientTexture;
-				VolumeDirectionalTexture = OutputVolumeDirectionalTexture;
-			}
-
-			FInjectMegaLightsCS::FPermutationDomain PermutationVector;
-			PermutationVector.Set<FInjectMegaLightsCS::FIndirectVoxelDispatch>(bUseVolumeMarkTexture);
-			PermutationVector.Set<FInjectMegaLightsCS::FUseUAVTypedLoad>(bUseUAVTypedLoad);
-
-			auto ComputeShader = View.ShaderMap->GetShader<FInjectMegaLightsCS>(PermutationVector);
-
-			if (bUseVolumeMarkTexture)
-			{
-				PassParameters->VoxelAllocator = GraphBuilder.CreateSRV(View.TranslucencyVolumeMarkData[VolumeCascadeIndex].VoxelAllocator);
-				PassParameters->VoxelData = GraphBuilder.CreateSRV(View.TranslucencyVolumeMarkData[VolumeCascadeIndex].VoxelData);
-				PassParameters->IndirectArgs = View.TranslucencyVolumeMarkData[VolumeCascadeIndex].VoxelIndirectArgs;
-
-				FComputeShaderUtils::AddPass(
-					GraphBuilder,
-					RDG_EVENT_NAME("InjectMegaLights(VolumeCascade=%d)", VolumeCascadeIndex),
-					ComputeShader,
-					PassParameters,
-					View.TranslucencyVolumeMarkData[VolumeCascadeIndex].VoxelIndirectArgs,
-					0);
-			}
-			else
-			{
-				const FIntVector NumGroups = FComputeShaderUtils::GetGroupCount(VolumeSize, FInjectMegaLightsCS::GetGroupSize());
-
-				FComputeShaderUtils::AddPass(
-					GraphBuilder,
-					RDG_EVENT_NAME("InjectMegaLights(VolumeCascade=%d)", VolumeCascadeIndex),
-					ComputeShader,
-					PassParameters,
-					NumGroups);
-			}
-		}
-	}
-}
-
 /** Calculates volume texture bounds for the given light in the given translucency lighting volume cascade. */
 FVolumeBounds CalculateLightVolumeBounds(const FSphere& LightBounds, const FViewInfo& View, uint32 VolumeCascadeIndex, bool bDirectionalLight)
 {
@@ -2406,12 +2236,16 @@ static void DrawDebugTranslucencyLightingVolume(
 void FilterTranslucencyLightingVolume(
 	FRDGBuilder& GraphBuilder,
 	const TArrayView<const FViewInfo> Views,
-	FTranslucencyLightingVolumeTextures& Textures)
+	FTranslucencyLightingVolumeTextures& Textures,
+	const FSortedLightSetSceneInfo& SortedLightSet)
 {
 	// Draw debug information if needed.
 	DrawDebugTranslucencyLightingVolume(GraphBuilder, Views);
 
-	const bool bMegaLightsTranslucencyVolume = Views.Num() > 0 && MegaLights::IsEnabled(*Views[0].Family) && MegaLights::UseTranslucencyVolume();
+	const bool bMegaLightsTranslucencyVolume = Views.Num() > 0
+		&& MegaLights::IsEnabled(*Views[0].Family)
+		&& MegaLights::UseTranslucencyVolume()
+		&& SortedLightSet.MegaLightsLightStart < SortedLightSet.SortedLights.Num();
 
 	const bool bSpatialFilter = bMegaLightsTranslucencyVolume
 		? MegaLights::IsTranslucencyVolumeSpatialFilterEnabled()
@@ -2421,7 +2255,7 @@ void FilterTranslucencyLightingVolume(
 		? MegaLights::IsTranslucencyVolumeTemporalFilterEnabled()
 		: CVarTranslucencyLightingVolumeTemporal.GetValueOnRenderThread();
 
-	if (!GUseTranslucencyLightingVolumes || Views.Num() == 0 || !RHISupportsVolumeTextureRendering(Views[0].GetShaderPlatform()) || (!bSpatialFilter && !bTemporalFilter))
+	if (!GUseTranslucencyLightingVolumes || Views.Num() == 0 || !RHISupportsVolumeTextureRendering(Views[0].GetShaderPlatform()) || (!bSpatialFilter && !bTemporalFilter && !bMegaLightsTranslucencyVolume))
 	{
 		// discard history since we are not updating it this frame
 
@@ -2449,7 +2283,7 @@ void FilterTranslucencyLightingVolume(
 	RDG_EVENT_SCOPE_STAT(GraphBuilder, TranslucentLighting, "FilterTranslucentVolume %dx%dx%d Cascades:%d", TranslucencyLightingVolumeDim, TranslucencyLightingVolumeDim, TranslucencyLightingVolumeDim, TVC_MAX);
 	RDG_GPU_STAT_SCOPE(GraphBuilder, TranslucentLighting);
 
-	if (bTemporalFilter)
+	if (bTemporalFilter || bMegaLightsTranslucencyVolume)
 	{
 		for (int32 ViewIndex = 0; ViewIndex < Views.Num(); ++ViewIndex)
 		{
@@ -2463,14 +2297,22 @@ void FilterTranslucencyLightingVolume(
 				for (int32 VolumeCascadeIndex = 0; VolumeCascadeIndex < TVC_MAX; ++VolumeCascadeIndex)
 				{
 					const uint32 TextureIndex = Textures.GetIndex(View, VolumeCascadeIndex);
+
 					FRDGTextureRef InputVolumeAmbientTexture = Textures.Ambient[TextureIndex];
 					FRDGTextureRef InputVolumeDirectionalTexture = Textures.Directional[TextureIndex];
-
-					if (InputVolumeAmbientTexture)
+					if (InputVolumeAmbientTexture && InputVolumeDirectionalTexture)
 					{
 						BatchBarrierTextures.Emplace(InputVolumeAmbientTexture, ERHIAccess::SRVCompute);
 						BatchBarrierTextures.Emplace(InputVolumeDirectionalTexture, ERHIAccess::SRVCompute);
 					}
+
+					FRDGTextureRef MegaLightsAmbient = View.GetMegaLightsVolume().TranslucencyAmbient[VolumeCascadeIndex];
+					FRDGTextureRef MegaLightsDirectional = View.GetMegaLightsVolume().TranslucencyDirectional[VolumeCascadeIndex];
+					if (MegaLightsAmbient && MegaLightsDirectional)
+					{
+						BatchBarrierTextures.Emplace(MegaLightsAmbient, ERHIAccess::SRVCompute);
+						BatchBarrierTextures.Emplace(MegaLightsDirectional, ERHIAccess::SRVCompute);
+					}
 				}
 
 				AddBarrierPass(GraphBuilder, ERDGPassFlags::Compute | ERDGPassFlags::NeverCull, BatchBarrierTextures);
@@ -2478,6 +2320,14 @@ void FilterTranslucencyLightingVolume(
 
 			for (int32 VolumeCascadeIndex = 0; VolumeCascadeIndex < TVC_MAX; VolumeCascadeIndex++)
 			{
+				const uint32 TextureIndex = Textures.GetIndex(View, VolumeCascadeIndex);
+
+				const FMegaLightsVolume& MegaLightsVolume = View.GetMegaLightsVolume();
+				FRDGTextureRef MegaLightsAmbient = MegaLightsVolume.TranslucencyAmbient[VolumeCascadeIndex];
+				FRDGTextureRef MegaLightsDirectional = MegaLightsVolume.TranslucencyDirectional[VolumeCascadeIndex];
+				FRDGTextureRef TranslucencyLightingVolumeMarkTexture = View.TranslucencyVolumeMarkData[VolumeCascadeIndex].MarkTexture;
+				const bool bUseVolumeMarkTexture = IsTranslucencyLightingVolumeUsingVoxelMarking() && TranslucencyLightingVolumeMarkTexture != nullptr;
+
 				bool bHasValidHistory = false;
 
 				FRDGTextureRef HistoryAmbient;
@@ -2487,7 +2337,7 @@ void FilterTranslucencyLightingVolume(
 				FVector4f PrevTranslucencyLightingVolumeMin;
 				FVector4f PrevTranslucencyLightingVolumeInvSize;
 
-				if (View.ViewState && !View.bCameraCut && !View.bPrevTransformsReset)
+				if (View.ViewState && !View.bCameraCut && !View.bPrevTransformsReset && bTemporalFilter)
 				{
 					FTranslucencyLightingViewState& TranslucencyLightingViewState = View.ViewState->TranslucencyLighting;
 
@@ -2495,8 +2345,6 @@ void FilterTranslucencyLightingVolume(
 					{
 						HistoryAmbient = GraphBuilder.RegisterExternalTexture(TranslucencyLightingViewState.HistoryAmbient[VolumeCascadeIndex]);
 						HistoryDirectional = GraphBuilder.RegisterExternalTexture(TranslucencyLightingViewState.HistoryDirectional[VolumeCascadeIndex]);
-
-						bHasValidHistory = true;
 					}
 					else
 					{
@@ -2514,6 +2362,8 @@ void FilterTranslucencyLightingVolume(
 						HistoryMark = nullptr;
 					}
 
+					bHasValidHistory = HistoryAmbient && HistoryDirectional && (!bUseVolumeMarkTexture || HistoryMark);
+
 					{
 						const FVector VolumeWorldMin = TranslucencyLightingViewState.HistoryVolumeMin[VolumeCascadeIndex];
 						const float VolumeVoxelSize = TranslucencyLightingViewState.HistoryVoxelSize[VolumeCascadeIndex];
@@ -2536,9 +2386,7 @@ void FilterTranslucencyLightingVolume(
 					PrevTranslucencyLightingVolumeInvSize = FVector4f::Zero();
 				}
 
-				const uint32 TextureIndex = Textures.GetIndex(View, VolumeCascadeIndex);
-
-				if (bHasValidHistory)
+				if (bHasValidHistory || bMegaLightsTranslucencyVolume)
 				{
 					FRDGTextureRef InputVolumeAmbientTexture = Textures.Ambient[TextureIndex];
 					FRDGTextureRef InputVolumeDirectionalTexture = Textures.Directional[TextureIndex];
@@ -2560,6 +2408,7 @@ void FilterTranslucencyLightingVolume(
 					PassParameters->TranslucencyLightingVolumeDirectional = InputVolumeDirectionalTexture;
 					PassParameters->TranslucencyLightingVolumeAmbientSampler = SamplerStateRHI;
 					PassParameters->TranslucencyLightingVolumeDirectionalSampler = SamplerStateRHI;
+					PassParameters->TranslucencyLightingVolumeMarkTexture = TranslucencyLightingVolumeMarkTexture;
 					PassParameters->RWTranslucencyLightingVolumeAmbient = GraphBuilder.CreateUAV(OutputVolumeAmbientTexture);
 					PassParameters->RWTranslucencyLightingVolumeDirectional = GraphBuilder.CreateUAV(OutputVolumeDirectionalTexture);
 
@@ -2576,11 +2425,17 @@ void FilterTranslucencyLightingVolume(
 
 					PassParameters->HistoryMark = HistoryMark;
 
+					PassParameters->MegaLightsAmbient = MegaLightsAmbient ? GraphBuilder.CreateSRV(MegaLightsAmbient) : nullptr;
+					PassParameters->MegaLightsDirectional = MegaLightsDirectional ? GraphBuilder.CreateSRV(MegaLightsDirectional) : nullptr;
+
 					PassParameters->HistoryWeight = CVarTranslucencyLightingVolumeHistoryWeight.GetValueOnRenderThread();
 
 					FFilterTranslucentVolumeCS::FPermutationDomain PermutationVector;
-					PermutationVector.Set<FFilterTranslucentVolumeCS::FUseTemporalReprojection>(true);
-					PermutationVector.Set<FFilterTranslucentVolumeCS::FCheckHistoryMark>(HistoryMark != nullptr);
+					PermutationVector.Set<FFilterTranslucentVolumeCS::FUseTemporalFilter>(bHasValidHistory && bTemporalFilter);
+					PermutationVector.Set<FFilterTranslucentVolumeCS::FUseSpatialFilter>(false);
+					PermutationVector.Set<FFilterTranslucentVolumeCS::FUseMarkTexture>(bUseVolumeMarkTexture);
+					PermutationVector.Set<FFilterTranslucentVolumeCS::FInjectMegaLights>(bMegaLightsTranslucencyVolume);
+					PermutationVector = FFilterTranslucentVolumeCS::RemapPermutation(PermutationVector);
 
 					auto ComputeShader = View.ShaderMap->GetShader<FFilterTranslucentVolumeCS>(PermutationVector);
 
@@ -2588,13 +2443,13 @@ void FilterTranslucencyLightingVolume(
 
 					FComputeShaderUtils::AddPass(
 						GraphBuilder,
-						RDG_EVENT_NAME("Temporal Filter Cascade %d", VolumeCascadeIndex),
+						RDG_EVENT_NAME("Temporal Filter Cascade:%d MegaLights:%d", VolumeCascadeIndex, bMegaLightsTranslucencyVolume),
 						ComputeShader,
 						PassParameters,
 						NumGroups);
 				}
 
-				if (View.ViewState && !View.bStatePrevViewInfoIsReadOnly)
+				if (View.ViewState && !View.bStatePrevViewInfoIsReadOnly && bTemporalFilter)
 				{
 					FTranslucencyLightingViewState& TranslucencyLightingViewState = View.ViewState->TranslucencyLighting;
 
@@ -2703,8 +2558,11 @@ void FilterTranslucencyLightingVolume(
 				PassParameters->RWTranslucencyLightingVolumeDirectional = GraphBuilder.CreateUAV(OutputVolumeDirectionalTexture);
 
 				FFilterTranslucentVolumeCS::FPermutationDomain PermutationVector;
-				PermutationVector.Set<FFilterTranslucentVolumeCS::FUseTemporalReprojection>(false);
-				PermutationVector.Set<FFilterTranslucentVolumeCS::FCheckHistoryMark>(false);
+				PermutationVector.Set<FFilterTranslucentVolumeCS::FUseTemporalFilter>(false);
+				PermutationVector.Set<FFilterTranslucentVolumeCS::FUseSpatialFilter>(true);
+				PermutationVector.Set<FFilterTranslucentVolumeCS::FUseMarkTexture>(false);
+				PermutationVector.Set<FFilterTranslucentVolumeCS::FInjectMegaLights>(false);
+				PermutationVector = FFilterTranslucentVolumeCS::RemapPermutation(PermutationVector);
 
 				auto ComputeShader = View.ShaderMap->GetShader<FFilterTranslucentVolumeCS>(PermutationVector);
 
@@ -2712,7 +2570,7 @@ void FilterTranslucencyLightingVolume(
 
 				FComputeShaderUtils::AddPass(
 					GraphBuilder,
-					RDG_EVENT_NAME("Spatial Filter Cascade %d", VolumeCascadeIndex),
+					RDG_EVENT_NAME("Spatial Filter Cascade:%d", VolumeCascadeIndex),
 					ComputeShader,
 					PassParameters,
 					NumGroups);
@@ -2923,11 +2781,7 @@ void FDeferredShadingSceneRenderer::RenderTranslucencyLightingVolume(FRDGBuilder
 	}
 
 	InjectTranslucencyLightingVolumeAmbientCubemap(GraphBuilder, Views, Textures);
-	if (SortedLightSet.MegaLightsLightStart < SortedLightSet.SortedLights.Num())
-	{
-		InjectTranslucencyLightingVolumeMegaLights(GraphBuilder, Views, Textures);
-	}
-	FilterTranslucencyLightingVolume(GraphBuilder, Views, Textures);
+	FilterTranslucencyLightingVolume(GraphBuilder, Views, Textures, SortedLightSet);
 }
 
 BEGIN_SHADER_PARAMETER_STRUCT(FVisualizeTranslucencyVolumeCommonParameters, )
