Support coarse pages ignore dynamic invalidation mode (MarkCoarsePagesDirectional = 2) on directional lights.
Fix minor interaction where we should also ignore the "considered for WPO" case in these pages/mode.

#Synced-CL 48164149

diff --git a/Engine/Shaders/Private/VirtualShadowMaps/VirtualShadowMapPhysicalPageManagement.usf b/Engine/Shaders/Private/VirtualShadowMaps/VirtualShadowMapPhysicalPageManagement.usf
--- a/Engine/Shaders/Private/VirtualShadowMaps/VirtualShadowMapPhysicalPageManagement.usf
+++ b/Engine/Shaders/Private/VirtualShadowMaps/VirtualShadowMapPhysicalPageManagement.usf
@@ -390,7 +390,7 @@ void UpdatePhysicalPages(uint3 Index : SV_DispatchThreadID)
 					// This can get significantly cleaned up if/when we free up an additional hierarchical page flag bit.
 					// We do NOT want these details showing up in debug visualizations or cache stats though
 					bool bPageValidForRendering = (NextPageFlags & VSM_FLAG_ANY_UNCACHED) != 0;
-					if (PrevPhysicalFlags & VSM_EXTENDED_FLAG_FORCE_CACHED)
+					if (!bSkipDynamicInvalidations && PrevPhysicalFlags & VSM_EXTENDED_FLAG_FORCE_CACHED)
 					{
 						// NOTE: WPO can only ever be dynamic cached
 						NextPageFlags |= VSM_FLAG_DYNAMIC_UNCACHED;
diff --git a/Engine/Source/Runtime/Renderer/Private/VirtualShadowMaps/VirtualShadowMapClipmap.cpp b/Engine/Source/Runtime/Renderer/Private/VirtualShadowMaps/VirtualShadowMapClipmap.cpp
--- a/Engine/Source/Runtime/Renderer/Private/VirtualShadowMaps/VirtualShadowMapClipmap.cpp
+++ b/Engine/Source/Runtime/Renderer/Private/VirtualShadowMaps/VirtualShadowMapClipmap.cpp
@@ -129,8 +129,10 @@ static TAutoConsoleVariable<bool> CVarOrthoVSMRayCastViewOrigin(
 TAutoConsoleVariable<int32> CVarMarkCoarsePagesDirectional(
 	TEXT("r.Shadow.Virtual.MarkCoarsePagesDirectional"),
 	1,
-	TEXT("Marks coarse pages in directional light virtual shadow maps so that low resolution data is available everywhere.")
-	TEXT("Ability to disable is primarily for profiling and debugging."),
+	TEXT("Marks coarse pages in local light virtual shadow maps so that low resolution data is available everywhere.\n")
+	TEXT(" 0 - Disabled\n")
+	TEXT(" 1 - Always mark the last clipmap level and levels indicated by First/LastCoarseLevel.\n")
+	TEXT(" 2 - Performance Mode: Suppress dynamic invalidations due to geometry changes (e.g., moving objects or WPO or animation).\n"),
 	ECVF_Scalability | ECVF_RenderThreadSafe
 );
 
@@ -560,6 +562,13 @@ FVirtualShadowMapProjectionShaderData FVirtualShadowMapClipmap::ComputeProjectio
 		Data.Flags |= VSM_PROJ_FLAG_USE_RECEIVER_MASK;
 	}
 
+	// If the coarse pages mode is 2, we suppress invalidaitons of pages not marked as "detail geometry". 
+	// Coarse pages are not marked with this flag so will not be invalidated.
+	if (CVarMarkCoarsePagesDirectional.GetValueOnRenderThread() == 2)
+	{
+		Data.Flags |= VSM_PROJ_FLAG_FORCE_CACHE_DYNAMIC_COARSE;
+	}
+
 	Data.TexelDitherScale = GetLightSceneInfo().Proxy->GetVSMTexelDitherScale();
 	return Data;
 }
