[PostProcessing] Fix issue where GBuffer visualization textures were not rendered if postprocessing was disabled.

#Synced-CL 48051732

diff --git a/Engine/Source/Runtime/Renderer/Private/PostProcess/PostProcessing.cpp b/Engine/Source/Runtime/Renderer/Private/PostProcess/PostProcessing.cpp
--- a/Engine/Source/Runtime/Renderer/Private/PostProcess/PostProcessing.cpp
+++ b/Engine/Source/Runtime/Renderer/Private/PostProcess/PostProcessing.cpp
@@ -1637,6 +1637,12 @@ void AddPostProcessingPasses(
 		PassSequence.SetEnabled(EPass::VisualizeLocalExposure, false);
 		PassSequence.Finalize();
 
+		{
+			FPostProcessMaterialInputs PostProcessMaterialInputs = GetPostProcessMaterialInputs(FScreenPassTextureSlice::CreateFromScreenPassTexture(GraphBuilder, SceneColor));
+			AddVisualizeBufferMaterialPass(GraphBuilder, View, BL_SceneColorBeforeDOF, PostProcessMaterialInputs, bOutputInHDR, VisualizeBufferContext);
+			AddVisualizeBufferMaterialPass(GraphBuilder, View, BL_SceneColorAfterDOF, PostProcessMaterialInputs, bOutputInHDR, VisualizeBufferContext);
+		}
+
 		// Compose separate translucency passes
 		{
 			FTranslucencyComposition TranslucencyComposition;
@@ -1659,6 +1665,13 @@ void AddPostProcessingPasses(
 			}
 		}
 
+		{
+			FPostProcessMaterialInputs PostProcessMaterialInputs = GetPostProcessMaterialInputs(FScreenPassTextureSlice::CreateFromScreenPassTexture(GraphBuilder, SceneColor));
+			AddVisualizeBufferMaterialPass(GraphBuilder, View, BL_TranslucencyAfterDOF, PostProcessMaterialInputs, bOutputInHDR, VisualizeBufferContext);
+			AddVisualizeBufferMaterialPass(GraphBuilder, View, BL_SSRInput, PostProcessMaterialInputs, bOutputInHDR, VisualizeBufferContext);
+			AddVisualizeBufferMaterialPass(GraphBuilder, View, BL_SceneColorBeforeBloom, PostProcessMaterialInputs, bOutputInHDR, VisualizeBufferContext);
+		}
+
 		SceneColorBeforeTonemapSlice = FScreenPassTextureSlice::CreateFromScreenPassTexture(GraphBuilder, SceneColor);
 
 		if (PassSequence.IsEnabled(EPass::Tonemap))
@@ -1676,6 +1689,11 @@ void AddPostProcessingPasses(
 
 		SceneColor = AddAfterPass(EPass::Tonemap, SceneColor);
 
+		{
+			FPostProcessMaterialInputs PostProcessMaterialInputs = GetPostProcessMaterialInputs(FScreenPassTextureSlice::CreateFromScreenPassTexture(GraphBuilder, SceneColor));
+			AddVisualizeBufferMaterialPass(GraphBuilder, View, BL_SceneColorAfterTonemapping, PostProcessMaterialInputs, bOutputInHDR, VisualizeBufferContext);
+		}
+
 		SceneColorAfterTonemap = SceneColor;
 	}
 
