[PostProcessing] Force ActiveBufferVisualizationMaterial in VisualizeGBuffer to render to a separate target to prevent transition errors.

#rnx
#tests Tested in editor, and failing test is fixed in preflight
#Synced-CL 48231253

diff --git a/Engine/Source/Runtime/Renderer/Private/PostProcess/PostProcessVisualizeBuffer.cpp b/Engine/Source/Runtime/Renderer/Private/PostProcess/PostProcessVisualizeBuffer.cpp
--- a/Engine/Source/Runtime/Renderer/Private/PostProcess/PostProcessVisualizeBuffer.cpp
+++ b/Engine/Source/Runtime/Renderer/Private/PostProcess/PostProcessVisualizeBuffer.cpp
@@ -353,6 +353,7 @@ void AddVisualizeBufferMaterialPass(
 	{
 		if (VisMaterial->BlendableLocation == Location || Location == EBlendableLocation::BL_MAX)
 		{
+			VisualizeBufferPPMInputs.bAllowSceneColorInputAsOutput = false;
 			Context.Buffer = AddPostProcessMaterialPass(GraphBuilder, View, VisualizeBufferPPMInputs, VisMaterial);
 		}
 	}
