Fix broken caching for ACES 2 LUT resources

#Synced-CL 48202918

diff --git a/Engine/Source/Runtime/Renderer/Private/PostProcess/ACESUtils.cpp b/Engine/Source/Runtime/Renderer/Private/PostProcess/ACESUtils.cpp
--- a/Engine/Source/Runtime/Renderer/Private/PostProcess/ACESUtils.cpp
+++ b/Engine/Source/Runtime/Renderer/Private/PostProcess/ACESUtils.cpp
@@ -854,38 +854,47 @@ public:
 	}
 };
 
-void GetTransformResources(FRDGBuilder& GraphBuilder, float InPeakLuminance, const UE::Color::FColorSpace& LimitingColorSpace, FRHIShaderResourceView*& OutReachMTable, FRHIShaderResourceView*& OutGamutCuspTable, FRHIShaderResourceView*& OutUpperHullGammaTable)
+struct AcesContext
+{
+	FTextureLookupBase* ReachMTable = new TGlobalResource<FReachMTable>;
+	FTextureLookupBase* GamutCuspTable = new TGlobalResource<FGamutCuspTable>;
+	FTextureLookupBase* UpperHullGammaTable = new TGlobalResource<FUpperHullGammaTable>;
+	float CachedPeakLuminance = 0.0f;
+	UE::Color::FColorSpace CachedLimitingColorSpace;
+};
+
+static void GetTransformResourcesInternal(
+	FRDGBuilder& GraphBuilder,
+	float InPeakLuminance,
+	const UE::Color::FColorSpace& LimitingColorSpace,
+	AcesContext& Context,
+	FRHIShaderResourceView*& OutReachMTable,
+	FRHIShaderResourceView*& OutGamutCuspTable,
+	FRHIShaderResourceView*& OutUpperHullGammaTable)
 {
 	static const auto CVarAcesVersion = IConsoleManager::Get().FindTConsoleVariableDataInt(TEXT("r.HDR.Aces.Version"));
 	
 	if (CVarAcesVersion->GetValueOnRenderThread() > 1)
 	{
-		// ACES 2.0 transform resources
+		// ACES 2.0
 
-		static FTextureLookupBase* ReachMTable = new TGlobalResource<FReachMTable>;
-		static FTextureLookupBase* GamutCuspTable = new TGlobalResource<FGamutCuspTable>;
-		static FTextureLookupBase* UpperHullGammaTable = new TGlobalResource<FUpperHullGammaTable>;
-
-		static float CachedPeakLuminance = 0.0f;
-		static UE::Color::FColorSpace CachedLimitingColorSpace;
-
-		if (!FMath::IsNearlyEqual(CachedPeakLuminance, InPeakLuminance) ||
-			(CachedLimitingColorSpace != LimitingColorSpace))
+		if (!FMath::IsNearlyEqual(Context.CachedPeakLuminance, InPeakLuminance) ||
+			(Context.CachedLimitingColorSpace != LimitingColorSpace))
 		{
 			FRHICommandListImmediate& RHICmdListImmediate = GraphBuilder.RHICmdList.GetAsImmediate();
 
-			ReachMTable->Update(RHICmdListImmediate, LimitingColorSpace, InPeakLuminance);
-			GamutCuspTable->Update(RHICmdListImmediate, LimitingColorSpace, InPeakLuminance);
-			UpperHullGammaTable->Update(RHICmdListImmediate, LimitingColorSpace, InPeakLuminance);
+			Context.ReachMTable->Update(RHICmdListImmediate, LimitingColorSpace, InPeakLuminance);
+			Context.GamutCuspTable->Update(RHICmdListImmediate, LimitingColorSpace, InPeakLuminance);
+			Context.UpperHullGammaTable->Update(RHICmdListImmediate, LimitingColorSpace, InPeakLuminance);
 
 			// Update static cached info.
-			CachedPeakLuminance = InPeakLuminance;
-			CachedLimitingColorSpace = LimitingColorSpace;
+			Context.CachedPeakLuminance = InPeakLuminance;
+			Context.CachedLimitingColorSpace = LimitingColorSpace;
 		}
 
-		OutReachMTable = ReachMTable->ShaderResourceViewRHI;
-		OutGamutCuspTable = GamutCuspTable->ShaderResourceViewRHI;
-		OutUpperHullGammaTable = UpperHullGammaTable->ShaderResourceViewRHI;
+		OutReachMTable = Context.ReachMTable->ShaderResourceViewRHI;
+		OutGamutCuspTable = Context.GamutCuspTable->ShaderResourceViewRHI;
+		OutUpperHullGammaTable = Context.UpperHullGammaTable->ShaderResourceViewRHI;
 	}
 	else
 	{
@@ -897,5 +906,17 @@ void GetTransformResources(FRDGBuilder& GraphBuilder, float InPeakLuminance, con
 	}
 }
 
+void GetTransformResources(FRDGBuilder& GraphBuilder, float InPeakLuminance, const UE::Color::FColorSpace& LimitingColorSpace, FRHIShaderResourceView*& OutReachMTable, FRHIShaderResourceView*& OutGamutCuspTable, FRHIShaderResourceView*& OutUpperHullGammaTable)
+{
+	static AcesContext Context;
+	GetTransformResourcesInternal(GraphBuilder, InPeakLuminance, LimitingColorSpace, Context, OutReachMTable, OutGamutCuspTable, OutUpperHullGammaTable);
+}
+
+void GetTransformResourcesSDR(FRDGBuilder& GraphBuilder, FRHIShaderResourceView*& OutReachMTable, FRHIShaderResourceView*& OutGamutCuspTable, FRHIShaderResourceView*& OutUpperHullGammaTable)
+{
+	static AcesContext Context;
+	GetTransformResourcesInternal(GraphBuilder, 100.f, UE::Color::FColorSpace::GetWorking(), Context, OutReachMTable, OutGamutCuspTable, OutUpperHullGammaTable);
+}
+
 }  // namespace UE::Color::ACES
 
diff --git a/Engine/Source/Runtime/Renderer/Private/PostProcess/ACESUtils.h b/Engine/Source/Runtime/Renderer/Private/PostProcess/ACESUtils.h
--- a/Engine/Source/Runtime/Renderer/Private/PostProcess/ACESUtils.h
+++ b/Engine/Source/Runtime/Renderer/Private/PostProcess/ACESUtils.h
@@ -27,4 +27,21 @@ namespace UE::Color::ACES
 		FRHIShaderResourceView*& OutGamutCuspTable,
 		FRHIShaderResourceView*& OutUpperHullGammaTable
 	);
+
+	/*
+	* Get the ACES 2.0 transform table resources.
+	*
+	* @param GraphBuilder			Frame render dependency graph.
+	* @param InPeakLuminance		Peak luminance value in nits (cd/m^2).
+	* @param LimitingColorSpace		Color space of the target display.
+	* @param OutReachMTable			Reach M value table resource output.
+	* @param OutGamutCuspTable		Gamut cusp table resource output.
+	* @param OutUpperHullGammaTable	Upper hull gamma table resource output.
+	*/
+	void GetTransformResourcesSDR(
+		FRDGBuilder& GraphBuilder,
+		FRHIShaderResourceView*& OutReachMTable,
+		FRHIShaderResourceView*& OutGamutCuspTable,
+		FRHIShaderResourceView*& OutUpperHullGammaTable
+	);
 } // namespace UE::Color::ACES
diff --git a/Engine/Source/Runtime/Renderer/Private/PostProcess/PostProcessCombineLUTs.cpp b/Engine/Source/Runtime/Renderer/Private/PostProcess/PostProcessCombineLUTs.cpp
--- a/Engine/Source/Runtime/Renderer/Private/PostProcess/PostProcessCombineLUTs.cpp
+++ b/Engine/Source/Runtime/Renderer/Private/PostProcess/PostProcessCombineLUTs.cpp
@@ -332,10 +332,8 @@ struct FCachedLUTSettings
 			Parameters.ACESTonemapParameters.ACESGamutTable,
 			Parameters.ACESTonemapParameters.ACESGammaTable
 		);
-		UE::Color::ACES::GetTransformResources(
+		UE::Color::ACES::GetTransformResourcesSDR(
 			GraphBuilder,
-			100.,
-			UE::Color::FColorSpace::GetWorking(),
 			Parameters.ACESTonemapParameters.ACESReachTableSDR,
 			Parameters.ACESTonemapParameters.ACESGamutTableSDR,
 			Parameters.ACESTonemapParameters.ACESGammaTableSDR
