Remove most of r.Mobile.SupportsSM5MaterialNodes usage around SLW water rendering on mobile
Allow SceneColor copies if translucency needs it
Rely on RDG pass merging to merge lighting and translucency passes when possible
Fixed a few RHI validation errors around parameters binding for a lighting passes
#Synced-CL 48272171

diff --git a/Engine/Source/Runtime/Engine/Private/Materials/MaterialInterface.cpp b/Engine/Source/Runtime/Engine/Private/Materials/MaterialInterface.cpp
--- a/Engine/Source/Runtime/Engine/Private/Materials/MaterialInterface.cpp
+++ b/Engine/Source/Runtime/Engine/Private/Materials/MaterialInterface.cpp
@@ -700,7 +700,7 @@ FMaterialRelevance UMaterialInterface::GetRelevance_Internal(const UMaterial* Ma
 	MaterialRelevance.bHasVolumeMaterialDomain = MaterialResource->IsVolumetricPrimitive();
 	MaterialRelevance.bUsesDistanceCullFade = MaterialResource->MaterialUsesDistanceCullFade_GameThread();
 	MaterialRelevance.bUsesSkyMaterial = Material->bIsSky;
-	MaterialRelevance.bUsesSingleLayerWaterMaterial = bUsesSingleLayerWaterMaterial;
+	MaterialRelevance.bUsesSingleLayerWaterMaterial = (bUsesSingleLayerWaterMaterial && !bIsSinglePassWaterTranslucent);
 	MaterialRelevance.bUsesAnisotropy = bUsesAnisotropy;
 	MaterialRelevance.bUsesCustomizedUVs = MaterialResource->GetNumCustomizedUVs() > 0;
 	MaterialRelevance.bUsesVertexInterpolator = MaterialResource->HasVertexInterpolator();
diff --git a/Engine/Source/Runtime/RenderCore/Private/ShaderParameterStruct.cpp b/Engine/Source/Runtime/RenderCore/Private/ShaderParameterStruct.cpp
--- a/Engine/Source/Runtime/RenderCore/Private/ShaderParameterStruct.cpp
+++ b/Engine/Source/Runtime/RenderCore/Private/ShaderParameterStruct.cpp
@@ -594,14 +594,18 @@ void ValidateShaderParameterResourcesRHI(const void* Contents, const FRHIUniform
 
 		FRHIResource* Resource = GetShaderParameterResourceRHI(Contents, Parameter.MemberOffset, Parameter.MemberType);
 
-		const bool bSRV =
+		const bool bSRVorUAV =
 			Parameter.MemberType == UBMT_SRV ||
+			Parameter.MemberType == UBMT_UAV ||
 			Parameter.MemberType == UBMT_RDG_TEXTURE_SRV ||
+			Parameter.MemberType == UBMT_RDG_TEXTURE_UAV ||
 			Parameter.MemberType == UBMT_RDG_TEXTURE_NON_PIXEL_SRV ||
-			Parameter.MemberType == UBMT_RDG_BUFFER_SRV;
+			Parameter.MemberType == UBMT_RDG_BUFFER_SRV ||
+			Parameter.MemberType == UBMT_RDG_BUFFER_UAV;
 
 		// Allow null SRV's in uniform buffers for feature levels that don't support SRV's in shaders
-		if (GMaxRHIFeatureLevel <= ERHIFeatureLevel::ES3_1 && bSRV)
+		// Mobile does not set some UAVs to not break RDG pass auto-merging
+		if (GMaxRHIFeatureLevel <= ERHIFeatureLevel::ES3_1 && bSRVorUAV)
 		{
 			continue;
 		}
diff --git a/Engine/Source/Runtime/Renderer/Private/DepthCopy.h b/Engine/Source/Runtime/Renderer/Private/DepthCopy.h
--- a/Engine/Source/Runtime/Renderer/Private/DepthCopy.h
+++ b/Engine/Source/Runtime/Renderer/Private/DepthCopy.h
@@ -34,11 +34,6 @@ public:
 		SHADER_PARAMETER_RDG_TEXTURE(Texture2DMS, DepthTextureMS)
 		RENDER_TARGET_BINDING_SLOTS()
 		END_SHADER_PARAMETER_STRUCT()
-
-		static bool ShouldCompilePermutation(const FGlobalShaderPermutationParameters& Parameters)
-	{
-		return IsFeatureLevelSupported(Parameters.Platform, ERHIFeatureLevel::SM5) || MobileSupportsSM5MaterialNodes(Parameters.Platform);
-	}
 };
 
 // This is a temporary workaround while we get AddCopyTexturePass to do a proper copy of depth texture (with source texture HTile maintained).
diff --git a/Engine/Source/Runtime/Renderer/Private/MobileBasePassRendering.cpp b/Engine/Source/Runtime/Renderer/Private/MobileBasePassRendering.cpp
--- a/Engine/Source/Runtime/Renderer/Private/MobileBasePassRendering.cpp
+++ b/Engine/Source/Runtime/Renderer/Private/MobileBasePassRendering.cpp
@@ -329,14 +329,6 @@ void SetupMobileBasePassUniformParameters(
 	}
 	BasePassParameters.ScreenSpaceShadowMaskSampler = TStaticSamplerState<SF_Point, AM_Clamp, AM_Clamp, AM_Clamp>::GetRHI();
 
-	FRDGBuffer* OcclusionBuffer = View.ViewState ? View.ViewState->OcclusionFeedback.GetGPUFeedbackBuffer() : nullptr;
-	if (!OcclusionBuffer)
-	{
-		OcclusionBuffer = GraphBuilder.CreateBuffer(FRDGBufferDesc::CreateStructuredDesc(sizeof(uint32), 1), TEXT("OcclusionBufferFallback"));
-		
-	}
-	BasePassParameters.RWOcclusionBufferUAV = GraphBuilder.CreateUAV(OcclusionBuffer);
-
 	// Substrate
 	Substrate::BindSubstrateMobileForwardPasslUniformParameters(GraphBuilder, View, BasePassParameters.Substrate);
 
@@ -380,21 +372,34 @@ void SetupMobileBasePassUniformParameters(
 		SetupDebugViewModePassUniformBufferConstants(View, BasePassParameters.DebugViewMode);
 	}
 #endif
+	
+	// We need some kind of "Dummy" type resource here, to be fully correct.
+	// Otherwise UAVs break RDG pass auto-merging
+	if (BasePass <= EMobileBasePass::Opaque || GMaxRHIFeatureLevel > ERHIFeatureLevel::ES3_1)
+	{
+		// QuadOverdraw is a UAV so it needs to be initialized even if not used
+		FRDGTextureRef QuadOverdrawTexture = nullptr;
+		EDebugViewShaderMode DebugViewMode = View.Family->GetDebugViewShaderMode();
+		if ((DebugViewMode == EDebugViewShaderMode::DVSM_QuadComplexity || DebugViewMode == DVSM_ShaderComplexityContainedQuadOverhead) && SceneTextures)
+		{
+			QuadOverdrawTexture = SceneTextures->DebugAux;
+		}
 
-	// QuadOverdraw is a UAV so it needs to be initialized even if not used
-	FRDGTextureRef QuadOverdrawTexture = nullptr;
+		if (!QuadOverdrawTexture)
+		{
+			QuadOverdrawTexture = GraphBuilder.CreateTexture(FRDGTextureDesc::Create2D(FIntPoint(1, 1), PF_R32_UINT, FClearValueBinding::None, TexCreate_UAV), TEXT("DummyOverdrawUAV"));
+		}
+		BasePassParameters.QuadOverdraw = GraphBuilder.CreateUAV(QuadOverdrawTexture);
 
-	EDebugViewShaderMode DebugViewMode = View.Family->GetDebugViewShaderMode();
-	if ((DebugViewMode == EDebugViewShaderMode::DVSM_QuadComplexity || DebugViewMode == DVSM_ShaderComplexityContainedQuadOverhead) && SceneTextures)
-	{
-		QuadOverdrawTexture = SceneTextures->DebugAux;
-	}
+		// Occlussion Feedback
+		FRDGBuffer* OcclusionBuffer = View.ViewState ? View.ViewState->OcclusionFeedback.GetGPUFeedbackBuffer() : nullptr;
+		if (!OcclusionBuffer)
+		{
+			OcclusionBuffer = GraphBuilder.CreateBuffer(FRDGBufferDesc::CreateStructuredDesc(sizeof(uint32), 1), TEXT("OcclusionBufferFallback"));
 
-	if (!QuadOverdrawTexture)
-	{
-		QuadOverdrawTexture = GraphBuilder.CreateTexture(FRDGTextureDesc::Create2D(FIntPoint(1, 1), PF_R32_UINT, FClearValueBinding::None, TexCreate_UAV), TEXT("DummyOverdrawUAV"));
+		}
+		BasePassParameters.RWOcclusionBufferUAV = GraphBuilder.CreateUAV(OcclusionBuffer);
 	}
-	BasePassParameters.QuadOverdraw = GraphBuilder.CreateUAV(QuadOverdrawTexture);
 }
 
 TRDGUniformBufferRef<FMobileBasePassUniformParameters> CreateMobileBasePassUniformBuffer(
diff --git a/Engine/Source/Runtime/Renderer/Private/MobileDeferredShadingPass.cpp b/Engine/Source/Runtime/Renderer/Private/MobileDeferredShadingPass.cpp
--- a/Engine/Source/Runtime/Renderer/Private/MobileDeferredShadingPass.cpp
+++ b/Engine/Source/Runtime/Renderer/Private/MobileDeferredShadingPass.cpp
@@ -287,7 +287,6 @@ public:
 
 	void SetParameters(FRHIBatchedShaderParameters& BatchedParameters, const FViewInfo& View, const FMaterialRenderProxy* Proxy, const FMaterial& Material)
 	{
-		FMaterialShader::SetViewParameters(BatchedParameters, View, View.ViewUniformBuffer);
 		FMaterialShader::SetParameters(BatchedParameters, Proxy, Material, View);
 		
 		// LightFunctions can use primitive data, set identity so we do not crash on a missing binding
@@ -325,8 +324,8 @@ public:
 	>;
 
 	BEGIN_SHADER_PARAMETER_STRUCT(FParameters, )
-		SHADER_PARAMETER_STRUCT_REF(FViewUniformShaderParameters, View)
-		SHADER_PARAMETER_STRUCT_REF(FMobileReflectionCaptureShaderData, MobileReflectionCaptureData)
+		SHADER_PARAMETER_STRUCT_INCLUDE(FViewShaderParameters, View)
+		SHADER_PARAMETER_STRUCT_REF(FMobileReflectionCaptureShaderData, ReflectionCapture)
 		SHADER_PARAMETER_STRUCT_INCLUDE(FDFAOUpsampleParameters, DFAOUpsampleParameters)
 		SHADER_PARAMETER_STRUCT_INCLUDE(FSkyDiffuseLightingParameters, SkyDiffuseLighting)
 	END_SHADER_PARAMETER_STRUCT()
@@ -469,8 +468,8 @@ void RenderReflectionEnvironmentSkyLighting(
 	GraphicsPSOInit.RasterizerState = TStaticRasterizerState<>::GetRHI();
 
 	FMobileReflectionEnvironmentSkyLightingPS::FParameters PassParameters;
-	PassParameters.View = GetShaderBinding(View.ViewUniformBuffer);
-	PassParameters.MobileReflectionCaptureData = GetShaderBinding(View.MobileReflectionCaptureUniformBuffer);
+	PassParameters.View = View.GetShaderParameters();
+	PassParameters.ReflectionCapture = View.MobileReflectionCaptureUniformBuffer;
 
 	// DFAO
 	if (DynamicBentNormalAOTexture != nullptr)
diff --git a/Engine/Source/Runtime/Renderer/Private/MobileShadingRenderer.cpp b/Engine/Source/Runtime/Renderer/Private/MobileShadingRenderer.cpp
--- a/Engine/Source/Runtime/Renderer/Private/MobileShadingRenderer.cpp
+++ b/Engine/Source/Runtime/Renderer/Private/MobileShadingRenderer.cpp
@@ -2575,10 +2575,8 @@ void FMobileSceneRenderer::RenderDeferredMultiPass(FRDGBuilder& GraphBuilder, FS
 		RenderMobileScreenSpaceShadows(GraphBuilder, SceneTextures.Config.Extent, BasePassTextures.ScreenSpaceShadowMaskTexture);
 	}
 
-	const bool bSupportsSM5Nodes = MobileSupportsSM5MaterialNodes(ShaderPlatform);
-	const bool bShouldRenderFullSLW = bSupportsSM5Nodes && ShouldRenderSingleLayerWater(Views);
-	const bool bNeedsDepthCopy = bShouldRenderFullSLW;
-
+	const bool bShouldRenderSingleLayerWater = ShouldRenderSingleLayerWater(Views);
+	const bool bNeedsDepthCopy = bShouldRenderSingleLayerWater;
 	FRDGTextureRef SceneDepthCopyTexture = nullptr;
 	if (bNeedsDepthCopy)
 	{
@@ -2637,71 +2635,65 @@ void FMobileSceneRenderer::RenderDeferredMultiPass(FRDGBuilder& GraphBuilder, FS
 
 		const EMobileSSRQuality MobileSSRQuality = ActiveMobileSSRQuality(View, bShouldRenderVelocities);
 
-		bool bSeparateTranslucencyPass = bSupportsSM5Nodes;
-
 		GraphBuilder.AddPass(
-			bSeparateTranslucencyPass ? RDG_EVENT_NAME("LightingAndFog") : RDG_EVENT_NAME("LightingAndTranslucency"),
+			RDG_EVENT_NAME("LightingAndFog"),
 			PassParameters,
 			ERDGPassFlags::Raster,
-			[this, PassParameters, ViewContext, MobileSSRQuality, bSeparateTranslucencyPass, &SortedLightSet](FRHICommandList& RHICmdList)
+			[this, PassParameters, ViewContext, MobileSSRQuality, &SortedLightSet](FRHICommandList& RHICmdList)
 			{
 				FViewInfo& View = *ViewContext.ViewInfo;
-
 				MobileDeferredShadingPass(RHICmdList, ViewContext.ViewIndex, Views.Num(), View, *Scene, SortedLightSet, VisibleLightInfos, MobileSSRQuality, PassParameters->BentNormalAOTexture);
 				RenderFog(RHICmdList, View);
-				
-				if (!bSeparateTranslucencyPass)
-				{
-					RenderTranslucency(RHICmdList, View, Views, StandardTranslucencyPass, StandardTranslucencyMeshPass, &PassParameters->InstanceCullingDrawParams);
-				}
 			});
 
-		if (bSeparateTranslucencyPass)
+		if (bShouldRenderSingleLayerWater 
+			&& View.bHasSingleLayerWaterMaterial
+			&& HasAnyDraw(View.ParallelMeshDrawCommandPasses[EMeshPass::SingleLayerWaterPass]))
 		{
-			const bool bShouldRenderTranslucency = ShouldRenderTranslucency(StandardTranslucencyPass, Views) && ViewFamily.EngineShowFlags.Translucency && (View.ParallelMeshDrawCommandPasses[StandardTranslucencyMeshPass] != nullptr);
-			const bool bShouldRenderWater = bShouldRenderFullSLW && View.bHasSingleLayerWaterMaterial;
+			auto* SLWPassParameters = GraphBuilder.AllocParameters<FMobileRenderPassParameters>();
+			*SLWPassParameters = *PassParameters;
+			SLWPassParameters->RenderTargets[0] = BasePassRenderTargets[0];
+			SLWPassParameters->RenderTargets.DepthStencil = FDepthStencilBinding(BasePassRenderTargets.DepthStencil.GetTexture(), ERenderTargetLoadAction::ELoad, FExclusiveDepthStencil::DepthWrite_StencilWrite);
 
-			if (bShouldRenderWater)
+			FMobileBasePassUniformParameters* SLWBasePassParameters = GraphBuilder.AllocParameters<FMobileBasePassUniformParameters>();
+			SetupMobileBasePassUniformParameters(GraphBuilder, View, EMobileBasePass::Translucent, SetupMode, *SLWBasePassParameters, BasePassTextures);
+			
+			FRDGTextureRef SceneColorCopyTexture = AddCopySceneColorPass(GraphBuilder, Views, SceneTextures.Color, false, false);
+			if (SceneColorCopyTexture != nullptr)
 			{
-				auto* SLWPassParameters = GraphBuilder.AllocParameters<FMobileRenderPassParameters>();
-				*SLWPassParameters = *PassParameters;
-				SLWPassParameters->RenderTargets = BasePassRenderTargets;
-				SLWPassParameters->RenderTargets.DepthStencil = FDepthStencilBinding(BasePassRenderTargets.DepthStencil.GetTexture(), ERenderTargetLoadAction::ELoad, FExclusiveDepthStencil::DepthWrite_StencilNop);
-
-				FMobileBasePassUniformParameters* SLWBasePassParameters = GraphBuilder.AllocParameters<FMobileBasePassUniformParameters>();
-				SetupMobileBasePassUniformParameters(GraphBuilder, View, EMobileBasePass::Translucent, SetupMode, *SLWBasePassParameters, BasePassTextures);
-				SLWPassParameters->MobileBasePass = GraphBuilder.CreateUniformBuffer(SLWBasePassParameters);
+				SLWBasePassParameters->SceneColorCopyTexture = SceneColorCopyTexture;
+			}
+			if (SceneDepthCopyTexture != nullptr)
+			{
+				SLWBasePassParameters->SceneDepthCopyTexture = SceneDepthCopyTexture;
+			}
+			SLWPassParameters->MobileBasePass = GraphBuilder.CreateUniformBuffer(SLWBasePassParameters);
 
-				FRDGTextureRef SceneColorCopyTexture = AddCopySceneColorPass(GraphBuilder, Views, SceneTextures.Color, false, false);
-				if (SceneColorCopyTexture != nullptr)
-				{
-					SLWBasePassParameters->SceneColorCopyTexture = SceneColorCopyTexture;
-				}
-				if (SceneDepthCopyTexture != nullptr)
-				{
-					SLWBasePassParameters->SceneDepthCopyTexture = SceneDepthCopyTexture;
-				}
+			if (Scene->GPUScene.IsEnabled())
+			{
+				BuildMeshRenderingCommands(GraphBuilder, EMeshPass::SingleLayerWaterPass, View, Scene->GPUScene, InstanceCullingManager, SLWPassParameters->InstanceCullingDrawParams);
+			}
 
-				if (Scene->GPUScene.IsEnabled())
+			GraphBuilder.AddPass(
+				RDG_EVENT_NAME("SingleLayerWater"),
+				SLWPassParameters,
+				ERDGPassFlags::Raster,
+				[this, &View, SLWPassParameters](FRHICommandList& RHICmdList)
 				{
-					BuildMeshRenderingCommands(GraphBuilder, EMeshPass::SingleLayerWaterPass, View, Scene->GPUScene, InstanceCullingManager, SLWPassParameters->InstanceCullingDrawParams);
-				}
+					RHICmdList.SetViewport(View.ViewRect.Min.X, View.ViewRect.Min.Y, 0.0f, View.ViewRect.Max.X, View.ViewRect.Max.Y, 1.0f);
 
-				GraphBuilder.AddPass(
-					RDG_EVENT_NAME("SingleLayerWater"),
-					SLWPassParameters,
-					ERDGPassFlags::Raster,
-					[this, &View, SLWPassParameters](FRHICommandList& RHICmdList)
+					if (auto* Pass = View.ParallelMeshDrawCommandPasses[EMeshPass::SingleLayerWaterPass])
 					{
-						RHICmdList.SetViewport(View.ViewRect.Min.X, View.ViewRect.Min.Y, 0.0f, View.ViewRect.Max.X, View.ViewRect.Max.Y, 1.0f);
-
-						if (auto* Pass = View.ParallelMeshDrawCommandPasses[EMeshPass::SingleLayerWaterPass])
-						{
-							Pass->Draw(RHICmdList, &SLWPassParameters->InstanceCullingDrawParams);
-						}
-					});
-			}
-			if (bShouldRenderTranslucency)
+						Pass->Draw(RHICmdList, &SLWPassParameters->InstanceCullingDrawParams);
+					}
+				});
+		}
+			
+		if (ShouldRenderTranslucency(StandardTranslucencyPass, Views) 
+			&& ViewFamily.EngineShowFlags.Translucency
+			&& HasAnyDraw(View.ParallelMeshDrawCommandPasses[StandardTranslucencyMeshPass]))
+		{
+			if (View.TranslucentPrimCount.UseSceneColorCopy(StandardTranslucencyPass))
 			{
 				FRDGTextureRef SceneColorCopyTexture = AddCopySceneColorPass(GraphBuilder, Views, SceneTextures.Color, false, false);
 				if (SceneColorCopyTexture != nullptr)
@@ -2712,16 +2704,18 @@ void FMobileSceneRenderer::RenderDeferredMultiPass(FRDGBuilder& GraphBuilder, FS
 				{
 					BasePassParameters->SceneDepthCopyTexture = SceneDepthCopyTexture;
 				}
-
-				GraphBuilder.AddPass(
-					RDG_EVENT_NAME("Translucency"),
-					PassParameters,
-					ERDGPassFlags::Raster,
-					[this, &View, PassParameters](FRHICommandList& RHICmdList)
-					{
-						RenderTranslucency(RHICmdList, View, Views, StandardTranslucencyPass, StandardTranslucencyMeshPass, &PassParameters->InstanceCullingDrawParams);
-					});
+				
+				PassParameters->MobileBasePass = GraphBuilder.CreateUniformBuffer(BasePassParameters);
 			}
+			
+			GraphBuilder.AddPass(
+				RDG_EVENT_NAME("Translucency"),
+				PassParameters,
+				ERDGPassFlags::Raster,
+				[this, &View, PassParameters](FRHICommandList& RHICmdList)
+				{
+					RenderTranslucency(RHICmdList, View, Views, StandardTranslucencyPass, StandardTranslucencyMeshPass, &PassParameters->InstanceCullingDrawParams);
+				});
 		}
 	}
 }
diff --git a/Engine/Source/Runtime/Renderer/Private/TranslucentRendering.cpp b/Engine/Source/Runtime/Renderer/Private/TranslucentRendering.cpp
--- a/Engine/Source/Runtime/Renderer/Private/TranslucentRendering.cpp
+++ b/Engine/Source/Runtime/Renderer/Private/TranslucentRendering.cpp
@@ -309,11 +309,6 @@ public:
 		SHADER_PARAMETER_RDG_TEXTURE(Texture2D, SceneColorTexture)
 		RENDER_TARGET_BINDING_SLOTS()
 	END_SHADER_PARAMETER_STRUCT()
-
-	static bool ShouldCompilePermutation(const FGlobalShaderPermutationParameters& Parameters)
-	{
-		return IsFeatureLevelSupported(Parameters.Platform, ERHIFeatureLevel::SM5) || MobileSupportsSM5MaterialNodes(Parameters.Platform);
-	}
 };
 
 IMPLEMENT_GLOBAL_SHADER(FCopySceneColorPS, "/Engine/Private/TranslucentLightingShaders.usf", "CopySceneColorMain", SF_Pixel);
@@ -343,7 +338,7 @@ FRDGTextureRef AddCopySceneColorPass(FRDGBuilder& GraphBuilder, TArrayView<const
 				break;
 			}
 		}
-		bNeedsResolve = bNeedsResolve || (MobileSupportsSM5MaterialNodes(View.GetShaderPlatform()) && View.bHasSingleLayerWaterMaterial);
+		bNeedsResolve = bNeedsResolve || View.bHasSingleLayerWaterMaterial;
 
 		if (bNeedsResolve)
 		{
