Refactor VSM cache entry invalidation such that we can track previous VSM ID even if triggering invalidation from the CPU logic.

#Synced-CL 42429612

diff --git a/Engine/Source/Runtime/Renderer/Private/Shadows/ShadowSceneRenderer.cpp b/Engine/Source/Runtime/Renderer/Private/Shadows/ShadowSceneRenderer.cpp
--- a/Engine/Source/Runtime/Renderer/Private/Shadows/ShadowSceneRenderer.cpp
+++ b/Engine/Source/Runtime/Renderer/Private/Shadows/ShadowSceneRenderer.cpp
@@ -420,10 +420,9 @@ TSharedPtr<FVirtualShadowMapPerLightCacheEntry> FShadowSceneRenderer::AddLocalLi
 	}
 
 	PerLightCacheEntry->UpdateLocal(
-		ProjectedShadowInitializer,
+		FVirtualShadowMapPerLightCacheEntry::FLocalLightCacheKey { ProjectedShadowInitializer.WorldToLight, ProjectedShadowInitializer.PreShadowTranslation, bIsDistantLight },
 		LightSceneProxy->GetOrigin(),
 		LightSceneProxy->GetRadius(),
-		bIsDistantLight,
 		bForceInvalidate,
 		IsVirtualShadowMapLocalReceiverMaskEnabled());
 
@@ -662,7 +661,7 @@ void FShadowSceneRenderer::AllocateVirtualShadowMapIds()
 	for (const FLocalLightShadowFrameSetup& FrameSetup : LocalLights)
 	{
 		const TSharedPtr<FVirtualShadowMapPerLightCacheEntry>& CacheEntry = FrameSetup.PerLightCacheEntry;
-		const int32 NextVirtualShadowMapId = VirtualShadowMapArray.AllocateLocal(CacheEntry->bIsDistantLight, CacheEntry->ShadowMapEntries.Num());
+		const int32 NextVirtualShadowMapId = VirtualShadowMapArray.AllocateLocal(CacheEntry->IsDistantLight(), CacheEntry->ShadowMapEntries.Num());
 		CacheEntry->UpdateVirtualShadowMapId(NextVirtualShadowMapId);
 	}
 
@@ -821,7 +820,7 @@ void FShadowSceneRenderer::PostSetupDebugRender()
 			for (const FLocalLightShadowFrameSetup& LightSetup : LocalLights)
 			{			
 				FLinearColor Color = FLinearColor(FColor::Blue);
-				if (LightSetup.PerLightCacheEntry && LightSetup.PerLightCacheEntry->bIsDistantLight)
+				if (LightSetup.PerLightCacheEntry && LightSetup.PerLightCacheEntry->IsDistantLight())
 				{
 					++NumDistant;
 					int32 FramesSinceLastRender = int32(Scene.GetFrameNumber()) - int32(LightSetup.PerLightCacheEntry->GetLastScheduledFrameNumber());
diff --git a/Engine/Source/Runtime/Renderer/Private/VirtualShadowMaps/VirtualShadowMapCacheManager.cpp b/Engine/Source/Runtime/Renderer/Private/VirtualShadowMaps/VirtualShadowMapCacheManager.cpp
--- a/Engine/Source/Runtime/Renderer/Private/VirtualShadowMaps/VirtualShadowMapCacheManager.cpp
+++ b/Engine/Source/Runtime/Renderer/Private/VirtualShadowMaps/VirtualShadowMapCacheManager.cpp
@@ -300,22 +300,14 @@ void FVirtualShadowMapCacheEntry::UpdateHZBMetadata(const FViewMatrices& ViewMat
 		CurrentHZBMetadata.ViewMatrices = ViewMatrices;
 	}
 }
-
-void FVirtualShadowMapPerLightCacheEntry::UpdateClipmap(
-	const FVector& LightDirection,
-	int FirstLevel,
-	bool bForceInvalidate,
-	bool bInUseReceiverMask)
+template <typename CacheKeyType>
+void FVirtualShadowMapPerLightCacheEntry::Update(const CacheKeyType& InKey, bool bForceInvalidate, bool bInUseReceiverMask, CacheKeyType& OutCacheKey)
 {
-	if (bForceInvalidate ||
-		LightDirection != ClipmapCacheKey.LightDirection ||
-		FirstLevel != ClipmapCacheKey.FirstLevel)
+	if (bForceInvalidate || InKey != OutCacheKey)
 	{
 		RenderedFrameNumber = -1;
-		ClipmapCacheKey.LightDirection = LightDirection;
-		ClipmapCacheKey.FirstLevel = FirstLevel;
 	}
-
+	OutCacheKey = InKey;
 	// If we swap receiver mask on/off we also may need to invalidate due to incomplete pages
 	if (bInUseReceiverMask != bUseReceiverMask)
 	{
@@ -338,72 +330,34 @@ void FVirtualShadowMapPerLightCacheEntry::UpdateClipmap(
 		bIsUncached = bNewIsUncached;
 	}
 
+	//// On invalidate, wipe out any existing VSM IDs
+	//if (RenderedFrameNumber < 0)
+	//{
+	//	PrevVirtualShadowMapId = INDEX_NONE;
+	//	VirtualShadowMapId = INDEX_NONE;
+	//}
+}
+
+void FVirtualShadowMapPerLightCacheEntry::UpdateClipmap(const FClipmapCacheKey& InKey, bool bForceInvalidate, bool bInUseReceiverMask)
+{
+	Update(InKey, bForceInvalidate, bInUseReceiverMask, ClipmapCacheKey);
+
 	// Not used by clipmaps
 	LightOrigin = FVector(0, 0, 0);
 	LightRadius = -1.0f;
-
-	// On invalidate, wipe out any existing VSM IDs
-	if (RenderedFrameNumber < 0)
-	{
-		PrevVirtualShadowMapId = INDEX_NONE;
-		VirtualShadowMapId = INDEX_NONE;
-	}
 }
 
 void FVirtualShadowMapPerLightCacheEntry::UpdateLocal(
-	const FProjectedShadowInitializer& InCacheKey,
+	const FLocalLightCacheKey& InKey,
 	const FVector& NewLightOrigin,
 	const float NewLightRadius,
-	bool bNewIsDistantLight,
 	bool bForceInvalidate,
 	bool bInUseReceiverMask)
 {
-	// TODO: The logic in this function is needlessly convoluted... clean up
-	const bool bLightMoved = 
-		LocalCacheKey.PreShadowTranslation != InCacheKey.PreShadowTranslation ||
-		LocalCacheKey.WorldToLight != InCacheKey.WorldToLight;
-
-	LocalCacheKey.PreShadowTranslation = InCacheKey.PreShadowTranslation;
-	LocalCacheKey.WorldToLight = InCacheKey.WorldToLight;
-
-	// Check cache validity based of shadow setup
-	// If it is a distant light, we want to let the time-share perform the invalidation.
-	if (bForceInvalidate || bLightMoved)
-	{
-		RenderedFrameNumber = -1;
-	}
-
-	// If we swap receiver mask on/off we also may need to invalidate due to incomplete pages
-	if (bInUseReceiverMask != bUseReceiverMask)
-	{
-		RenderedFrameNumber = -1;
-		bUseReceiverMask = bInUseReceiverMask;
-	}
-
-	// On transition between uncached <-> cached we must invalidate since the static pages may not be initialized
-	bool bNewIsUncached = RenderedFrameNumber < 0;
-	if (bNewIsUncached != bIsUncached)
-	{
-		RenderedFrameNumber = -1;
-		bIsUncached = bNewIsUncached;
-	}
-
-	// On transition between distant <-> regular we must invalidate
-	if (bNewIsDistantLight != bIsDistantLight)
-	{
-		RenderedFrameNumber = -1;
-		bIsDistantLight = bNewIsDistantLight;
-	}
+	Update(InKey, bForceInvalidate, bInUseReceiverMask, LocalCacheKey);
 
 	LightOrigin = NewLightOrigin;
 	LightRadius = NewLightRadius;
-
-	// On invalidate, wipe out any existing VSM IDs
-	if (RenderedFrameNumber < 0)
-	{
-		PrevVirtualShadowMapId = INDEX_NONE;
-		VirtualShadowMapId = INDEX_NONE;
-	}
 }
 
 void FVirtualShadowMapPerLightCacheEntry::UpdateVirtualShadowMapId(int32 NextVirtualShadowMapId)
@@ -463,7 +417,7 @@ FVirtualShadowMapArrayCacheManager::FInvalidatingPrimitiveCollector::FInvalidati
 				continue;
 			}
 			int32 VirtualShadowMapId = CacheEntry.GetVirtualShadowMapId();
-			if (VirtualShadowMapId < 0)
+			if (!ensure(VirtualShadowMapId >= 0))
 			{
 				continue;
 			}
@@ -598,6 +552,12 @@ void FVirtualShadowMapArrayCacheManager::FInvalidatingPrimitiveCollector::AddInv
 	bool bLightRadiusCulling,
 	const FBoxSphereBounds& PrimitiveBounds)
 {	
+	// 
+	if (CacheEntry.IsFullyCached() || CacheEntry.IsUncached() || CacheEntry.IsInvalidated())
+	{
+		return;
+	}
+
 	// We don't need explicit dynamic invalidations when using receiver mask
 	if (bCachePrimitiveAsDynamic && CacheEntry.ShouldUseReceiverMask())
 	{
@@ -609,7 +569,7 @@ void FVirtualShadowMapArrayCacheManager::FInvalidatingPrimitiveCollector::AddInv
 		return;
 	}
 	int32 VirtualShadowMapId = CacheEntry.GetVirtualShadowMapId();
-	if (VirtualShadowMapId < 0)
+	if (!ensure(VirtualShadowMapId >= 0))
 	{
 		return;
 	}
@@ -1389,7 +1349,7 @@ void FVirtualShadowMapArrayCacheManager::UpdateUnreferencedCacheEntries(
 			// Not active this render, but still recent enough to keep it and its pages alive
 			// Reallocate new VSM ID for this frame
 			int32 NumMaps = CacheEntry->ShadowMapEntries.Num();
-			const int32 NextVirtualShadowMapId = VirtualShadowMapArray.AllocateUnreferenced(CacheEntry->bIsDistantLight, NumMaps);
+			const int32 NextVirtualShadowMapId = VirtualShadowMapArray.AllocateUnreferenced(CacheEntry->IsDistantLight(), NumMaps);
 			CacheEntry->UpdateVirtualShadowMapId(NextVirtualShadowMapId);
 			for (int32 Map = 0; Map < NumMaps; ++Map)
 			{
@@ -1411,11 +1371,15 @@ void FVirtualShadowMapArrayCacheManager::UpdateUnreferencedCacheEntries(
 			const int32 PrevVirtualShadowMapId = CacheEntry->GetPrevVirtualShadowMapId();
 			if (PrevVirtualShadowMapId >= 0)
 			{
-				for (int32 Index = 0; Index < CacheEntry->ShadowMapEntries.Num(); ++Index)
+				const int32 BaseId = CacheEntry->GetVirtualShadowMapId();
+				if (ensure(BaseId >= 0))
 				{
-					auto& NextData = CacheEntry->ShadowMapEntries[Index].NextData;
-					NextData.NextVirtualShadowMapId = CacheEntry->GetVirtualShadowMapId() + Index;
-					VirtualShadowMapArray.UpdateNextData(PrevVirtualShadowMapId + Index, NextData);
+					for (int32 Index = 0; Index < CacheEntry->ShadowMapEntries.Num(); ++Index)
+					{
+						auto& NextData = CacheEntry->ShadowMapEntries[Index].NextData;
+						NextData.NextVirtualShadowMapId = BaseId + Index;
+						VirtualShadowMapArray.UpdateNextData(PrevVirtualShadowMapId + Index, NextData);
+					}
 				}
 			}
 		}
diff --git a/Engine/Source/Runtime/Renderer/Private/VirtualShadowMaps/VirtualShadowMapCacheManager.h b/Engine/Source/Runtime/Renderer/Private/VirtualShadowMaps/VirtualShadowMapCacheManager.h
--- a/Engine/Source/Runtime/Renderer/Private/VirtualShadowMaps/VirtualShadowMapCacheManager.h
+++ b/Engine/Source/Runtime/Renderer/Private/VirtualShadowMaps/VirtualShadowMapCacheManager.h
@@ -104,9 +104,14 @@ public:
 	 * The (local) VSM is fully cached if it is distant and has been rendered to previously
 	 * "Fully" implies that we know all pages are mapped as well as rendered to (ignoring potential CPU-side object culling).
 	 */
-	inline bool IsFullyCached() const
+	bool IsFullyCached() const
 	{
-		return bIsDistantLight && RenderedFrameNumber >= 0;
+		return LocalCacheKey.bIsDistantLight && RenderedFrameNumber >= 0;
+	}
+
+	bool IsDistantLight() const
+	{
+		return LocalCacheKey.bIsDistantLight;
 	}
 
 	inline bool IsUncached() const { return bIsUncached; }
@@ -124,17 +129,36 @@ public:
 	int32 GetVirtualShadowMapId() const { return VirtualShadowMapId; }
 	int32 GetPrevVirtualShadowMapId() const { return PrevVirtualShadowMapId; }
 
-	void UpdateClipmap(
-		const FVector& LightDirection,
-		int FirstLevel,
-		bool bForceInvalidate,
-		bool bInUseReceiverMask);
+	/**
+	 * Contains attributes which trigger invalidation if any change is detected.
+	 */
+	struct FClipmapCacheKey
+	{
+		bool operator==(const FClipmapCacheKey&) const = default;
+		bool operator!=(const FClipmapCacheKey&) const = default;
 
+		FVector LightDirection;
+		int FirstLevel = 0;
+	};
+
+	void UpdateClipmap(const FClipmapCacheKey& InKey, bool bForceInvalidate, bool bInUseReceiverMask);
+
+	struct FLocalLightCacheKey
+	{
+		bool operator==(const FLocalLightCacheKey&) const = default;
+		bool operator!=(const FLocalLightCacheKey&) const = default;
+
+		FMatrix WorldToLight;
+		FVector PreShadowTranslation;
+		bool bIsDistantLight = false;
+	};
+
+	/**
+	 */
 	void UpdateLocal(
-		const FProjectedShadowInitializer &InCacheKey,
+		const FLocalLightCacheKey& InCacheKey,
 		const FVector& NewLightOrigin,
 		const float NewLightRadius,
-		bool bNewIsDistantLight,
 		bool bForceInvalidate,
 		bool bInUseReceiverMask);
 
@@ -152,7 +176,6 @@ public:
 	int32 VirtualShadowMapId = INDEX_NONE;
 
 	bool bIsUncached = false;
-	bool bIsDistantLight = false;
 	bool bUseReceiverMask = false;
 
 	// Tracks if this cache entry is being used "this render", i.e. "active". Note that there may be multiple renders per frame in the case of
@@ -173,19 +196,10 @@ public:
 	TArray<FVirtualShadowMapInstanceRange> PrimitiveInstancesToInvalidate;
 
 private:
-	struct FLocalLightCacheKey
-	{
-		FMatrix WorldToLight;
-		FVector PreShadowTranslation;
-	};
-	FLocalLightCacheKey LocalCacheKey;
+	template <typename CacheKeyType>
+	void Update(const CacheKeyType& InKey, bool bForceInvalidate, bool bInUseReceiverMask, CacheKeyType& OutCacheKey);
 
-	struct FClipmapCacheKey
-	{
-		FVector LightDirection;
-		int FirstLevel;
-		int LevelCount;
-	};
+	FLocalLightCacheKey LocalCacheKey;
 	FClipmapCacheKey ClipmapCacheKey;
 
 	// Rough bounds for invalidation culling
diff --git a/Engine/Source/Runtime/Renderer/Private/VirtualShadowMaps/VirtualShadowMapClipmap.cpp b/Engine/Source/Runtime/Renderer/Private/VirtualShadowMaps/VirtualShadowMapClipmap.cpp
--- a/Engine/Source/Runtime/Renderer/Private/VirtualShadowMaps/VirtualShadowMapClipmap.cpp
+++ b/Engine/Source/Runtime/Renderer/Private/VirtualShadowMaps/VirtualShadowMapClipmap.cpp
@@ -318,7 +318,7 @@ FVirtualShadowMapClipmap::FVirtualShadowMapClipmap(
 	// This is likely to all change as we refactor the multiple cache managers stuff anyways; for now this is hopefully safe since they do have separate copies
 	const uint32 UniqueViewKey = InDependentView->ViewState ? InDependentView->ViewState->GetViewKey() : 0U;
 	PerLightCacheEntry = VirtualShadowMapArrayCacheManager->FindCreateLightCacheEntry(LightSceneInfo.Id, UniqueViewKey, LevelCount, Config.ShadowTypeId);
-	PerLightCacheEntry->UpdateClipmap(LightDirection, FirstLevel, bUncached, Config.bUseReceiverMask);
+	PerLightCacheEntry->UpdateClipmap( { LightDirection, FirstLevel }, bUncached, Config.bUseReceiverMask);
 
 	const int RadiusLn = GetLevelRadius(static_cast<float>(LastLevel));
 	const int RadiiPerLevel = 4;
