Fixed: FCopySceneColorPS shader compilation for platforms that do not support shared samplers
#Synced-CL 48272171

diff --git a/Engine/Shaders/Private/TranslucentLightingShaders.usf b/Engine/Shaders/Private/TranslucentLightingShaders.usf
--- a/Engine/Shaders/Private/TranslucentLightingShaders.usf
+++ b/Engine/Shaders/Private/TranslucentLightingShaders.usf
@@ -50,10 +50,15 @@ MaterialFloat3 EncodeSceneColorForMaterialNode(MaterialFloat3 LinearSceneColor)
 }
 
 Texture2D SceneColorTexture;
+#if SUPPORTS_INDEPENDENT_SAMPLERS
+	#define SceneColorTextureSampler GlobalPointClampedSampler
+#else
+	SamplerState SceneColorTextureSampler;
+#endif
 
 void CopySceneColorMain(FScreenVertexOutput Input, out float4 OutColor : SV_Target0)
 {
-	float4 LinearColor = Texture2DSample(SceneColorTexture, GlobalPointClampedSampler, Input.UV);
+	float4 LinearColor = Texture2DSample(SceneColorTexture, SceneColorTextureSampler, Input.UV);
 	OutColor = float4(EncodeSceneColorForMaterialNode(LinearColor.rgb), LinearColor.a);
 }
 
diff --git a/Engine/Source/Runtime/Renderer/Private/TranslucentRendering.cpp b/Engine/Source/Runtime/Renderer/Private/TranslucentRendering.cpp
--- a/Engine/Source/Runtime/Renderer/Private/TranslucentRendering.cpp
+++ b/Engine/Source/Runtime/Renderer/Private/TranslucentRendering.cpp
@@ -307,6 +307,7 @@ public:
 	BEGIN_SHADER_PARAMETER_STRUCT(FParameters, )
 		SHADER_PARAMETER_STRUCT_REF(FViewUniformShaderParameters, View)
 		SHADER_PARAMETER_RDG_TEXTURE(Texture2D, SceneColorTexture)
+		SHADER_PARAMETER_SAMPLER(SamplerState, SceneColorTextureSampler)
 		RENDER_TARGET_BINDING_SLOTS()
 	END_SHADER_PARAMETER_STRUCT()
 };
@@ -362,6 +363,7 @@ FRDGTextureRef AddCopySceneColorPass(FRDGBuilder& GraphBuilder, TArrayView<const
 			auto* PassParameters = GraphBuilder.AllocParameters<FCopySceneColorPS::FParameters>();
 			PassParameters->View = View.ViewUniformBuffer;
 			PassParameters->SceneColorTexture = SceneColor.Resolve;
+			PassParameters->SceneColorTextureSampler = TStaticSamplerState<SF_Point, AM_Clamp, AM_Clamp, AM_Clamp>::GetRHI();
 			PassParameters->RenderTargets[0] = FRenderTargetBinding(SceneColorCopyTexture, LoadAction);
 
 			if (!View.Family->bMultiGPUForkAndJoin)
@@ -416,6 +418,7 @@ static void AddCopySceneColorAlphaPass(FRDGBuilder& GraphBuilder, TArrayView<con
 		auto* PassParameters = GraphBuilder.AllocParameters<FCopySceneColorPS::FParameters>();
 		PassParameters->View = View.ViewUniformBuffer;
 		PassParameters->SceneColorTexture = SourceTexture;
+		PassParameters->SceneColorTextureSampler = TStaticSamplerState<SF_Point, AM_Clamp, AM_Clamp, AM_Clamp>::GetRHI();
 		PassParameters->RenderTargets[0] = FRenderTargetBinding(TargetTexture, LoadAction);
 
 		if (!View.Family->bMultiGPUForkAndJoin)
