Fix for Substrate UB nullptr check firing when running the SLW prepass

#Synced-CL 48348317

diff --git a/Engine/Source/Runtime/Renderer/Private/SingleLayerWaterRendering.cpp b/Engine/Source/Runtime/Renderer/Private/SingleLayerWaterRendering.cpp
--- a/Engine/Source/Runtime/Renderer/Private/SingleLayerWaterRendering.cpp
+++ b/Engine/Source/Runtime/Renderer/Private/SingleLayerWaterRendering.cpp
@@ -835,8 +835,10 @@ static FSingleLayerWaterTileClassification ClassifyTiles(FRDGBuilder& GraphBuild
 
 		// Mark used tiles based on SHADING_MODEL_ID
 		{
+			bool bUsePrepassStencil = DepthPrepassTexture != nullptr;
+
 			FWaterTileCategorisationMarkCS::FPermutationDomain PermutationVector;
-			PermutationVector.Set<FWaterTileCategorisationMarkCS::FUsePrepassStencil>(DepthPrepassTexture != nullptr);
+			PermutationVector.Set<FWaterTileCategorisationMarkCS::FUsePrepassStencil>(bUsePrepassStencil);
 			PermutationVector.Set<FWaterTileCategorisationMarkCS::FBuildFroxels>(bProduceFroxelData);
 
 			TShaderMapRef<FWaterTileCategorisationMarkCS> ComputeShader(View.ShaderMap, PermutationVector);
@@ -845,7 +847,7 @@ static FSingleLayerWaterTileClassification ClassifyTiles(FRDGBuilder& GraphBuild
 
 			PassParameters->SceneTextures = GetSceneTextureParameters(GraphBuilder, SceneTextures);
 			PassParameters->View = View.GetShaderParameters();
-			PassParameters->Substrate = Substrate::BindSubstrateGlobalUniformParameters(View);
+			PassParameters->Substrate = !bUsePrepassStencil ? Substrate::BindSubstrateGlobalUniformParameters(View) : nullptr;
 			PassParameters->TiledViewRes = Result.TiledViewRes;
 			PassParameters->WaterDepthStencilTexture = DepthPrepassTexture ? GraphBuilder.CreateSRV(FRDGTextureSRVDesc::CreateWithPixelFormat(DepthPrepassTexture, PF_X24_G8)) : nullptr;
 			PassParameters->WaterDepthTexture = DepthPrepassTexture ? DepthPrepassTexture : nullptr;
