[Lit Translucency] Fix blending against black volume when history is invalidated.
- also release TLV history when r.TranslucencyLightingVolume is toggled off.
- and moved call to DrawDebugTranslucencyLightingVolume(...) out of FilterTranslucencyLightingVolume(...).

#Synced-CL 48292508

diff --git a/Engine/Source/Runtime/Renderer/Private/TranslucentLighting.cpp b/Engine/Source/Runtime/Renderer/Private/TranslucentLighting.cpp
--- a/Engine/Source/Runtime/Renderer/Private/TranslucentLighting.cpp
+++ b/Engine/Source/Runtime/Renderer/Private/TranslucentLighting.cpp
@@ -351,6 +351,23 @@ void FViewInfo::CalcTranslucencyLightingVolumeBounds(FBox* InOutCascadeBoundsArr
 	}
 }
 
+void ReleaseTranslucencyLightingVolumeHistory(TArrayView<FViewInfo> Views)
+{
+	for (FViewInfo& View : Views)
+	{
+		if (View.ViewState && !View.bStatePrevViewInfoIsReadOnly)
+		{
+			for (int32 VolumeCascadeIndex = 0; VolumeCascadeIndex < TVC_MAX; VolumeCascadeIndex++)
+			{
+				FTranslucencyLightingViewState& TranslucencyLightingViewState = View.ViewState->TranslucencyLighting;
+				TranslucencyLightingViewState.HistoryAmbient[VolumeCascadeIndex] = nullptr;
+				TranslucencyLightingViewState.HistoryDirectional[VolumeCascadeIndex] = nullptr;
+				TranslucencyLightingViewState.HistoryMark[VolumeCascadeIndex] = nullptr;
+			}
+		}
+	}
+}
+
 class FTranslucencyDepthShaderElementData : public FMeshMaterialShaderElementData
 {
 public:
@@ -2235,13 +2252,10 @@ static void DrawDebugTranslucencyLightingVolume(
 
 void FilterTranslucencyLightingVolume(
 	FRDGBuilder& GraphBuilder,
-	const TArrayView<const FViewInfo> Views,
+	TArrayView<FViewInfo> Views,
 	FTranslucencyLightingVolumeTextures& Textures,
 	const FSortedLightSetSceneInfo& SortedLightSet)
 {
-	// Draw debug information if needed.
-	DrawDebugTranslucencyLightingVolume(GraphBuilder, Views);
-
 	const bool bMegaLightsTranslucencyVolume = Views.Num() > 0
 		&& MegaLights::IsEnabled(*Views[0].Family)
 		&& MegaLights::UseTranslucencyVolume()
@@ -2330,12 +2344,12 @@ void FilterTranslucencyLightingVolume(
 
 				bool bHasValidHistory = false;
 
-				FRDGTextureRef HistoryAmbient;
-				FRDGTextureRef HistoryDirectional;
-				FRDGTextureRef HistoryMark;
+				FRDGTextureRef HistoryAmbient = nullptr;
+				FRDGTextureRef HistoryDirectional = nullptr;
+				FRDGTextureRef HistoryMark = nullptr;
 
-				FVector4f PrevTranslucencyLightingVolumeMin;
-				FVector4f PrevTranslucencyLightingVolumeInvSize;
+				FVector4f PrevTranslucencyLightingVolumeMin = FVector4f::Zero();
+				FVector4f PrevTranslucencyLightingVolumeInvSize = FVector4f::Zero();
 
 				if (View.ViewState && !View.bCameraCut && !View.bPrevTransformsReset && bTemporalFilter)
 				{
@@ -2346,21 +2360,11 @@ void FilterTranslucencyLightingVolume(
 						HistoryAmbient = GraphBuilder.RegisterExternalTexture(TranslucencyLightingViewState.HistoryAmbient[VolumeCascadeIndex]);
 						HistoryDirectional = GraphBuilder.RegisterExternalTexture(TranslucencyLightingViewState.HistoryDirectional[VolumeCascadeIndex]);
 					}
-					else
-					{
-						const FRDGSystemTextures& SystemTextures = FRDGSystemTextures::Get(GraphBuilder);
-						HistoryAmbient = SystemTextures.VolumetricBlack;
-						HistoryDirectional = SystemTextures.VolumetricBlack;
-					}
 
 					if (TranslucencyLightingViewState.HistoryMark[VolumeCascadeIndex])
 					{
 						HistoryMark = GraphBuilder.RegisterExternalTexture(TranslucencyLightingViewState.HistoryMark[VolumeCascadeIndex]);
 					}
-					else
-					{
-						HistoryMark = nullptr;
-					}
 
 					bHasValidHistory = HistoryAmbient && HistoryDirectional && (!bUseVolumeMarkTexture || HistoryMark);
 
@@ -2375,16 +2379,6 @@ void FilterTranslucencyLightingVolume(
 						PrevTranslucencyLightingVolumeInvSize = FVector4f(FVector3f(1.0f) / VolumeSize, VolumeVoxelSize);
 					}
 				}
-				else
-				{
-					const FRDGSystemTextures& SystemTextures = FRDGSystemTextures::Get(GraphBuilder);
-					HistoryAmbient = SystemTextures.VolumetricBlack;
-					HistoryDirectional = SystemTextures.VolumetricBlack;
-					HistoryMark = nullptr;
-
-					PrevTranslucencyLightingVolumeMin = FVector4f::Zero();
-					PrevTranslucencyLightingVolumeInvSize = FVector4f::Zero();
-				}
 
 				if (bHasValidHistory || bMegaLightsTranslucencyVolume)
 				{
@@ -2471,22 +2465,7 @@ void FilterTranslucencyLightingVolume(
 	else
 	{
 		// release history resources if temporal is not enabled
-
-		for (int32 ViewIndex = 0; ViewIndex < Views.Num(); ++ViewIndex)
-		{
-			const FViewInfo& View = Views[ViewIndex];
-
-			if (View.ViewState && !View.bStatePrevViewInfoIsReadOnly)
-			{
-				for (int32 VolumeCascadeIndex = 0; VolumeCascadeIndex < TVC_MAX; VolumeCascadeIndex++)
-				{
-					FTranslucencyLightingViewState& TranslucencyLightingViewState = View.ViewState->TranslucencyLighting;
-					TranslucencyLightingViewState.HistoryAmbient[VolumeCascadeIndex] = nullptr;
-					TranslucencyLightingViewState.HistoryDirectional[VolumeCascadeIndex] = nullptr;
-					TranslucencyLightingViewState.HistoryMark[VolumeCascadeIndex] = nullptr;
-				}
-			}
-		}
+		ReleaseTranslucencyLightingVolumeHistory(Views);
 	}
 
 	// update HistoryVolumeMin/Size regardless of temporal filter being enabled or not so r.TranslucencyLightingVolume.Debug.FreezePosition always works
@@ -2709,6 +2688,7 @@ void FDeferredShadingSceneRenderer::RenderTranslucencyLightingVolume(FRDGBuilder
 {
 	if (!GUseTranslucencyLightingVolumes || !ViewFamily.EngineShowFlags.DirectLighting || !RHISupportsVolumeTextureRendering(ViewFamily.GetShaderPlatform()))
 	{
+		ReleaseTranslucencyLightingVolumeHistory(Views);
 		return;
 	}
 
@@ -2782,6 +2762,7 @@ void FDeferredShadingSceneRenderer::RenderTranslucencyLightingVolume(FRDGBuilder
 
 	InjectTranslucencyLightingVolumeAmbientCubemap(GraphBuilder, Views, Textures);
 	FilterTranslucencyLightingVolume(GraphBuilder, Views, Textures, SortedLightSet);
+	DrawDebugTranslucencyLightingVolume(GraphBuilder, Views);
 }
 
 BEGIN_SHADER_PARAMETER_STRUCT(FVisualizeTranslucencyVolumeCommonParameters, )
