[HWRT] Optimize FinishGatherVisibleShaderBindings(...) to minimize reallocations.
- in test level with 10K dynamic instances: 0.7ms -> 0.45ms

#Synced-CL 48396324

diff --git a/Engine/Source/Runtime/Renderer/Private/RayTracing/RayTracing.cpp b/Engine/Source/Runtime/Renderer/Private/RayTracing/RayTracing.cpp
--- a/Engine/Source/Runtime/Renderer/Private/RayTracing/RayTracing.cpp
+++ b/Engine/Source/Runtime/Renderer/Private/RayTracing/RayTracing.cpp
@@ -2228,7 +2228,21 @@ namespace RayTracing
 
 		TaskData.VisibleRayTracingShaderBindingsFinalizeTask.Wait();
 
-		// merge dynamic bindings
+		// calculate total number of visible shader bindings
+		int32 NumVisibleShaderBindings = TaskData.VisibleShaderBindings.Num();
+
+		for (FGatherInstancesViewTaskData& ViewTaskData : TaskData.ViewTaskDatas)
+		{
+			for (int32 TaskIndex = 0; TaskIndex < ViewTaskData.View.DynamicRayTracingShaderBindingsPerTask.Num(); TaskIndex++)
+			{
+				NumVisibleShaderBindings += ViewTaskData.View.DynamicRayTracingShaderBindingsPerTask[TaskIndex]->Num();
+			}
+		}
+
+		// reserve enough space to avoid multiple realloc
+		TaskData.VisibleShaderBindings.Reserve(NumVisibleShaderBindings);
+
+		// actually merge dynamic bindings into the main array
 		for (FGatherInstancesViewTaskData& ViewTaskData : TaskData.ViewTaskDatas)
 		{
 			for (int32 TaskIndex = 0; TaskIndex < ViewTaskData.View.DynamicRayTracingShaderBindingsPerTask.Num(); TaskIndex++)
@@ -2239,7 +2253,6 @@ namespace RayTracing
 			ViewTaskData.View.DynamicRayTracingShaderBindingsPerTask.Empty();
 		}
 
-
 		// Even though task dependencies are setup so all work is done by this point, we still have to wait on the
 		// pipe to clear out its internal state. Otherwise it can assert that it still has work at shutdown.
 		TaskData.AddInstancesPipe.WaitUntilEmpty();
