[Ray Tracing] Fix a bug in the build instance buffer pass where incorrect GPU scene resources are used due to getting them too late. We call GPUScene->GetShaderParametersRHI() inside an RDG pass instead of getting those parameters at pass addition time and capture them as a lambda argument

#tests tested in editor and on a console

diff --git a/Engine/Plugins/Experimental/GPULightmass/Source/GPULightmass/Private/LightmapRenderer.cpp b/Engine/Plugins/Experimental/GPULightmass/Source/GPULightmass/Private/LightmapRenderer.cpp
--- a/Engine/Plugins/Experimental/GPULightmass/Source/GPULightmass/Private/LightmapRenderer.cpp
+++ b/Engine/Plugins/Experimental/GPULightmass/Source/GPULightmass/Private/LightmapRenderer.cpp
@@ -1212,6 +1212,7 @@ bool FSceneRenderState::SetupRayTracingScene(FRDGBuilder& GraphBuilder, FSceneUn
 				GraphBuilder.RHICmdList,
 				nullptr,
 				nullptr,
+				nullptr,
 				InstanceBuffer.UAV,
 				GRHIGlobals.RayTracing.RequiresSeparateHitGroupContributionsBuffer ? HitGroupContributionsBuffer.UAV : nullptr,
 				SceneInitializer.MaxNumInstances,
diff --git a/Engine/Source/Runtime/Renderer/Private/RayTracing/RayTracingInstanceBufferUtil.cpp b/Engine/Source/Runtime/Renderer/Private/RayTracing/RayTracingInstanceBufferUtil.cpp
--- a/Engine/Source/Runtime/Renderer/Private/RayTracing/RayTracingInstanceBufferUtil.cpp
+++ b/Engine/Source/Runtime/Renderer/Private/RayTracing/RayTracingInstanceBufferUtil.cpp
@@ -633,7 +633,7 @@ IMPLEMENT_GLOBAL_SHADER(FRayTracingBuildInstanceBufferCS, "/Engine/Private/RayTr
 
 void BuildRayTracingInstanceBuffer(
 	FRHICommandList& RHICmdList,
-	const FGPUScene* GPUScene,
+	const FGPUSceneResourceParametersRHI* GPUSceneParameters,
 	const FDFVector3& PreViewTranslation,
 	uint32 MaxNumInstances,
 	uint32 NumGroups,
@@ -667,9 +667,9 @@ void BuildRayTracingInstanceBuffer(
 	PassParams.PreViewTranslationHigh = PreViewTranslation.High;
 	PassParams.PreViewTranslationLow = PreViewTranslation.Low;
 
-	if (GPUScene)
+	if (GPUSceneParameters)
 	{
-		PassParams.GPUSceneParameters = GPUScene->GetShaderParametersRHI();
+		PassParams.GPUSceneParameters = *GPUSceneParameters;
 	}
 
 	if (CullingParameters)
@@ -715,7 +715,7 @@ void BuildRayTracingInstanceBuffer(
 
 void BuildRayTracingInstanceBuffer(
 	FRHICommandList& RHICmdList,
-	const FGPUScene* GPUScene,
+	const FGPUSceneResourceParametersRHI* GPUSceneParameters,
 	const FDFVector3& PreViewTranslation,
 	FRHIUnorderedAccessView* InstancesUAV,
 	FRHIUnorderedAccessView* HitGroupContributionsUAV,
@@ -738,7 +738,7 @@ void BuildRayTracingInstanceBuffer(
 	{
 		BuildRayTracingInstanceBuffer(
 			RHICmdList,
-			GPUScene,
+			GPUSceneParameters,
 			PreViewTranslation,
 			MaxNumInstances,
 			NumGPUGroups,
@@ -762,7 +762,7 @@ void BuildRayTracingInstanceBuffer(
 	{
 		BuildRayTracingInstanceBuffer(
 			RHICmdList,
-			GPUScene,
+			GPUSceneParameters,
 			PreViewTranslation,
 			MaxNumInstances,
 			NumCPUGroups,
@@ -946,10 +946,51 @@ void FRayTracingInstanceBufferBuilder::BuildRayTracingInstanceBuffer(
 	FRHIUnorderedAccessView* OutputStatsUAV,
 	uint32 OutputStatsOffset,
 	FRHIUnorderedAccessView* InstanceExtraDataUAV)
+{
+	FGPUSceneResourceParametersRHI GPUSceneParameters;
+	if (GPUScene)
+	{
+		GPUSceneParameters = GPUScene->GetShaderParametersRHI();
+	}
+
+	::BuildRayTracingInstanceBuffer(
+		RHICmdList,
+		GPUScene ? &GPUSceneParameters : nullptr,
+		FDFVector3(PreViewTranslation),
+		InstancesUAV,
+		HitGroupContributionsUAV,
+		InstanceGroupUploadSRV,
+		InstanceUploadSRV,
+		AccelerationStructureAddressesBuffer.SRV,
+		TransformUploadSRV,
+		MaxNumInstances,
+		Data.NumGPUInstanceGroups,
+		Data.NumCPUInstanceGroups,
+		Data.NumGPUInstanceDescriptors,
+		Data.NumCPUInstanceDescriptors,
+		CullingParameters,
+		bCompactOutput,
+		OutputStatsUAV,
+		OutputStatsOffset,
+		InstanceExtraDataUAV);
+}
+
+void FRayTracingInstanceBufferBuilder::BuildRayTracingInstanceBuffer(
+	FRHICommandList& RHICmdList,
+	const FGPUScene* GPUScene,
+	const FGPUSceneResourceParametersRHI* GPUSceneParameters,
+	const FRayTracingCullingParameters* CullingParameters,
+	FRHIUnorderedAccessView* InstancesUAV,
+	FRHIUnorderedAccessView* HitGroupContributionsUAV,
+	uint32 MaxNumInstances,
+	bool bCompactOutput,
+	FRHIUnorderedAccessView* OutputStatsUAV,
+	uint32 OutputStatsOffset,
+	FRHIUnorderedAccessView* InstanceExtraDataUAV)
 {
 	::BuildRayTracingInstanceBuffer(
 		RHICmdList,
-		GPUScene,
+		GPUSceneParameters,
 		FDFVector3(PreViewTranslation),
 		InstancesUAV,
 		HitGroupContributionsUAV,
diff --git a/Engine/Source/Runtime/Renderer/Private/RayTracing/RayTracingScene.cpp b/Engine/Source/Runtime/Renderer/Private/RayTracing/RayTracingScene.cpp
--- a/Engine/Source/Runtime/Renderer/Private/RayTracing/RayTracingScene.cpp
+++ b/Engine/Source/Runtime/Renderer/Private/RayTracing/RayTracingScene.cpp
@@ -87,6 +87,7 @@ BEGIN_SHADER_PARAMETER_STRUCT(FBuildInstanceBufferPassParams, )
 	SHADER_PARAMETER_RDG_BUFFER_UAV(RWStructuredBuffer, OutputStats)
 	SHADER_PARAMETER_RDG_BUFFER_UAV(RWStructuredBuffer, InstanceExtraDataBuffer)
 	SHADER_PARAMETER_RDG_UNIFORM_BUFFER(FSceneUniformParameters, Scene)
+	SHADER_PARAMETER_STRUCT_INCLUDE(FGPUSceneResourceParameters, GPUSceneParameters)
 END_SHADER_PARAMETER_STRUCT()
 
 const FRayTracingScene::FInstanceHandle FRayTracingScene::INVALID_INSTANCE_HANDLE = FInstanceHandle();
@@ -401,6 +402,7 @@ void FRayTracingScene::Update(FRDGBuilder& GraphBuilder, FSceneUniformBuffer& Sc
 					
 					PassParams->InstanceExtraDataBuffer = InstanceExtraDataBufferUAV;
 					PassParams->Scene = SceneUniformBuffer.GetBuffer(GraphBuilder);
+					PassParams->GPUSceneParameters = GPUScene->GetShaderParameters(GraphBuilder);
 					PassParams->OutputStats = bCompactInstanceBuffer || bStatsEnabled ? InstanceStatsBufferUAV : nullptr;
 
 					const uint32 OutputStatsOffset = LayerIndex * MaxNumViews + ViewIndex;
@@ -418,9 +420,18 @@ void FRayTracingScene::Update(FRDGBuilder& GraphBuilder, FSceneUniformBuffer& Sc
 						bCompactInstanceBuffer
 						](FRHICommandList& RHICmdList)
 						{
+							FGPUSceneResourceParametersRHI GPUSceneParameters;
+							GPUSceneParameters.CommonParameters = PassParams->GPUSceneParameters.CommonParameters;
+							GPUSceneParameters.GPUSceneInstanceSceneData = PassParams->GPUSceneParameters.GPUSceneInstanceSceneData->GetRHI();
+							GPUSceneParameters.GPUSceneInstancePayloadData = PassParams->GPUSceneParameters.GPUSceneInstancePayloadData->GetRHI();
+							GPUSceneParameters.GPUScenePrimitiveSceneData = PassParams->GPUSceneParameters.GPUScenePrimitiveSceneData->GetRHI();
+							GPUSceneParameters.GPUSceneLightmapData = PassParams->GPUSceneParameters.GPUSceneLightmapData->GetRHI();
+							GPUSceneParameters.GPUSceneLightData = PassParams->GPUSceneParameters.GPUSceneLightData->GetRHI();
+
 							InstanceBufferBuilder.BuildRayTracingInstanceBuffer(
 								RHICmdList,
 								GPUScene,
+								&GPUSceneParameters,
 								CullingParameters,
 								PassParams->InstanceBuffer->GetRHI(),
 								GRHIGlobals.RayTracing.RequiresSeparateHitGroupContributionsBuffer ? PassParams->HitGroupContributionsBuffer->GetRHI() : nullptr,
diff --git a/Engine/Source/Runtime/Renderer/Private/Tests/RayTracingTestbed.cpp b/Engine/Source/Runtime/Renderer/Private/Tests/RayTracingTestbed.cpp
--- a/Engine/Source/Runtime/Renderer/Private/Tests/RayTracingTestbed.cpp
+++ b/Engine/Source/Runtime/Renderer/Private/Tests/RayTracingTestbed.cpp
@@ -219,6 +219,7 @@ bool RunRayTracingTestbed_RenderThread(const FString& Parameters)
 		RHICmdList,
 		nullptr,
 		nullptr,
+		nullptr,
 		InstanceBuffer.UAV,
 		GRHIGlobals.RayTracing.RequiresSeparateHitGroupContributionsBuffer ? HitGroupContributionsBuffer.UAV : nullptr,
 		SceneInitializer.MaxNumInstances,
diff --git a/Engine/Source/Runtime/Renderer/Public/RayTracingInstanceBufferUtil.h b/Engine/Source/Runtime/Renderer/Public/RayTracingInstanceBufferUtil.h
--- a/Engine/Source/Runtime/Renderer/Public/RayTracingInstanceBufferUtil.h
+++ b/Engine/Source/Runtime/Renderer/Public/RayTracingInstanceBufferUtil.h
@@ -9,6 +9,7 @@
 #if RHI_RAYTRACING
 
 class FGPUScene;
+class FGPUSceneResourceParametersRHI;
 struct FRayTracingCullingParameters;
 struct FDFVector3;
 
@@ -47,6 +48,19 @@ public:
 		uint32 OutputStatsOffset,
 		FRHIUnorderedAccessView* InstanceExtraDataUAV);
 
+	RENDERER_API void BuildRayTracingInstanceBuffer(
+		FRHICommandList& RHICmdList,
+		const FGPUScene* GPUScene,
+		const FGPUSceneResourceParametersRHI* GPUSceneParameters,
+		const FRayTracingCullingParameters* CullingParameters,
+		FRHIUnorderedAccessView* InstancesUAV,
+		FRHIUnorderedAccessView* HitGroupContributionsUAV,
+		uint32 MaxNumInstances,
+		bool bCompactOutput,
+		FRHIUnorderedAccessView* OutputStatsUAV,
+		uint32 OutputStatsOffset,
+		FRHIUnorderedAccessView* InstanceExtraDataUAV);
+
 	uint32 GetMaxNumInstances() const { return Data.NumNativeGPUSceneInstances + Data.NumNativeCPUInstances; }
 
 	TConstArrayView<FRHIRayTracingGeometry*> GetReferencedGeometries() const { return Data.ReferencedGeometries; }
