[Ray Tracing] Deprecate the old FRayTracingInstanceBufferBuilder::BuildRayTracingInstanceBuffer API

#tests compile editor
#Synced-CL 48332239

diff --git a/Engine/Source/Runtime/Renderer/Public/RayTracingInstanceBufferUtil.h b/Engine/Source/Runtime/Renderer/Public/RayTracingInstanceBufferUtil.h
--- a/Engine/Source/Runtime/Renderer/Public/RayTracingInstanceBufferUtil.h
+++ b/Engine/Source/Runtime/Renderer/Public/RayTracingInstanceBufferUtil.h
@@ -36,6 +36,7 @@ public:
 	RENDERER_API void FillRayTracingInstanceUploadBuffer(FRHICommandList& RHICmdList);
 	RENDERER_API void FillAccelerationStructureAddressesBuffer(FRHICommandList& RHICmdList);
 
+	UE_DEPRECATED(5.8, "Use the version that takes FGPUSceneResourceParametersRHI* instead.")
 	RENDERER_API void BuildRayTracingInstanceBuffer(
 		FRHICommandList& RHICmdList,
 		const FGPUScene* GPUScene,
