[DOF] Fix issue where the front-end passes were not ViewRectOffset aware after the TAA-prepass, leading to visual artifacts if the taa-prepass was disabled.

#Synced-CL 48380560

diff --git a/Engine/Shaders/Private/DiaphragmDOF/DOFCocTileFlatten.usf b/Engine/Shaders/Private/DiaphragmDOF/DOFCocTileFlatten.usf
--- a/Engine/Shaders/Private/DiaphragmDOF/DOFCocTileFlatten.usf
+++ b/Engine/Shaders/Private/DiaphragmDOF/DOFCocTileFlatten.usf
@@ -24,9 +24,9 @@
 
 //------------------------------------------------------- PARAMETERS.
 
-uint4 ViewportRect;
-float2 ThreadIdToBufferUV;
-float2 MaxBufferUV;
+float2 Stride;
+float4 GatherInputExtentAndInvExtent;
+float4 GatherInputViewRectUV;
 
 Texture2D GatherInput_SceneColor;
 Texture2D GatherInput_SeparateCoc;
@@ -132,11 +132,11 @@ void CocFlattenMainCS(
 
 	#if DIM_DO_COC_GATHER4
 	{
-		float2 BufferUV = (DispatchThreadId + 0.5) * ThreadIdToBufferUV;
-
-		if (true)
+		float2 BufferUV;
 		{
-			BufferUV = min(BufferUV, MaxBufferUV);
+			float2 PixelPosition = (DispatchThreadId * Stride) + 0.5;
+			BufferUV = GatherInputViewRectUV.xy + (PixelPosition * GatherInputExtentAndInvExtent.zw);
+			BufferUV = clamp(BufferUV, GatherInputViewRectUV.xy, GatherInputViewRectUV.zw);
 		}
 		
 		#if CONFIG_DOF_ALPHA
@@ -158,13 +158,19 @@ void CocFlattenMainCS(
 	}
 	#else // !DIM_DO_COC_GATHER4
 	{
-		uint2 PixelPos = min(DispatchThreadId, ViewportRect.zw - 1);
+		uint2 BufferPixelPosition;
+		{
+			float2 PixelPosition = DispatchThreadId + 0.5;
+			float2 BufferUV = GatherInputViewRectUV.xy + (PixelPosition * GatherInputExtentAndInvExtent.zw);
+			BufferUV = clamp(BufferUV, GatherInputViewRectUV.xy, GatherInputViewRectUV.zw);
+			BufferPixelPosition = BufferUV * GatherInputExtentAndInvExtent.xy;
+		}
 
 		#if CONFIG_DOF_ALPHA
 			// The coc radius is within separate buffer's red channel.
-			float CocRadius = GatherInput_SeparateCoc[PixelPos].r;
+			float CocRadius = GatherInput_SeparateCoc[BufferPixelPosition].r;
 		#else
-			float CocRadius = GatherInput_SceneColor[PixelPos].a;
+			float CocRadius = GatherInput_SceneColor[BufferPixelPosition].a;
 		#endif
 
 		InitCocTileSampleFromCocRadius(CocRadius, TileSample);
diff --git a/Engine/Shaders/Private/DiaphragmDOF/DOFDownsample.usf b/Engine/Shaders/Private/DiaphragmDOF/DOFDownsample.usf
--- a/Engine/Shaders/Private/DiaphragmDOF/DOFDownsample.usf
+++ b/Engine/Shaders/Private/DiaphragmDOF/DOFDownsample.usf
@@ -10,10 +10,10 @@
 //------------------------------------------------------- PARAMETERS
 
 uint4 ViewportRect;
-float2 MaxBufferUV;
 float OutputCocRadiusMultiplier;
 
-float4 GatherInputSize;
+float4 GatherInputExtentAndInvExtent;
+float4 GatherInputViewRectUV;
 Texture2D GatherInput_SceneColor;
 Texture2D GatherInput_SeparateCoc;
 
@@ -36,7 +36,8 @@ void DownsampleCS(uint2 DispatchThreadId : SV_DispatchThreadID)
 
 	// Compute scene buffer uv, but carefully to handle view size % 2 == 1,
 	// because this pass is running at half resolution.
-	float2 BufferUV = (2.0 * (DispatchThreadId + 0.5)) * GatherInputSize.zw;
+	float2 PixelPosition = 2.0 * DispatchThreadId;
+	float2 BufferUV = GatherInputViewRectUV.xy + ((PixelPosition + 0.5) * GatherInputExtentAndInvExtent.zw);
 	
 	// Sample scene color.
 	float4 SceneColors[4];
@@ -44,9 +45,9 @@ void DownsampleCS(uint2 DispatchThreadId : SV_DispatchThreadID)
 
 	for (uint i = 0; i < 4; i++)
 	{
-		float2 SampleUV = BufferUV + GatherInputSize.zw * (0.5 * float2(kOffsetsCross3x3[i]));
+		float2 SampleUV = BufferUV + GatherInputExtentAndInvExtent.zw * (0.5 * float2(kOffsetsCross3x3[i]));
 
-		SampleUV = min(SampleUV, MaxBufferUV);
+		SampleUV = clamp(SampleUV, GatherInputViewRectUV.xy, GatherInputViewRectUV.zw);
 
 		float4 RawSample0 = GatherInput_SceneColor.SampleLevel(GlobalPointClampedSampler, SampleUV, 0);
 
diff --git a/Engine/Shaders/Private/DiaphragmDOF/DOFHybridScatterCompilation.ush b/Engine/Shaders/Private/DiaphragmDOF/DOFHybridScatterCompilation.ush
--- a/Engine/Shaders/Private/DiaphragmDOF/DOFHybridScatterCompilation.ush
+++ b/Engine/Shaders/Private/DiaphragmDOF/DOFHybridScatterCompilation.ush
@@ -19,8 +19,8 @@ struct FHybridScatterInputs
 
 	// Texture used for performing neighborhood comparison.
 	Texture2D NeighborhoodComparisonBuffer;
-	float2 NeighborhoodComparisonBufferInvSize;
-	float2 NeighborhoodComparisonMaxBufferUV;
+	float2 NeighborhoodComparisonBufferInvExtent;
+	float4 NeighborhoodComparisonViewRectUV;
 	
 
 	// ------- run time pixel independent category.
@@ -68,17 +68,8 @@ float ComputeNeighborBasedScatterFactor(in const FHybridScatterInputs Parameters
 				continue;
 			}
 
-			float2 SampleUV = Parameters.BufferUV + Offset * Parameters.NeighborhoodComparisonBufferInvSize;
-
-			// Parameters.BufferUV is assumed to be already clamped.
-			if (Offset.x > 0)
-			{
-				SampleUV.x = min(SampleUV.x, Parameters.NeighborhoodComparisonMaxBufferUV.x);
-			}
-			if (Offset.y > 0)
-			{
-				SampleUV.y = min(SampleUV.y, Parameters.NeighborhoodComparisonMaxBufferUV.y);
-			}
+			float2 SampleUV = Parameters.BufferUV + Offset * Parameters.NeighborhoodComparisonBufferInvExtent;
+			SampleUV = clamp(SampleUV, Parameters.NeighborhoodComparisonViewRectUV.xy, Parameters.NeighborhoodComparisonViewRectUV.zw);
 
 			float4 SampleColor = Parameters.NeighborhoodComparisonBuffer.SampleLevel(GlobalPointClampedSampler, SampleUV, 0);
 			float SampleCoc = SampleColor.a;
diff --git a/Engine/Shaders/Private/DiaphragmDOF/DOFRecombine.usf b/Engine/Shaders/Private/DiaphragmDOF/DOFRecombine.usf
--- a/Engine/Shaders/Private/DiaphragmDOF/DOFRecombine.usf
+++ b/Engine/Shaders/Private/DiaphragmDOF/DOFRecombine.usf
@@ -107,9 +107,6 @@
 // Clamp full res gathering to have spec hits still temporally stable.
 #define CONFIG_CLAMP_FULLRES_GATHER 1
 
-// Clamp buffer UVs
-#define CONFIG_CLAMP_SCENE_BUFFER_UV 1
-
 // Clamp buffer UVs
 #define CONFIG_CLAMP_DOF_BUFFER_UV 1
 
@@ -143,6 +140,9 @@ float EncodedCocRadiusToRecombineCocRadius;
 float MaxRecombineAbsCocRadius;
 float SOFHistoryRejectionBoost;
 
+float4 SceneColorExtentAndInvExtent;
+float4 SceneColorViewRectUV;
+
 Texture2D BokehLUT;
 
 Texture2D SceneColorInput;
@@ -566,7 +566,7 @@ void FetchAndAccumulateSamplePair(
 	float SampleDistance = FullResPixelDistanceToCocDistance(length(PixelOffset));
 
 	// Scene buffer offset.
-	float2 SceneBufferUVOffset = PixelOffset * View.BufferSizeAndInvSize.zw * float2(CocInvSqueeze, 1.0);
+	float2 SceneBufferUVOffset = PixelOffset * SceneColorExtentAndInvExtent.zw * float2(CocInvSqueeze, 1.0);
 	
 	// Two sample to gather at same time.
 	FGatherSample Sample[2];
@@ -589,11 +589,7 @@ void FetchAndAccumulateSamplePair(
 		#endif
 
 		float2 BufferUV = InputParameters.SceneBufferUV + SampleSign * SceneBufferUVOffset;
-
-		if (true) // TODO.
-		{
-			BufferUV = clamp(BufferUV, View.BufferBilinearUVMinMax.xy, View.BufferBilinearUVMinMax.zw);
-		}
+		BufferUV = clamp(BufferUV, SceneColorViewRectUV.xy, SceneColorViewRectUV.zw);
 		
 		// Fetch full res color and CocRadius.
 		#if CONFIG_FETCH_FULLRES_COC_FROM_ALPHA
@@ -703,12 +699,9 @@ void RecombineMainCS(
 		InputParameters.DOFBufferSize = ConvolutionInputSize;
 
 		InputParameters.ViewportUV = (DispatchThreadId + 0.5) * ViewportSize.zw;
-		InputParameters.SceneBufferUV = ViewportUVToBufferUV(InputParameters.ViewportUV);
-		
-		if (CONFIG_CLAMP_SCENE_BUFFER_UV)
-		{
-			InputParameters.SceneBufferUV = clamp(InputParameters.SceneBufferUV, View.BufferBilinearUVMinMax.xy, View.BufferBilinearUVMinMax.zw);
-		}
+
+		InputParameters.SceneBufferUV = SceneColorViewRectUV.xy + InputParameters.ViewportUV * (SceneColorViewRectUV.zw - SceneColorViewRectUV.xy);
+		InputParameters.SceneBufferUV = clamp(InputParameters.SceneBufferUV, SceneColorViewRectUV.xy, SceneColorViewRectUV.zw);
 
 		// - 0.5 * TemporalJitterPixels because DOF buffer is non temporally jittering, thanks to half res TAA pass.
 		InputParameters.DOFBufferUV = ApplyScreenTransform(float2(DispatchThreadId), DispatchThreadIdToDOFBufferUV);
@@ -786,13 +779,9 @@ void RecombineMainCS(
 			{
 				const float2 SamplePixelOffset = float2(kOffsetsCross3x3[j]) * CocDistanceToFullResPixelDistance(MaxRecombineAbsCocRadius);
 
-				float2 SampleUVOffset = View.BufferSizeAndInvSize.zw * SamplePixelOffset;
+				float2 SampleUVOffset = SceneColorExtentAndInvExtent.zw * SamplePixelOffset;
 				float2 SampleUV = InputParameters.SceneBufferUV + SampleUVOffset;
-			
-				if (CONFIG_CLAMP_SCENE_BUFFER_UV)
-				{
-					SampleUV = clamp(SampleUV, View.BufferBilinearUVMinMax.xy, View.BufferBilinearUVMinMax.zw);
-				}
+				SampleUV = clamp(SampleUV, SceneColorViewRectUV.xy, SceneColorViewRectUV.zw);
 
 				#if CONFIG_FETCH_FULLRES_COC_FROM_ALPHA
 					float SampleCocRadius = SceneColorInput.SampleLevel(GlobalPointClampedSampler, SampleUV, 0).a * EncodedCocRadiusToRecombineCocRadius;
@@ -1182,11 +1171,7 @@ void RecombineMainCS(
 	}
 	#endif
 	
-	#if 1 // Lower VGPR footprint.
-		uint2 OutputPixelPosition = InputParameters.SceneBufferUV * View.BufferSizeAndInvSize.xy;
-	#else
-		uint2 OutputPixelPosition = ViewportRect.xy + DispatchThreadId;
-	#endif
+	uint2 OutputPixelPosition = ViewportRect.xy + DispatchThreadId;
 
 	#if CONFIG_DOF_ALPHA
 		SceneColorOutput[OutputPixelPosition] = OutputFinalSceneColor;
diff --git a/Engine/Shaders/Private/DiaphragmDOF/DOFReduce.usf b/Engine/Shaders/Private/DiaphragmDOF/DOFReduce.usf
--- a/Engine/Shaders/Private/DiaphragmDOF/DOFReduce.usf
+++ b/Engine/Shaders/Private/DiaphragmDOF/DOFReduce.usf
@@ -61,11 +61,10 @@
 
 //------------------------------------------------------- PARAMETERS
 
-uint4 ViewportRect;
+uint2 PassViewSize;
 float4 ScatteringViewportSize;
 float ScatteringScaling;
 uint MaxScatteringGroupCount;
-float2 MaxInputBufferUV;
 float MinScatteringCocRadius;
 float PreProcessingToProcessingCocRadiusFactor;
 float BokehGatherDistinctionLimit;
@@ -77,11 +76,13 @@ float BarrelRadius;
 float BarrelLength;
 uint4 MatteBoxPlanes[MAX_MATTE_BOX_FLAGS];
 
-float4 GatherInputSize;
+float4 GatherInputExtentAndInvExtent;
+float4 GatherInputViewRectUV;
 Texture2D GatherInput_SceneColor;
 Texture2D GatherInput_SeparateCoc;
 
-float4 QuarterResGatherInputSize;
+float4 QuarterResGatherInputExtentAndInvExtent;
+float4 QuarterResGatherInputViewRectUV;
 Texture2D QuarterResGatherInput_SceneColor;
 
 
@@ -361,11 +362,11 @@ void ReduceCS(
 
 	uint2 DispatchThreadId = THREADGROUP_TILE_SIZE * GroupId + GroupThreadId;
 	
-	float2 BufferUV = (DispatchThreadId + 0.5) * GatherInputSize.zw;
-
-	if (true)
+	float2 BufferUV;
 	{
-		BufferUV = min(BufferUV, MaxInputBufferUV);
+		float2 PixelPosition = DispatchThreadId + 0.5;
+		BufferUV = GatherInputViewRectUV.xy + (PixelPosition * GatherInputExtentAndInvExtent.zw);
+		BufferUV = clamp(BufferUV, GatherInputViewRectUV.xy, GatherInputViewRectUV.zw);
 	}
 
 	// Fetch scene color.
@@ -404,16 +405,26 @@ void ReduceCS(
 		// is getting close to viewport edge.
 		float ScatterFactor;
 		{
+			float2 QuarterResBufferUV;
+			{
+				float2 FullResViewRectPixelSize = (GatherInputViewRectUV.zw - GatherInputViewRectUV.xy) * GatherInputExtentAndInvExtent.xy;
+				float2 QuarterResViewRectPixelSize = (QuarterResGatherInputViewRectUV.zw - QuarterResGatherInputViewRectUV.xy) * QuarterResGatherInputExtentAndInvExtent.xy;
+				float2 Rescale = QuarterResViewRectPixelSize / FullResViewRectPixelSize;
+
+				float2 PixelPosition = (BufferUV - GatherInputViewRectUV.xy) * GatherInputExtentAndInvExtent.xy * Rescale;
+				QuarterResBufferUV = QuarterResGatherInputViewRectUV.xy + (PixelPosition * QuarterResGatherInputExtentAndInvExtent.zw);
+			}
+
 			FHybridScatterInputs Parameters;
 			Parameters.NeighborhoodComparisonBuffer = QuarterResGatherInput_SceneColor;
-			Parameters.NeighborhoodComparisonBufferInvSize = QuarterResGatherInputSize.zw;
-			Parameters.NeighborhoodComparisonMaxBufferUV = MaxInputBufferUV;
+			Parameters.NeighborhoodComparisonBufferInvExtent = QuarterResGatherInputExtentAndInvExtent.zw;
+			Parameters.NeighborhoodComparisonViewRectUV = QuarterResGatherInputViewRectUV;
 
 			Parameters.FrameExposureScale = FrameExposureScale;
 
 			Parameters.Color = GatheredColor;
 			Parameters.CocRadius = ScatterCocRadius;
-			Parameters.BufferUV = BufferUV;
+			Parameters.BufferUV = QuarterResBufferUV;
 
 			ScatterFactor = ComputeLuminanceOnlyBasedScatterFactor(Parameters);
 
@@ -439,7 +450,7 @@ void ReduceCS(
 
 				float ClosestBorderDistance = float(min(
 					min(SvPosition.x, SvPosition.y),
-					min(ViewportRect.z - SvPosition.x, ViewportRect.w - SvPosition.y)));
+					min(PassViewSize.x - SvPosition.x, PassViewSize.y - SvPosition.y)));
 	
 				ScatterFactor *= saturate(ClosestBorderDistance - abs(ScatterCocRadius));
 			}
diff --git a/Engine/Source/Runtime/Renderer/Private/PostProcess/DiaphragmDOF.cpp b/Engine/Source/Runtime/Renderer/Private/PostProcess/DiaphragmDOF.cpp
--- a/Engine/Source/Runtime/Renderer/Private/PostProcess/DiaphragmDOF.cpp
+++ b/Engine/Source/Runtime/Renderer/Private/PostProcess/DiaphragmDOF.cpp
@@ -818,10 +818,10 @@ class FDiaphragmDOFCocFlattenCS : public FDiaphragmDOFShader
 	using FPermutationDomain = TShaderPermutationDomain<FDiaphragmDOFShader::FAlphaChannelDim, FDoCocGather4>;
 	
 	BEGIN_SHADER_PARAMETER_STRUCT(FParameters, )
-		SHADER_PARAMETER(FIntRect, ViewportRect)
-		SHADER_PARAMETER(FVector2f, ThreadIdToBufferUV)
-		SHADER_PARAMETER(FVector2f, MaxBufferUV)
-		
+		SHADER_PARAMETER(FVector2f, Stride)
+		SHADER_PARAMETER(FVector4f, GatherInputExtentAndInvExtent)
+		SHADER_PARAMETER(FVector4f, GatherInputViewRectUV)
+
 		SHADER_PARAMETER_STRUCT_INCLUDE(FDOFCommonShaderParameters, CommonParameters)
 		SHADER_PARAMETER_STRUCT(FDOFGatherInputTextures, GatherInput)
 		SHADER_PARAMETER_STRUCT(FDOFTileClassificationUAVs, TileOutput)
@@ -872,10 +872,10 @@ class FDiaphragmDOFDownsampleCS : public FDiaphragmDOFShader
 
 	BEGIN_SHADER_PARAMETER_STRUCT(FParameters, )
 		SHADER_PARAMETER(FIntRect, ViewportRect)
-		SHADER_PARAMETER(FVector2f, MaxBufferUV)
 		SHADER_PARAMETER(float, OutputCocRadiusMultiplier)
 
-		SHADER_PARAMETER(FVector4f, GatherInputSize)
+		SHADER_PARAMETER(FVector4f, GatherInputExtentAndInvExtent)
+		SHADER_PARAMETER(FVector4f, GatherInputViewRectUV)
 		SHADER_PARAMETER_STRUCT(FDOFGatherInputTextures, GatherInput)
 
 		SHADER_PARAMETER_STRUCT_INCLUDE(FDOFCommonShaderParameters, CommonParameters)
@@ -960,8 +960,7 @@ class FDiaphragmDOFReduceCS : public FDiaphragmDOFShader
 	}
 
 	BEGIN_SHADER_PARAMETER_STRUCT(FParameters, )
-		SHADER_PARAMETER(FIntRect, ViewportRect)
-		SHADER_PARAMETER(FVector2f, MaxInputBufferUV)
+		SHADER_PARAMETER(FIntPoint, PassViewSize)
 		SHADER_PARAMETER(FVector4f, ScatteringViewportSize)
 		SHADER_PARAMETER(float, ScatteringScaling)
 		SHADER_PARAMETER(int32, MaxScatteringGroupCount)
@@ -983,10 +982,12 @@ class FDiaphragmDOFReduceCS : public FDiaphragmDOFShader
 		SHADER_PARAMETER_RDG_BUFFER_SRV(StructuredBuffer<float4>, EyeAdaptationBuffer)
 		SHADER_PARAMETER_STRUCT_INCLUDE(FDOFCommonShaderParameters, CommonParameters)
 
-		SHADER_PARAMETER(FVector4f, GatherInputSize)
+		SHADER_PARAMETER(FVector4f, GatherInputExtentAndInvExtent)
+		SHADER_PARAMETER(FVector4f, GatherInputViewRectUV)
 		SHADER_PARAMETER_STRUCT(FDOFGatherInputSRVs, GatherInput)
 		
-		SHADER_PARAMETER(FVector4f, QuarterResGatherInputSize)
+		SHADER_PARAMETER(FVector4f, QuarterResGatherInputExtentAndInvExtent)
+		SHADER_PARAMETER(FVector4f, QuarterResGatherInputViewRectUV)
 		SHADER_PARAMETER_STRUCT(FDOFGatherInputTextures, QuarterResGatherInput)
 
 		SHADER_PARAMETER_STRUCT_ARRAY(FDOFGatherInputUAVs, OutputMips, [kMaxMipLevelCount])
@@ -1430,6 +1431,9 @@ class FDiaphragmDOFRecombineCS : public FDiaphragmDOFShader
 		SHADER_PARAMETER(float, EncodedCocRadiusToRecombineCocRadius)
 		SHADER_PARAMETER(float, MaxRecombineAbsCocRadius)
 		SHADER_PARAMETER(float, SOFHistoryRejectionBoost)
+
+		SHADER_PARAMETER(FVector4f, SceneColorExtentAndInvExtent)
+		SHADER_PARAMETER(FVector4f, SceneColorViewRectUV)
 		
 		SHADER_PARAMETER_RDG_TEXTURE(Texture2D, BokehLUT)
 
@@ -1745,7 +1749,9 @@ bool DiaphragmDOF::AddPasses(
 	CommonParameters.ViewUniformBuffer = View.ViewUniformBuffer;
 
 	FDOFGatherInputTextures FullResGatherInputTextures;
+	FIntRect FullResGatherInputViewRect = View.ViewRect;
 	FDOFGatherInputTextures HalfResGatherInputTextures;
+	FIntRect HalfResGatherInputViewRect = View.ViewRect / 2;
 
 	// Setup at lower resolution from full resolution scene color and scene depth.
 	{
@@ -1829,6 +1835,7 @@ bool DiaphragmDOF::AddPasses(
 		{
 			check(PrefilteringResolutionDivisor == 1);
 			HalfResGatherInputTextures = FullResGatherInputTextures;
+			HalfResGatherInputViewRect = FullResGatherInputViewRect;
 		}
 	}
 	
@@ -1850,6 +1857,9 @@ bool DiaphragmDOF::AddPasses(
 		
 		HalfResGatherInputTextures.SceneColor = TAAOutputs.SceneColor;
 		HalfResGatherInputTextures.SeparateCoc = TAAOutputs.SceneMetadata;
+		HalfResGatherInputViewRect = TAAParameters.OutputViewRect;
+		HalfResGatherInputViewRect.Min = FIntPoint::DivideAndRoundUp(HalfResGatherInputViewRect.Min, TAAParameters.ResolutionDivisor);
+		HalfResGatherInputViewRect.Max = FIntPoint::DivideAndRoundUp(HalfResGatherInputViewRect.Max, TAAParameters.ResolutionDivisor);
 
 		HalfResGatherInputDescs.SceneColor = TAAOutputs.SceneColor->Desc;
 		if (TAAOutputs.SceneMetadata)
@@ -1876,23 +1886,31 @@ bool DiaphragmDOF::AddPasses(
 		// Adds a coc flatten pass.
 		FDOFTileClassificationTextures FlattenedTileClassificationTextures;
 		{
-			FIntPoint SrcSize = HalfResGatherInputTextures.SceneColor->Desc.Extent;
-
 			static const TCHAR* OutputDebugNames[2] = {
 				TEXT("DOF.FlattenFgdCoc"),
 				TEXT("DOF.FlattenBgdCoc"),
 			};
 			FlattenedTileClassificationTextures = CreateTextures(GraphBuilder, TileClassificationDescs, OutputDebugNames);
 
+			const FIntPoint HalfResExtent = HalfResGatherInputDescs.SceneColor.Extent;
+			const FVector4f HalfResExtentAndInvExtent
+			{
+				float(HalfResExtent.X), float(HalfResExtent.Y), 1.0f / HalfResExtent.X, 1.0f / HalfResExtent.Y
+			};
+			const FVector4f HalfResViewRectUV
+			{
+				FVector2f(HalfResGatherInputViewRect.Min) / HalfResExtent, 
+				(FVector2f(HalfResGatherInputViewRect.Max) - 0.5f) / HalfResExtent
+			};
+
 			FDiaphragmDOFCocFlattenCS::FParameters* PassParameters = GraphBuilder.AllocParameters<FDiaphragmDOFCocFlattenCS::FParameters>();
 			PassParameters->CommonParameters = CommonParameters;
-			PassParameters->ViewportRect = FIntRect(0, 0, GatheringViewSize.X, GatheringViewSize.Y);
-			PassParameters->ThreadIdToBufferUV = FVector2f(
-				PreprocessViewSize.X / float(GatheringViewSize.X * SrcSize.X),
-				PreprocessViewSize.Y / float(GatheringViewSize.Y * SrcSize.Y));
-			PassParameters->MaxBufferUV = FVector2f(
-				(PreprocessViewSize.X - 1.0f) / float(SrcSize.X),
-				(PreprocessViewSize.Y - 1.0f) / float(SrcSize.Y));
+
+			PassParameters->Stride = FVector2f(
+				PreprocessViewSize.X / float(GatheringViewSize.X),
+				PreprocessViewSize.Y / float(GatheringViewSize.Y));
+			PassParameters->GatherInputExtentAndInvExtent = HalfResExtentAndInvExtent;
+			PassParameters->GatherInputViewRectUV = HalfResViewRectUV;
 			PassParameters->GatherInput = HalfResGatherInputTextures;
 			PassParameters->TileOutput = CreateUAVs(GraphBuilder, FlattenedTileClassificationTextures);
 			
@@ -2086,6 +2104,17 @@ bool DiaphragmDOF::AddPasses(
 	{
 		FIntPoint SrcSize = HalfResGatherInputDescs.SceneColor.Extent;
 		
+		const FIntPoint HalfResExtent = HalfResGatherInputDescs.SceneColor.Extent;
+		const FVector4f HalfResExtentAndInvExtent
+		{
+			float(HalfResExtent.X), float(HalfResExtent.Y), 1.0f / HalfResExtent.X, 1.0f / HalfResExtent.Y
+		};
+		const FVector4f HalfResViewRectUV
+		{
+			FVector2f(HalfResGatherInputViewRect.Min) / HalfResExtent, 
+			(FVector2f(HalfResGatherInputViewRect.Max) - 0.5f) / HalfResExtent
+		};
+		
 		// Compute the number of mip level required by the gathering pass.
 		const int32 MipLevelsRequested = FMath::CeilToInt(FMath::Log2(MaxBluringRadius * 0.5 / HalfResRingCount)) + (bUseLowAccumulatorQuality ? 1 : 0);
 		const int32 MipLevelCount = FMath::Clamp(MipLevelsRequested, 1, CVarGatherMaxMipMapCount.GetValueOnRenderThread());
@@ -2123,14 +2152,27 @@ bool DiaphragmDOF::AddPasses(
 
 		// Downsample the gather color setup to have faster neighborhood comparisons.
 		FDOFGatherInputTextures QuarterResGatherInputTextures;
+		FDOFGatherInputDescs QuarterResGatherInputDescs;
+		{
+			QuarterResGatherInputDescs = HalfResGatherInputDescs;
+			QuarterResGatherInputDescs.SceneColor.Extent /= 2;
+			QuarterResGatherInputDescs.SeparateCoc.Extent /= 2;
+		}
+		const FIntPoint QuarterResExtent = QuarterResGatherInputDescs.SceneColor.Extent;
+		const FVector4f QuarterResExtentAndInvExtent
+		{
+			float(QuarterResExtent.X), float(QuarterResExtent.Y), 1.0f / QuarterResExtent.X, 1.0f / QuarterResExtent.Y
+		};
+		const FIntRect QuarterResGatherInputViewRect{ FIntPoint::ZeroValue, HalfResGatherInputViewRect.Size() / 2 }; // rebase to (0,0) and halve size
+		const FVector4f QuarterResViewRectUV
+		{
+			FVector2f(QuarterResGatherInputViewRect.Min) / QuarterResExtent, 
+			(FVector2f(QuarterResGatherInputViewRect.Max) - 0.5f) / QuarterResExtent
+		};
 		if (bForegroundHybridScattering || bBackgroundHybridScattering)
 		{
 			// Allocate quarter res textures.
 			{
-				FDOFGatherInputDescs QuarterResGatherInputDescs = HalfResGatherInputDescs;
-				QuarterResGatherInputDescs.SceneColor.Extent /= 2;
-				QuarterResGatherInputDescs.SeparateCoc.Extent /= 2;
-
 				// Lower the bit depth to speed up texture fetches in reduce pass, that is ok since this is used only for comparison purposes.
 				if (bRGBBufferSeparateCocBuffer && !bProcessSceneAlpha)
 					QuarterResGatherInputDescs.SceneColor.Format = PF_FloatR11G11B10;
@@ -2146,12 +2188,9 @@ bool DiaphragmDOF::AddPasses(
 			FDiaphragmDOFDownsampleCS::FParameters* PassParameters = GraphBuilder.AllocParameters<FDiaphragmDOFDownsampleCS::FParameters>();
 			PassParameters->CommonParameters = CommonParameters;
 			PassParameters->ViewportRect = FIntRect(0, 0, PassViewSize.X, PassViewSize.Y);
-			PassParameters->MaxBufferUV = FVector2f(
-				(PreprocessViewSize.X - 0.5f) / float(SrcSize.X),
-				(PreprocessViewSize.Y - 0.5f) / float(SrcSize.Y));
 			PassParameters->OutputCocRadiusMultiplier = GatheringCocRadiusBasis / EncodedCocRadiusBasis;
-
-			PassParameters->GatherInputSize = FVector4f(SrcSize.X, SrcSize.Y, 1.0f / SrcSize.X, 1.0f / SrcSize.Y);
+			PassParameters->GatherInputExtentAndInvExtent = HalfResExtentAndInvExtent;
+			PassParameters->GatherInputViewRectUV = HalfResViewRectUV;
 			PassParameters->GatherInput = HalfResGatherInputTextures;
 
 			PassParameters->OutDownsampledGatherInput = CreateUAVs(GraphBuilder, QuarterResGatherInputTextures);
@@ -2200,10 +2239,7 @@ bool DiaphragmDOF::AddPasses(
 			PermutationVector.Set<FDDOFRGBColorBufferDim>(bRGBBufferSeparateCocBuffer);
 
 			FDiaphragmDOFReduceCS::FParameters* PassParameters = GraphBuilder.AllocParameters<FDiaphragmDOFReduceCS::FParameters>();
-			PassParameters->ViewportRect = FIntRect(0, 0, PassViewSize.X, PassViewSize.Y);
-			PassParameters->MaxInputBufferUV = FVector2f(
-				(PreprocessViewSize.X - 0.5f) / SrcSize.X,
-				(PreprocessViewSize.Y - 0.5f) / SrcSize.Y);
+			PassParameters->PassViewSize = PassViewSize;
 			PassParameters->ScatteringViewportSize = FVector4f(GatheringViewSize.X, GatheringViewSize.Y, 1.0f / GatheringViewSize.X, 1.0f / GatheringViewSize.Y);
 			PassParameters->ScatteringScaling = float(GatheringViewSize.X) / float(PreprocessViewSize.X);
 			PassParameters->MaxScatteringGroupCount = MaxScatteringGroupCount;
@@ -2223,10 +2259,12 @@ bool DiaphragmDOF::AddPasses(
 			PassParameters->EyeAdaptationBuffer = GraphBuilder.CreateSRV(GetEyeAdaptationBuffer(GraphBuilder, View));
 			PassParameters->CommonParameters = CommonParameters;
 
-			PassParameters->GatherInputSize = FVector4f(SrcSize.X, SrcSize.Y, 1.0f / SrcSize.X, 1.0f / SrcSize.Y);
+			PassParameters->GatherInputExtentAndInvExtent = HalfResExtentAndInvExtent;
+			PassParameters->GatherInputViewRectUV = HalfResViewRectUV;
 			PassParameters->GatherInput = CreateSRVs(GraphBuilder, HalfResGatherInputTextures);
 			
-			PassParameters->QuarterResGatherInputSize = FVector4f(SrcSize.X / 2, SrcSize.Y / 2, 2.0f / SrcSize.X, 2.0f / SrcSize.Y);
+			PassParameters->QuarterResGatherInputExtentAndInvExtent = QuarterResExtentAndInvExtent;
+			PassParameters->QuarterResGatherInputViewRectUV = QuarterResViewRectUV;
 			PassParameters->QuarterResGatherInput = QuarterResGatherInputTextures;
 
 			for (int32 MipLevel = 0; MipLevel < ProcessingMipLevelCount; MipLevel++)
@@ -2277,16 +2315,22 @@ bool DiaphragmDOF::AddPasses(
 			PermutationVector.Set<FDiaphragmDOFReduceCS::FDisableOutputMip0>(true);
 
 			FDiaphragmDOFReduceCS::FParameters* PassParameters = GraphBuilder.AllocParameters<FDiaphragmDOFReduceCS::FParameters>();
-			PassParameters->ViewportRect = FIntRect(0, 0, PassViewSize.X, PassViewSize.Y);
-			PassParameters->MaxInputBufferUV = FVector2f(
-				(PreprocessViewSize.X - 0.5f) / GatherInputExtent.X,
-				(PreprocessViewSize.Y - 0.5f) / GatherInputExtent.Y);
+			PassParameters->PassViewSize = PassViewSize;
 
 			PassParameters->EyeAdaptationBuffer = GraphBuilder.CreateSRV(GetEyeAdaptationBuffer(GraphBuilder, View));
 			PassParameters->CommonParameters = CommonParameters;
 
+			const FVector4f ViewRectUV
+			{ 
+				0.0f, 
+				0.0f,
+				(PreprocessViewSize.X - 0.5f) / GatherInputExtent.X,
+				(PreprocessViewSize.Y - 0.5f) / GatherInputExtent.Y
+			};
+
 			float InputMipLevelPow2 = 1 << InputMipLevel;
-			PassParameters->GatherInputSize = FVector4f(GatherInputExtent.X / InputMipLevelPow2, GatherInputExtent.Y / InputMipLevelPow2, InputMipLevelPow2 / GatherInputExtent.X, InputMipLevelPow2 / GatherInputExtent.Y);
+			PassParameters->GatherInputExtentAndInvExtent = FVector4f(GatherInputExtent.X / InputMipLevelPow2, GatherInputExtent.Y / InputMipLevelPow2, InputMipLevelPow2 / GatherInputExtent.X, InputMipLevelPow2 / GatherInputExtent.Y);
+			PassParameters->GatherInputViewRectUV = ViewRectUV;
 			PassParameters->GatherInput = CreateSRVs(GraphBuilder, ReducedGatherInputTextures, InputMipLevel);
 
 			for (int32 MipLevel = 0; MipLevel < ProcessingMipLevelCount; MipLevel++)
@@ -3011,6 +3055,17 @@ bool DiaphragmDOF::AddPasses(
 		FRDGTextureRef SeparateTranslucencyDepth = TranslucencyPassResources.GetDepthForRead(GraphBuilder);
 		FRDGTextureRef SeparateTranslucencyModulateColor = TranslucencyPassResources.GetColorModulateForRead(GraphBuilder);
 
+		const FIntPoint FullResExtent = FullResGatherInputDescs.SceneColor.Extent;
+		const FVector4f FullResExtentAndInvExtent
+		{
+			float(FullResExtent.X), float(FullResExtent.Y), 1.0f / FullResExtent.X, 1.0f / FullResExtent.Y
+		};
+		const FVector4f FullResViewRectUV
+		{
+			FVector2f(FullResGatherInputViewRect.Min) / FullResExtent, 
+			(FVector2f(FullResGatherInputViewRect.Max) - 0.5f) / FullResExtent
+		};
+
 		FDiaphragmDOFRecombineCS::FParameters* PassParameters = GraphBuilder.AllocParameters<FDiaphragmDOFRecombineCS::FParameters>();
 		PassParameters->CommonParameters = CommonParameters;
 		SetCocModelParameters(GraphBuilder, &PassParameters->CocModel, CocModel, /* CocRadiusBasis = */ float(GatheringViewSize.X));
@@ -3039,6 +3094,9 @@ bool DiaphragmDOF::AddPasses(
 		PassParameters->SeparateTranslucencyBilinearUVMinMax.W = (SeparateTranslucencyViewport.Rect.Max.Y - 0.5f) / float(SeparateTranslucencyViewport.Extent.Y);
 		PassParameters->SeparateTranslucencyUpscaling = bScaleSeparateTranslucency ? 1 : 0;
 
+		PassParameters->SceneColorExtentAndInvExtent = FullResExtentAndInvExtent;
+		PassParameters->SceneColorViewRectUV = FullResViewRectUV;
+
 		PassParameters->SceneColorInput = FullResGatherInputTextures.SceneColor;
 		PassParameters->SceneDepthTexture = SceneTextures.SceneDepthTexture;
 		PassParameters->SceneSeparateCoc = FullResGatherInputTextures.SeparateCoc; // TODO looks useless.
