Fixed: Restrict AddCopySceneColorPass for SLW case only to mobile platforms
#Synced-CL 48389552

diff --git a/Engine/Source/Runtime/Renderer/Private/TranslucentRendering.cpp b/Engine/Source/Runtime/Renderer/Private/TranslucentRendering.cpp
--- a/Engine/Source/Runtime/Renderer/Private/TranslucentRendering.cpp
+++ b/Engine/Source/Runtime/Renderer/Private/TranslucentRendering.cpp
@@ -339,7 +339,7 @@ FRDGTextureRef AddCopySceneColorPass(FRDGBuilder& GraphBuilder, TArrayView<const
 				break;
 			}
 		}
-		bNeedsResolve = bNeedsResolve || View.bHasSingleLayerWaterMaterial;
+		bNeedsResolve = bNeedsResolve || (IsMobilePlatform(View.GetShaderPlatform()) && View.bHasSingleLayerWaterMaterial);
 
 		if (bNeedsResolve)
 		{
