Add flags to BeginScenePrimitiveRendering methods.
Add a flag that allows for skipping the render scene update.  Useful for situations where the execution of the render task is not dispatched from the game thread, but all primitives that are going to render have already been updated in scene.

diff --git a/Engine/Source/Runtime/RenderCore/Public/RendererInterface.h b/Engine/Source/Runtime/RenderCore/Public/RendererInterface.h
--- a/Engine/Source/Runtime/RenderCore/Public/RendererInterface.h
+++ b/Engine/Source/Runtime/RenderCore/Public/RendererInterface.h
@@ -547,6 +547,14 @@ enum EDrawRectangleFlags
 	EDRF_UseTesselatedIndexBuffer
 };
 
+enum class EBeginScenePrimitiveRenderingFlags
+{
+	None = 0,
+	UpdateScenePrimitives = 1 << 0,
+	All = UpdateScenePrimitives
+};
+ENUM_CLASS_FLAGS(EBeginScenePrimitiveRenderingFlags)
+
 class FPostOpaqueRenderParameters
 {
 public:
@@ -879,11 +887,13 @@ public:
 	virtual void RegisterPersistentViewUniformBufferExtension(IPersistentViewUniformBufferExtension* Extension) = 0;
 
 	/**
-	 * Prepare Scene primitive rendering and return context. Ensures all primitives that are created are commited and GPU-Scene is updated and allocates a dynamic primitive context.
+	 * Prepare Scene primitive rendering and return context. 
+	 * Ensures all primitives that are created are commited and GPU-Scene is updated (if EBeginScenePrimitiveRenderingFlags::UpdateScenePrimitives is passed in BeginScenePrimitiveRenderingFlags).
+	 * Allocates a dynamic primitive context.
 	 * The intended use is for stand-alone rendering that involves Scene proxies (that then may need the machinery to render GPU-Scene aware primitives.
 	 */
-	virtual IScenePrimitiveRenderingContext* BeginScenePrimitiveRendering(FRDGBuilder& GraphBuilder, FSceneViewFamily* ViewFamily) = 0;
-	virtual IScenePrimitiveRenderingContext* BeginScenePrimitiveRendering(FRDGBuilder& GraphBuilder, FSceneInterface& Scene) = 0;
+	virtual IScenePrimitiveRenderingContext* BeginScenePrimitiveRendering(FRDGBuilder& GraphBuilder, FSceneViewFamily* ViewFamily, EBeginScenePrimitiveRenderingFlags BeginScenePrimitiveRenderingFlags = EBeginScenePrimitiveRenderingFlags::All) = 0;
+	virtual IScenePrimitiveRenderingContext* BeginScenePrimitiveRendering(FRDGBuilder& GraphBuilder, FSceneInterface& Scene, EBeginScenePrimitiveRenderingFlags BeginScenePrimitiveRenderingFlags = EBeginScenePrimitiveRenderingFlags::All) = 0;
 
 	/** Mark all the current scenes as needing to restart path tracer accumulation */
 	virtual void InvalidatePathTracedOutput(PathTracing::EInvalidateReason InvalidateReason = PathTracing::EInvalidateReason::Uncategorized) = 0;
diff --git a/Engine/Source/Runtime/Renderer/Private/RendererModule.h b/Engine/Source/Runtime/Renderer/Private/RendererModule.h
--- a/Engine/Source/Runtime/Renderer/Private/RendererModule.h
+++ b/Engine/Source/Runtime/Renderer/Private/RendererModule.h
@@ -142,8 +142,8 @@ public:
 
 	void RenderPostResolvedSceneColorExtension(FRDGBuilder& GraphBuilder, const FSceneTextures& SceneTextures);
 
-	virtual IScenePrimitiveRenderingContext* BeginScenePrimitiveRendering(FRDGBuilder& GraphBuilder, FSceneViewFamily* ViewFamily) override;
-	virtual IScenePrimitiveRenderingContext* BeginScenePrimitiveRendering(FRDGBuilder& GraphBuilder, FSceneInterface& Scene) override;
+	virtual IScenePrimitiveRenderingContext* BeginScenePrimitiveRendering(FRDGBuilder& GraphBuilder, FSceneViewFamily* ViewFamily, EBeginScenePrimitiveRenderingFlags BeginScenePrimitiveRenderingFlags = EBeginScenePrimitiveRenderingFlags::All) override;
+	virtual IScenePrimitiveRenderingContext* BeginScenePrimitiveRendering(FRDGBuilder& GraphBuilder, FSceneInterface& Scene, EBeginScenePrimitiveRenderingFlags BeginScenePrimitiveRenderingFlags = EBeginScenePrimitiveRenderingFlags::All) override;
 
 	virtual void InvalidatePathTracedOutput(PathTracing::EInvalidateReason InvalidateReason = PathTracing::EInvalidateReason::Uncategorized) override;
 
diff --git a/Engine/Source/Runtime/Renderer/Private/SceneRendering.cpp b/Engine/Source/Runtime/Renderer/Private/SceneRendering.cpp
--- a/Engine/Source/Runtime/Renderer/Private/SceneRendering.cpp
+++ b/Engine/Source/Runtime/Renderer/Private/SceneRendering.cpp
@@ -5412,7 +5412,7 @@ void FRendererModule::RenderPostResolvedSceneColorExtension(FRDGBuilder& GraphBu
 class FScenePrimitiveRenderingContext : public IScenePrimitiveRenderingContext
 {
 public:
-	FScenePrimitiveRenderingContext(FRDGBuilder& InGraphBuilder, FScene& Scene, FSceneViewFamily* InViewFamily = nullptr) :
+	FScenePrimitiveRenderingContext(FRDGBuilder& InGraphBuilder, FScene& Scene, FSceneViewFamily* InViewFamily = nullptr, EBeginScenePrimitiveRenderingFlags BeginScenePrimitiveRenderingFlags = EBeginScenePrimitiveRenderingFlags::All) :
 		GraphBuilder(InGraphBuilder),
 		GPUScene(Scene.GPUScene),
 		GPUSceneDynamicContext(GPUScene),
@@ -5425,7 +5425,11 @@ public:
 			ViewFamily->SetSceneRenderer(&Renderer);
 		}
 
-		Scene.UpdateAllPrimitiveSceneInfos(GraphBuilder);
+		if (EnumHasAllFlags(BeginScenePrimitiveRenderingFlags, EBeginScenePrimitiveRenderingFlags::UpdateScenePrimitives))
+		{
+			Scene.UpdateAllPrimitiveSceneInfos(GraphBuilder);
+		}
+		
 		GPUScene.BeginRender(GraphBuilder, GPUSceneDynamicContext);
 
 		GPUScene.FillSceneUniformBuffer(GraphBuilder, Renderer.GetSceneUniforms());
@@ -5457,24 +5461,24 @@ public:
 };
 
 
-IScenePrimitiveRenderingContext* FRendererModule::BeginScenePrimitiveRendering(FRDGBuilder &GraphBuilder, FSceneViewFamily* ViewFamily)
+IScenePrimitiveRenderingContext* FRendererModule::BeginScenePrimitiveRendering(FRDGBuilder &GraphBuilder, FSceneViewFamily* ViewFamily, EBeginScenePrimitiveRenderingFlags BeginScenePrimitiveRenderingFlags)
 {
 	check(ViewFamily);
 	check(ViewFamily->Scene);
 	FScene* Scene = ViewFamily->Scene->GetRenderScene();
 	check(Scene);
 
-	FScenePrimitiveRenderingContext* ScenePrimitiveRenderingContext = new FScenePrimitiveRenderingContext(GraphBuilder, *Scene, ViewFamily);
+	FScenePrimitiveRenderingContext* ScenePrimitiveRenderingContext = new FScenePrimitiveRenderingContext(GraphBuilder, *Scene, ViewFamily, BeginScenePrimitiveRenderingFlags);
 
 	return ScenePrimitiveRenderingContext;
 }
 
-IScenePrimitiveRenderingContext* FRendererModule::BeginScenePrimitiveRendering(FRDGBuilder &GraphBuilder, FSceneInterface& InScene)
+IScenePrimitiveRenderingContext* FRendererModule::BeginScenePrimitiveRendering(FRDGBuilder &GraphBuilder, FSceneInterface& InScene, EBeginScenePrimitiveRenderingFlags BeginScenePrimitiveRenderingFlags)
 {
 	FScene* Scene = InScene.GetRenderScene();
 	check(Scene);
 
-	FScenePrimitiveRenderingContext* ScenePrimitiveRenderingContext = new FScenePrimitiveRenderingContext(GraphBuilder, *Scene);
+	FScenePrimitiveRenderingContext* ScenePrimitiveRenderingContext = new FScenePrimitiveRenderingContext(GraphBuilder, *Scene, nullptr, BeginScenePrimitiveRenderingFlags);
 
 	return ScenePrimitiveRenderingContext;
 }
