- Add check for null resource

#rnx

#Synced-CL 48488387

diff --git a/Engine/Source/Runtime/Renderer/Private/CompositionLighting/PostProcessAmbient.cpp b/Engine/Source/Runtime/Renderer/Private/CompositionLighting/PostProcessAmbient.cpp
--- a/Engine/Source/Runtime/Renderer/Private/CompositionLighting/PostProcessAmbient.cpp
+++ b/Engine/Source/Runtime/Renderer/Private/CompositionLighting/PostProcessAmbient.cpp
@@ -27,7 +27,7 @@ void SetupAmbientCubemapParameters(const FFinalPostProcessSettings::FCubemapEntr
 
 	// cubemap texture
 	{
-		FTexture* AmbientCubemapTexture = Entry.AmbientCubemap ? Entry.AmbientCubemap->GetResource() : (FTexture*)GBlackTextureCube;
+		FTexture* AmbientCubemapTexture = Entry.AmbientCubemap && Entry.AmbientCubemap->GetResource() ? Entry.AmbientCubemap->GetResource() : (FTexture*)GBlackTextureCube;
 		OutParameters->AmbientCubemap = AmbientCubemapTexture->TextureRHI;
 		OutParameters->AmbientCubemapSampler = AmbientCubemapTexture->SamplerStateRHI;
 	}
