[TSR] Give TSR shaders unique entry point name.

#tests local
#Synced-CL 48025942

diff --git a/Engine/Shaders/Private/TemporalSuperResolution/TSRClearPrevTextures.usf b/Engine/Shaders/Private/TemporalSuperResolution/TSRClearPrevTextures.usf
--- a/Engine/Shaders/Private/TemporalSuperResolution/TSRClearPrevTextures.usf
+++ b/Engine/Shaders/Private/TemporalSuperResolution/TSRClearPrevTextures.usf
@@ -18,7 +18,7 @@ RWTexture2DArray<uint> PrevAtomicOutput;
 //------------------------------------------------------- ENTRY POINT
 
 [numthreads(TILE_SIZE * TILE_SIZE, 1, 1)]
-void MainCS(
+void ClearPrevTexturesCS(
 	uint2 GroupId : SV_GroupID,
 	uint GroupThreadIndex : SV_GroupIndex)
 {
diff --git a/Engine/Shaders/Private/TemporalSuperResolution/TSRDecimateHistory.usf b/Engine/Shaders/Private/TemporalSuperResolution/TSRDecimateHistory.usf
--- a/Engine/Shaders/Private/TemporalSuperResolution/TSRDecimateHistory.usf
+++ b/Engine/Shaders/Private/TemporalSuperResolution/TSRDecimateHistory.usf
@@ -254,7 +254,7 @@ void ReprojectAllPrevTextures(
 //------------------------------------------------------- ENTRY POINT
 
 [numthreads(TILE_SIZE * TILE_SIZE, 1, 1)]
-void MainCS(
+void DecimateHistoryCS(
 	uint2 GroupId : SV_GroupID,
 	uint GroupThreadIndex : SV_GroupIndex)
 {
diff --git a/Engine/Shaders/Private/TemporalSuperResolution/TSRDetectThinGeometry.usf b/Engine/Shaders/Private/TemporalSuperResolution/TSRDetectThinGeometry.usf
--- a/Engine/Shaders/Private/TemporalSuperResolution/TSRDetectThinGeometry.usf
+++ b/Engine/Shaders/Private/TemporalSuperResolution/TSRDetectThinGeometry.usf
@@ -267,7 +267,7 @@ tsr_tensor_uint EncodeThinGeometryWeight(tsr_tensor_half ThinGeometryEdge
 WAVESIZE(DIM_WAVE_SIZE)
 #endif
 [numthreads(LANE_COUNT * WAVE_COUNT, 1, 1)]
-void MainCS(
+void DetectThinGeometryCS(
 	uint2 GroupId : SV_GroupID,
 	uint GroupThreadIndex : SV_GroupIndex)
 {
diff --git a/Engine/Shaders/Private/TemporalSuperResolution/TSRDilateVelocity.usf b/Engine/Shaders/Private/TemporalSuperResolution/TSRDilateVelocity.usf
--- a/Engine/Shaders/Private/TemporalSuperResolution/TSRDilateVelocity.usf
+++ b/Engine/Shaders/Private/TemporalSuperResolution/TSRDilateVelocity.usf
@@ -146,7 +146,7 @@ tsr_half2 ComputeReprojectionBoundary(
 //------------------------------------------------------- ENTRY POINT
 
 [numthreads(TILE_SIZE * TILE_SIZE, 1, 1)]
-void MainCS(
+void DilateVelocityCS(
 	uint2 GroupId : SV_GroupID,
 	uint GroupThreadIndex : SV_GroupIndex)
 {
diff --git a/Engine/Shaders/Private/TemporalSuperResolution/TSRMeasureCoverage.usf b/Engine/Shaders/Private/TemporalSuperResolution/TSRMeasureCoverage.usf
--- a/Engine/Shaders/Private/TemporalSuperResolution/TSRMeasureCoverage.usf
+++ b/Engine/Shaders/Private/TemporalSuperResolution/TSRMeasureCoverage.usf
@@ -39,7 +39,7 @@ uint GetThinGeometry(const FGBufferData GBufferData)
 //------------------------------------------------------- ENTRY POINT
 
 [numthreads(TILE_SIZE * TILE_SIZE, 1, 1)]
-void MainCS(
+void MeasureCoverageCS(
 	uint2 GroupId : SV_GroupID,
 	uint GroupThreadIndex : SV_GroupIndex)
 {
diff --git a/Engine/Shaders/Private/TemporalSuperResolution/TSRMeasureFlickeringLuma.usf b/Engine/Shaders/Private/TemporalSuperResolution/TSRMeasureFlickeringLuma.usf
--- a/Engine/Shaders/Private/TemporalSuperResolution/TSRMeasureFlickeringLuma.usf
+++ b/Engine/Shaders/Private/TemporalSuperResolution/TSRMeasureFlickeringLuma.usf
@@ -19,7 +19,7 @@ RWTexture2DArray<tsr_half> FlickeringLumaOutput;
 //------------------------------------------------------- ENTRY POINT
 
 [numthreads(TILE_SIZE * TILE_SIZE, 1, 1)]
-void MainCS(
+void MeasureFlickeringLumaCS(
 	uint2 GroupId : SV_GroupID,
 	uint GroupThreadIndex : SV_GroupIndex)
 {
diff --git a/Engine/Shaders/Private/TemporalSuperResolution/TSRRejectShading.usf b/Engine/Shaders/Private/TemporalSuperResolution/TSRRejectShading.usf
--- a/Engine/Shaders/Private/TemporalSuperResolution/TSRRejectShading.usf
+++ b/Engine/Shaders/Private/TemporalSuperResolution/TSRRejectShading.usf
@@ -449,7 +449,7 @@ void RefetchInputAndHistoryFromLDSOrCache(
 WAVESIZE(DIM_WAVE_SIZE)
 #endif
 [numthreads(LANE_COUNT * WAVE_COUNT, 1, 1)]
-void MainCS(
+void RejectShadingCS(
 	uint2 GroupId : SV_GroupID,
 	uint GroupThreadIndex : SV_GroupIndex)
 {
diff --git a/Engine/Shaders/Private/TemporalSuperResolution/TSRResolveHistory.usf b/Engine/Shaders/Private/TemporalSuperResolution/TSRResolveHistory.usf
--- a/Engine/Shaders/Private/TemporalSuperResolution/TSRResolveHistory.usf
+++ b/Engine/Shaders/Private/TemporalSuperResolution/TSRResolveHistory.usf
@@ -146,7 +146,7 @@ RWTexture2D<tsr_halfC> SceneColorOutputMip1;
 	WAVESIZE(DIM_NYQUIST_WAVE_SIZE)
 #endif
 [numthreads(LANE_COUNT, 1, 1)]
-void MainCS(
+void ResolveHistoryCS(
 	uint2 GroupId : SV_GroupID,
 	uint GroupThreadIndex : SV_GroupIndex)
 {
@@ -297,7 +297,7 @@ void MainCS(
 //------------------------------------------------------- ENTRY POINT REFERENCE
 
 [numthreads(TILE_SIZE * TILE_SIZE, 1, 1)]
-void MainCS(
+void ResolveHistoryCS(
 	uint2 GroupId : SV_GroupID,
 	uint GroupThreadIndex : SV_GroupIndex)
 {
diff --git a/Engine/Shaders/Private/TemporalSuperResolution/TSRSpatialAntiAliasing.usf b/Engine/Shaders/Private/TemporalSuperResolution/TSRSpatialAntiAliasing.usf
--- a/Engine/Shaders/Private/TemporalSuperResolution/TSRSpatialAntiAliasing.usf
+++ b/Engine/Shaders/Private/TemporalSuperResolution/TSRSpatialAntiAliasing.usf
@@ -85,7 +85,7 @@ tsr_half2 SampleInputLDRLuma(tsr_short2x2 KernelCenter, tsr_short2 Offset)
 #else
 	[numthreads(TILE_SIZE * TILE_SIZE, 1, 1)]
 #endif
-void MainCS(
+void SpatialAntiAliasingCS(
 	uint2 GroupId : SV_GroupID,
 	uint GroupThreadIndex : SV_GroupIndex)
 {
diff --git a/Engine/Shaders/Private/TemporalSuperResolution/TSRUpdateHistory.usf b/Engine/Shaders/Private/TemporalSuperResolution/TSRUpdateHistory.usf
--- a/Engine/Shaders/Private/TemporalSuperResolution/TSRUpdateHistory.usf
+++ b/Engine/Shaders/Private/TemporalSuperResolution/TSRUpdateHistory.usf
@@ -635,7 +635,7 @@ void FetchVelocityDilateBoundary(
 #else
 	[numthreads(TILE_SIZE * TILE_SIZE, 1, 1)]
 #endif
-void MainCS(
+void UpdateHistoryCS(
 	uint2 GroupId : SV_GroupID,
 	uint GroupThreadIndex : SV_GroupIndex)
 {
diff --git a/Engine/Shaders/Private/TemporalSuperResolution/TSRVisualize.usf b/Engine/Shaders/Private/TemporalSuperResolution/TSRVisualize.usf
--- a/Engine/Shaders/Private/TemporalSuperResolution/TSRVisualize.usf
+++ b/Engine/Shaders/Private/TemporalSuperResolution/TSRVisualize.usf
@@ -494,7 +494,7 @@ tsr_half3 DisplayReprojectionFieldAA(FContext Ctx)
 //------------------------------------------------------- ENTRY POINT
 
 [numthreads(TILE_SIZE * TILE_SIZE, 1, 1)]
-void MainCS(
+void VisualizeCS(
 	uint2 GroupId : SV_GroupID,
 	uint GroupThreadIndex : SV_GroupIndex)
 {
diff --git a/Engine/Shaders/Private/TemporalSuperResolution/TSRWeightRelaxation.usf b/Engine/Shaders/Private/TemporalSuperResolution/TSRWeightRelaxation.usf
--- a/Engine/Shaders/Private/TemporalSuperResolution/TSRWeightRelaxation.usf
+++ b/Engine/Shaders/Private/TemporalSuperResolution/TSRWeightRelaxation.usf
@@ -152,7 +152,7 @@ void FetchInput(
 WAVESIZE(DIM_WAVE_SIZE)
 #endif
 [numthreads(LANE_COUNT * WAVE_COUNT, 1, 1)]
-void MainCS(
+void WeightRelaxationCS(
 	uint2 GroupId : SV_GroupID,
 	uint GroupThreadIndex : SV_GroupIndex)
 {
diff --git a/Engine/Source/Runtime/Renderer/Private/PostProcess/TemporalSuperResolution.cpp b/Engine/Source/Runtime/Renderer/Private/PostProcess/TemporalSuperResolution.cpp
--- a/Engine/Source/Runtime/Renderer/Private/PostProcess/TemporalSuperResolution.cpp
+++ b/Engine/Source/Runtime/Renderer/Private/PostProcess/TemporalSuperResolution.cpp
@@ -1457,18 +1457,18 @@ class FTSRVisualizeCS : public FTSRShader
 	END_SHADER_PARAMETER_STRUCT()
 }; // class FTSRVisualizeCS
 
-IMPLEMENT_GLOBAL_SHADER(FTSRMeasureFlickeringLumaCS, "/Engine/Private/TemporalSuperResolution/TSRMeasureFlickeringLuma.usf", "MainCS", SF_Compute);
-IMPLEMENT_GLOBAL_SHADER(FTSRMeasureThinGeometryCoverageCS, "/Engine/Private/TemporalSuperResolution/TSRMeasureCoverage.usf", "MainCS", SF_Compute);
-IMPLEMENT_GLOBAL_SHADER(FTSRClearPrevTexturesCS,     "/Engine/Private/TemporalSuperResolution/TSRClearPrevTextures.usf",     "MainCS", SF_Compute);
-IMPLEMENT_GLOBAL_SHADER(FTSRDilateVelocityCS,        "/Engine/Private/TemporalSuperResolution/TSRDilateVelocity.usf",        "MainCS", SF_Compute);
-IMPLEMENT_GLOBAL_SHADER(FTSRDetectThinGeometryCS,	 "/Engine/Private/TemporalSuperResolution/TSRDetectThinGeometry.usf",	 "MainCS", SF_Compute);
-IMPLEMENT_GLOBAL_SHADER(FTSRTSRWeightRelaxationCS,	 "/Engine/Private/TemporalSuperResolution/TSRWeightRelaxation.usf",		 "MainCS", SF_Compute);
-IMPLEMENT_GLOBAL_SHADER(FTSRDecimateHistoryCS,       "/Engine/Private/TemporalSuperResolution/TSRDecimateHistory.usf",       "MainCS", SF_Compute);
-IMPLEMENT_GLOBAL_SHADER(FTSRRejectShadingCS,         "/Engine/Private/TemporalSuperResolution/TSRRejectShading.usf",         "MainCS", SF_Compute);
-IMPLEMENT_GLOBAL_SHADER(FTSRSpatialAntiAliasingCS,   "/Engine/Private/TemporalSuperResolution/TSRSpatialAntiAliasing.usf",   "MainCS", SF_Compute);
-IMPLEMENT_GLOBAL_SHADER(FTSRUpdateHistoryCS,         "/Engine/Private/TemporalSuperResolution/TSRUpdateHistory.usf",         "MainCS", SF_Compute);
-IMPLEMENT_GLOBAL_SHADER(FTSRResolveHistoryCS,        "/Engine/Private/TemporalSuperResolution/TSRResolveHistory.usf",        "MainCS", SF_Compute);
-IMPLEMENT_GLOBAL_SHADER(FTSRVisualizeCS,             "/Engine/Private/TemporalSuperResolution/TSRVisualize.usf",             "MainCS", SF_Compute);
+IMPLEMENT_GLOBAL_SHADER(FTSRMeasureFlickeringLumaCS, "/Engine/Private/TemporalSuperResolution/TSRMeasureFlickeringLuma.usf", "MeasureFlickeringLumaCS", SF_Compute);
+IMPLEMENT_GLOBAL_SHADER(FTSRMeasureThinGeometryCoverageCS, "/Engine/Private/TemporalSuperResolution/TSRMeasureCoverage.usf", "MeasureCoverageCS", SF_Compute);
+IMPLEMENT_GLOBAL_SHADER(FTSRClearPrevTexturesCS,     "/Engine/Private/TemporalSuperResolution/TSRClearPrevTextures.usf",     "ClearPrevTexturesCS", SF_Compute);
+IMPLEMENT_GLOBAL_SHADER(FTSRDilateVelocityCS,        "/Engine/Private/TemporalSuperResolution/TSRDilateVelocity.usf",        "DilateVelocityCS", SF_Compute);
+IMPLEMENT_GLOBAL_SHADER(FTSRDetectThinGeometryCS,	 "/Engine/Private/TemporalSuperResolution/TSRDetectThinGeometry.usf",	 "DetectThinGeometryCS", SF_Compute);
+IMPLEMENT_GLOBAL_SHADER(FTSRTSRWeightRelaxationCS,	 "/Engine/Private/TemporalSuperResolution/TSRWeightRelaxation.usf",		 "WeightRelaxationCS", SF_Compute);
+IMPLEMENT_GLOBAL_SHADER(FTSRDecimateHistoryCS,       "/Engine/Private/TemporalSuperResolution/TSRDecimateHistory.usf",       "DecimateHistoryCS", SF_Compute);
+IMPLEMENT_GLOBAL_SHADER(FTSRRejectShadingCS,         "/Engine/Private/TemporalSuperResolution/TSRRejectShading.usf",         "RejectShadingCS", SF_Compute);
+IMPLEMENT_GLOBAL_SHADER(FTSRSpatialAntiAliasingCS,   "/Engine/Private/TemporalSuperResolution/TSRSpatialAntiAliasing.usf",   "SpatialAntiAliasingCS", SF_Compute);
+IMPLEMENT_GLOBAL_SHADER(FTSRUpdateHistoryCS,         "/Engine/Private/TemporalSuperResolution/TSRUpdateHistory.usf",         "UpdateHistoryCS", SF_Compute);
+IMPLEMENT_GLOBAL_SHADER(FTSRResolveHistoryCS,        "/Engine/Private/TemporalSuperResolution/TSRResolveHistory.usf",        "ResolveHistoryCS", SF_Compute);
+IMPLEMENT_GLOBAL_SHADER(FTSRVisualizeCS,             "/Engine/Private/TemporalSuperResolution/TSRVisualize.usf",             "VisualizeCS", SF_Compute);
 
 DECLARE_GPU_STAT(TemporalSuperResolution);
 
