[FastGeo] PointLights, SpotLights and RectLights support.

#Synced-CL 48520266

diff --git a/Engine/Plugins/Experimental/FastGeoStreaming/Config/DefaultFastGeoStreaming.ini b/Engine/Plugins/Experimental/FastGeoStreaming/Config/DefaultFastGeoStreaming.ini
--- a/Engine/Plugins/Experimental/FastGeoStreaming/Config/DefaultFastGeoStreaming.ini
+++ b/Engine/Plugins/Experimental/FastGeoStreaming/Config/DefaultFastGeoStreaming.ini
@@ -3,10 +3,16 @@
 +BuiltinAllowedActorClasses=/Script/CoreUObject.Class'/Script/Engine.StaticMeshActor'
 +BuiltinAllowedActorClasses=/Script/CoreUObject.Class'/Script/Engine.WorldPartitionHLOD'
 +BuiltinAllowedActorClasses=/Script/CoreUObject.Class'/Script/Engine.WorldPartitionAutoInstancedActor'
++BuiltinAllowedActorClasses=/Script/CoreUObject.Class'/Script/Engine.PointLight'
++BuiltinAllowedActorClasses=/Script/CoreUObject.Class'/Script/Engine.SpotLight'
++BuiltinAllowedActorClasses=/Script/CoreUObject.Class'/Script/Engine.RectLight'
 +BuiltinDisallowedActorClasses=/Script/CoreUObject.Class'/Script/Water.WaterBody'
 +BuiltinDisallowedActorClasses=/Script/CoreUObject.Class'/Script/Engine.Info'
 +BuiltinAllowedComponentClasses=/Script/CoreUObject.Class'/Script/Engine.StaticMeshComponent'
 +BuiltinAllowedComponentClasses=/Script/CoreUObject.Class'/Script/Engine.DecalComponent'
++BuiltinAllowedComponentClasses=/Script/CoreUObject.Class'/Script/Engine.PointLightComponent'
++BuiltinAllowedComponentClasses=/Script/CoreUObject.Class'/Script/Engine.SpotLightComponent'
++BuiltinAllowedComponentClasses=/Script/CoreUObject.Class'/Script/Engine.RectLightComponent'
 +BuiltinDisallowedComponentClasses=/Script/CoreUObject.Class'/Script/Foliage.InteractiveFoliageComponent'
 +BuiltinDisallowedComponentClasses=/Script/CoreUObject.Class'/Script/Landscape.LandscapeNaniteComponent'
 +BuiltinDisallowedComponentClasses=/Script/CoreUObject.Class'/Script/Landscape.LandscapeMeshProxyComponent'
diff --git a/Engine/Plugins/Experimental/FastGeoStreaming/Source/FastGeoStreaming/Internal/FastGeoComponent.h b/Engine/Plugins/Experimental/FastGeoStreaming/Source/FastGeoStreaming/Internal/FastGeoComponent.h
--- a/Engine/Plugins/Experimental/FastGeoStreaming/Source/FastGeoStreaming/Internal/FastGeoComponent.h
+++ b/Engine/Plugins/Experimental/FastGeoStreaming/Source/FastGeoStreaming/Internal/FastGeoComponent.h
@@ -8,6 +8,7 @@
 #include "FastGeoElementType.h"
 #include "Physics/Experimental/AsyncPhysicsStateProcessorInterface.h"
 #include "Components/ComponentInterfaces.h"
+#include "Components/SceneComponent.h"
 #include "Async/TaskGraphInterfaces.h"
 
 class UWorld;
@@ -52,7 +53,8 @@ public:
 	virtual bool IsCollisionEnabled() const { return false; }
 	virtual void Serialize(FArchive& Ar);
 	virtual void InitializeDynamicProperties() {}
-	virtual void ApplyWorldTransform(const FTransform& InTransform) {}
+	virtual void ApplyWorldTransform(const FTransform& InTransform);
+	virtual bool ShouldComponentAddToScene() const;
 	virtual bool ShouldCreateRenderState() const { return false; }
 	virtual void CreateRenderState(FRegisterComponentContext* Context) {}
 	virtual void DestroyRenderState(FFastGeoDestroyRenderStateContext* Context) {}
@@ -77,10 +79,14 @@ public:
 
 protected:
 	// Persistent Data
-	int32 ComponentIndex;
+	int32 ComponentIndex = INDEX_NONE;
+	FTransform LocalTransform{ FTransform::Identity };
+	FTransform WorldTransform{ FTransform::Identity };
+	TEnumAsByte<enum EDetailMode> DetailMode = EDetailMode::DM_Low;
 
 	// Transient Data
-	FFastGeoComponentCluster* Owner;
+	FFastGeoComponentCluster* Owner = nullptr;
+	TUniquePtr<FRWLock> Lock;
 
 	enum EPhysicsStateCreation
 	{
diff --git a/Engine/Plugins/Experimental/FastGeoStreaming/Source/FastGeoStreaming/Internal/FastGeoComponentCluster.h b/Engine/Plugins/Experimental/FastGeoStreaming/Source/FastGeoStreaming/Internal/FastGeoComponentCluster.h
--- a/Engine/Plugins/Experimental/FastGeoStreaming/Source/FastGeoStreaming/Internal/FastGeoComponentCluster.h
+++ b/Engine/Plugins/Experimental/FastGeoStreaming/Source/FastGeoStreaming/Internal/FastGeoComponentCluster.h
@@ -10,6 +10,9 @@
 #include "FastGeoInstancedSkinnedMeshComponent.h"
 #include "FastGeoProceduralISMComponent.h"
 #include "FastGeoDecalComponent.h"
+#include "FastGeoRectLightComponent.h"
+#include "FastGeoSpotLightComponent.h"
+#include "FastGeoPointLightComponent.h"
 
 class ULevel;
 class UFastGeoContainer;
@@ -90,6 +93,9 @@ protected:
 	TArray<FFastGeoInstancedSkinnedMeshComponent> InstancedSkinnedMeshComponents;
 	TArray<FFastGeoProceduralISMComponent> ProceduralISMComponents;
 	TArray<FFastGeoDecalComponent> DecalComponents;
+	TArray<FFastGeoRectLightComponent> RectLightComponents;
+	TArray<FFastGeoSpotLightComponent> SpotLightComponents;
+	TArray<FFastGeoPointLightComponent> PointLightComponents;
 
 	template <typename ThisType, typename TComponent, typename TFunc>
 	static void ForEachComponent(ThisType* Self, TFunc&& InFunc)
@@ -122,6 +128,9 @@ protected:
 		ForEachArray(Self->InstancedSkinnedMeshComponents);
 		ForEachArray(Self->ProceduralISMComponents);
 		ForEachArray(Self->DecalComponents);
+		ForEachArray(Self->RectLightComponents);
+		ForEachArray(Self->SpotLightComponents);
+		ForEachArray(Self->PointLightComponents);
 	}
 
 	template <typename ThisType, typename TComponent, typename TFunc>
@@ -161,8 +170,11 @@ protected:
 			   ForEachArray(Self->InstancedStaticMeshComponents) &&
 			   ForEachArray(Self->SkinnedMeshComponents) &&
 			   ForEachArray(Self->InstancedSkinnedMeshComponents) &&
-			   ForEachArray(Self->ProceduralISMComponents) && 
-			   ForEachArray(Self->DecalComponents);
+			   ForEachArray(Self->ProceduralISMComponents) &&
+			   ForEachArray(Self->DecalComponents) &&
+			   ForEachArray(Self->RectLightComponents) &&
+			   ForEachArray(Self->SpotLightComponents) &&
+			   ForEachArray(Self->PointLightComponents);
 	}
 
 private:
diff --git a/Engine/Plugins/Experimental/FastGeoStreaming/Source/FastGeoStreaming/Internal/FastGeoDecalComponent.h b/Engine/Plugins/Experimental/FastGeoStreaming/Source/FastGeoStreaming/Internal/FastGeoDecalComponent.h
--- a/Engine/Plugins/Experimental/FastGeoStreaming/Source/FastGeoStreaming/Internal/FastGeoDecalComponent.h
+++ b/Engine/Plugins/Experimental/FastGeoStreaming/Source/FastGeoStreaming/Internal/FastGeoDecalComponent.h
@@ -46,13 +46,10 @@ protected:
 
 private:
 	// Persistent data
-	FTransform LocalTransform;
-	FTransform WorldTransform;
 	FVector DecalSize;
 	uint8 bIsVisible : 1 = true;
 	FDeferredDecalSceneProxyDesc SceneProxyDesc{};
 	FDeferredDecalProxy* SceneProxy = nullptr;
-	TUniquePtr<FRWLock> Lock;
 	FTimerHandle TimerHandle_DestroyDecalComponent;
 
 #if !WITH_EDITOR && UE_WITH_PSO_PRECACHING
diff --git a/Engine/Plugins/Experimental/FastGeoStreaming/Source/FastGeoStreaming/Internal/FastGeoLocalLightComponent.h b/Engine/Plugins/Experimental/FastGeoStreaming/Source/FastGeoStreaming/Internal/FastGeoLocalLightComponent.h
new file mode 100644
--- /dev/null
+++ b/Engine/Plugins/Experimental/FastGeoStreaming/Source/FastGeoStreaming/Internal/FastGeoLocalLightComponent.h
@@ -0,0 +1,48 @@
+// Copyright Epic Games, Inc. All Rights Reserved.
+
+#pragma once
+
+#include "FastGeoComponent.h"
+#include "LocalLightSceneProxyDesc.h"
+
+class FLightSceneProxy;
+
+class FASTGEOSTREAMING_API FFastGeoLocalLightComponent : public FFastGeoComponent
+{
+public:
+	typedef FFastGeoComponent Super;
+
+	/** Static type identifier for this element class */
+	static const FFastGeoElementType Type;
+
+	FFastGeoLocalLightComponent(int32 InComponentIndex = INDEX_NONE, FFastGeoElementType InType = Type);
+	virtual ~FFastGeoLocalLightComponent() = default;
+
+	//~ Begin FFastGeoComponent interface
+	virtual void Serialize(FArchive& Ar) override;
+	virtual bool ShouldCreateRenderState() const override;
+	virtual void CreateRenderState(FRegisterComponentContext* Context) override;
+	virtual void DestroyRenderState(FFastGeoDestroyRenderStateContext* Context) override;
+	virtual void UpdateVisibility() override;
+#if WITH_EDITOR
+	virtual void InitializeFromComponent(UActorComponent* Component) override;
+#endif
+	//~ End FFastGeoComponent interface
+
+protected:
+	virtual FLocalLightSceneProxyDesc& GetLocalLightSceneProxyDesc() = 0;
+	virtual const FLocalLightSceneProxyDesc& GetLocalLightSceneProxyDesc() const = 0;
+	virtual FLightSceneProxy* CreateSceneProxy() = 0;
+
+	void InitializeSceneProxyDescDynamicProperties();
+	bool UpdateVisibilityInternal();
+
+private:
+	// Persistent data
+	bool bAffectsWorld = true;
+	bool bComponentIsVisible = true;
+
+	// Transient data
+	bool bIsVisible = true;
+	FLightSceneProxy* SceneProxy = nullptr;
+};
\ No newline at end of file
diff --git a/Engine/Plugins/Experimental/FastGeoStreaming/Source/FastGeoStreaming/Internal/FastGeoPointLightComponent.h b/Engine/Plugins/Experimental/FastGeoStreaming/Source/FastGeoStreaming/Internal/FastGeoPointLightComponent.h
new file mode 100644
--- /dev/null
+++ b/Engine/Plugins/Experimental/FastGeoStreaming/Source/FastGeoStreaming/Internal/FastGeoPointLightComponent.h
@@ -0,0 +1,61 @@
+// Copyright Epic Games, Inc. All Rights Reserved.
+
+#pragma once
+
+#include "FastGeoLocalLightComponent.h"
+#include "PointLightSceneProxyDesc.h"
+
+class FLightSceneProxy;
+
+class FASTGEOSTREAMING_API FFastGeoPointLightComponentBase : public FFastGeoLocalLightComponent
+{
+public:
+	typedef FFastGeoLocalLightComponent Super;
+
+	/** Static type identifier for this element class */
+	static const FFastGeoElementType Type;
+
+	FFastGeoPointLightComponentBase(int32 InComponentIndex = INDEX_NONE, FFastGeoElementType InType = Type);
+	virtual ~FFastGeoPointLightComponentBase() = default;
+
+protected:
+	//~ Begin FFastGeoLocalLightComponent interface
+	virtual FLocalLightSceneProxyDesc& GetLocalLightSceneProxyDesc() override { return GetPointLightSceneProxyDesc(); }
+	virtual const FLocalLightSceneProxyDesc& GetLocalLightSceneProxyDesc() const override { return GetPointLightSceneProxyDesc(); }
+	virtual FLightSceneProxy* CreateSceneProxy() override;
+	//~ End FFastGeoLocalLightComponent interface
+
+	virtual FPointLightSceneProxyDesc& GetPointLightSceneProxyDesc() = 0;
+	virtual const FPointLightSceneProxyDesc& GetPointLightSceneProxyDesc() const = 0;
+	virtual FLightSceneProxy* CreateTypedSceneProxy() = 0;
+
+	bool IsPointLightSupported() const;
+};
+
+class FASTGEOSTREAMING_API FFastGeoPointLightComponent : public FFastGeoPointLightComponentBase
+{
+public:
+	typedef FFastGeoPointLightComponentBase Super;
+
+	/** Static type identifier for this element class */
+	static const FFastGeoElementType Type;
+
+	FFastGeoPointLightComponent(int32 InComponentIndex = INDEX_NONE, FFastGeoElementType InType = Type);
+	virtual ~FFastGeoPointLightComponent() = default;
+
+	//~ Begin FFastGeoComponent interface
+#if WITH_EDITOR
+	virtual void InitializeFromComponent(UActorComponent* Component) override;
+#endif
+	//~ End FFastGeoComponent interface
+
+protected:
+	//~ Begin FFastGeoPointLightComponentBase interface
+	virtual FPointLightSceneProxyDesc& GetPointLightSceneProxyDesc() override { return SceneProxyDesc; }
+	virtual const FPointLightSceneProxyDesc& GetPointLightSceneProxyDesc() const override { return SceneProxyDesc; }
+	virtual FLightSceneProxy* CreateTypedSceneProxy() override;
+	//~ End FFastGeoPointLightComponentBase interface
+
+private:
+	FPointLightSceneProxyDesc SceneProxyDesc{};
+};
\ No newline at end of file
diff --git a/Engine/Plugins/Experimental/FastGeoStreaming/Source/FastGeoStreaming/Internal/FastGeoPrimitiveComponent.h b/Engine/Plugins/Experimental/FastGeoStreaming/Source/FastGeoStreaming/Internal/FastGeoPrimitiveComponent.h
--- a/Engine/Plugins/Experimental/FastGeoStreaming/Source/FastGeoStreaming/Internal/FastGeoPrimitiveComponent.h
+++ b/Engine/Plugins/Experimental/FastGeoStreaming/Source/FastGeoStreaming/Internal/FastGeoPrimitiveComponent.h
@@ -141,7 +141,6 @@ public:
 	virtual void OnAsyncDestroyPhysicsStateEnd_GameThread() override;
 	virtual void OnAsyncDestroyPhysicsState() override;
 	virtual bool IsCollisionEnabled() const override;
-	virtual void ApplyWorldTransform(const FTransform& InTransform) override;
 	virtual void DestroyRenderState(FFastGeoDestroyRenderStateContext* Context) override;
 	virtual void MarkPrecachePSOsRequired() override;
 #if WITH_EDITOR
@@ -188,8 +187,6 @@ protected:
 
 protected:
 	// Persistent data
-	FTransform LocalTransform;
-	FTransform WorldTransform;
 	FBoxSphereBounds LocalBounds;
 	FBoxSphereBounds WorldBounds;
 	uint8 bIsVisible : 1 = true;
@@ -198,7 +195,6 @@ protected:
 	uint8 bRasterizeAsFilledConvexVolume : 1 = false;
 	uint8 bCanEverAffectNavigation : 1 = false;
 	FCustomPrimitiveData CustomPrimitiveData;
-	TEnumAsByte<enum EDetailMode> DetailMode = EDetailMode::DM_Low;
 	TEnumAsByte<EHasCustomNavigableGeometry::Type> bHasCustomNavigableGeometry = EHasCustomNavigableGeometry::No;
 	FBodyInstance BodyInstance;
 	TArray<URuntimeVirtualTexture*> RuntimeVirtualTextures;
@@ -234,8 +230,6 @@ protected:
 	static_assert((int)EPSOPrecachePriority::Highest < 1 << 2);
 #endif
 
-	TUniquePtr<FRWLock> Lock;
-
 	friend class FFastGeoComponentCluster;
 	friend class FPrimitiveComponentHelper;
 	friend class FFastGeoPhysicsBodyInstanceOwner;
diff --git a/Engine/Plugins/Experimental/FastGeoStreaming/Source/FastGeoStreaming/Internal/FastGeoRectLightComponent.h b/Engine/Plugins/Experimental/FastGeoStreaming/Source/FastGeoStreaming/Internal/FastGeoRectLightComponent.h
new file mode 100644
--- /dev/null
+++ b/Engine/Plugins/Experimental/FastGeoStreaming/Source/FastGeoStreaming/Internal/FastGeoRectLightComponent.h
@@ -0,0 +1,34 @@
+// Copyright Epic Games, Inc. All Rights Reserved.
+
+#pragma once
+
+#include "FastGeoLocalLightComponent.h"
+#include "RectLightSceneProxyDesc.h"
+
+class FASTGEOSTREAMING_API FFastGeoRectLightComponent : public FFastGeoLocalLightComponent
+{
+public:
+	typedef FFastGeoLocalLightComponent Super;
+
+	/** Static type identifier for this element class */
+	static const FFastGeoElementType Type;
+
+	FFastGeoRectLightComponent(int32 InComponentIndex = INDEX_NONE, FFastGeoElementType InType = Type);
+	virtual ~FFastGeoRectLightComponent() = default;
+
+	//~ Begin FFastGeoComponent interface
+#if WITH_EDITOR
+	virtual void InitializeFromComponent(UActorComponent* Component) override;
+#endif
+	//~ End FFastGeoComponent interface
+
+protected:
+	//~ Begin FFastGeoLocalLightComponent interface
+	virtual FLocalLightSceneProxyDesc& GetLocalLightSceneProxyDesc() override { return SceneProxyDesc; }
+	virtual const FLocalLightSceneProxyDesc& GetLocalLightSceneProxyDesc() const override { return SceneProxyDesc; }
+	virtual FLightSceneProxy* CreateSceneProxy() override;
+	//~ End FFastGeoLocalLightComponent interface
+
+private:
+	FRectLightSceneProxyDesc SceneProxyDesc{};
+};
diff --git a/Engine/Plugins/Experimental/FastGeoStreaming/Source/FastGeoStreaming/Internal/FastGeoSpotLightComponent.h b/Engine/Plugins/Experimental/FastGeoStreaming/Source/FastGeoStreaming/Internal/FastGeoSpotLightComponent.h
new file mode 100644
--- /dev/null
+++ b/Engine/Plugins/Experimental/FastGeoStreaming/Source/FastGeoStreaming/Internal/FastGeoSpotLightComponent.h
@@ -0,0 +1,34 @@
+// Copyright Epic Games, Inc. All Rights Reserved.
+
+#pragma once
+
+#include "FastGeoPointLightComponent.h"
+#include "SpotLightSceneProxyDesc.h"
+
+class FASTGEOSTREAMING_API FFastGeoSpotLightComponent : public FFastGeoPointLightComponentBase
+{
+public:
+	typedef FFastGeoPointLightComponentBase Super;
+
+	/** Static type identifier for this element class */
+	static const FFastGeoElementType Type;
+
+	FFastGeoSpotLightComponent(int32 InComponentIndex = INDEX_NONE, FFastGeoElementType InType = Type);
+	virtual ~FFastGeoSpotLightComponent() = default;
+
+	//~ Begin FFastGeoComponent interface
+#if WITH_EDITOR
+	virtual void InitializeFromComponent(UActorComponent* Component) override;
+#endif
+	//~ End FFastGeoComponent interface
+
+protected:
+	//~ Begin FFastGeoPointLightComponentBase interface
+	virtual FPointLightSceneProxyDesc& GetPointLightSceneProxyDesc() override { return SceneProxyDesc; }
+	virtual const FPointLightSceneProxyDesc& GetPointLightSceneProxyDesc() const override { return SceneProxyDesc; }
+	virtual FLightSceneProxy* CreateTypedSceneProxy() override;
+	//~ End FFastGeoPointLightComponentBase interface
+
+private:
+	FSpotLightSceneProxyDesc SceneProxyDesc{};
+};
diff --git a/Engine/Plugins/Experimental/FastGeoStreaming/Source/FastGeoStreaming/Private/FastGeoComponent.cpp b/Engine/Plugins/Experimental/FastGeoStreaming/Source/FastGeoStreaming/Private/FastGeoComponent.cpp
--- a/Engine/Plugins/Experimental/FastGeoStreaming/Source/FastGeoStreaming/Private/FastGeoComponent.cpp
+++ b/Engine/Plugins/Experimental/FastGeoStreaming/Source/FastGeoStreaming/Private/FastGeoComponent.cpp
@@ -7,6 +7,7 @@
 #include "Engine/Level.h"
 #include "Engine/World.h"
 #include "SceneInterface.h"
+#include "UnrealEngine.h"
 
 const FFastGeoElementType FFastGeoComponent::Type(&IFastGeoElement::Type);
 
@@ -14,13 +15,18 @@ FFastGeoComponent::FFastGeoComponent(int32 InComponentIndex, FFastGeoElementType
 	: Super(InType)
 	, ComponentIndex(InComponentIndex)
 	, Owner(nullptr)
+	, Lock(MakeUnique<FRWLock>())
 {
 }
 
 FFastGeoComponent::FFastGeoComponent(const FFastGeoComponent& Other)
 : Super(Other)
 , ComponentIndex(Other.ComponentIndex)
+, LocalTransform(Other.LocalTransform)
+, WorldTransform(Other.WorldTransform)
+, DetailMode(Other.DetailMode)
 , Owner(Other.Owner)
+, Lock(MakeUnique<FRWLock>())
 , PhysicsStateCreation(EPhysicsStateCreation::NotCreated)
 , ProxyState(EProxyCreationState::None)
 , bRenderStateDirty(false)
@@ -30,6 +36,11 @@ FFastGeoComponent::FFastGeoComponent(const FFastGeoComponent& Other)
 #if WITH_EDITOR
 void FFastGeoComponent::InitializeFromComponent(UActorComponent* InComponent)
 {
+	USceneComponent* SceneComponent = CastChecked<USceneComponent>(InComponent);
+	SceneComponent->UpdateComponentToWorld();
+	LocalTransform = SceneComponent->GetComponentToWorld();
+	WorldTransform = LocalTransform;
+	DetailMode = SceneComponent->DetailMode;
 }
 #endif
 
@@ -86,6 +97,15 @@ FLinearColor FFastGeoComponent::GetDebugColor() const
 void FFastGeoComponent::Serialize(FArchive& Ar)
 {
 	Ar << ComponentIndex;
+	Ar << LocalTransform;
+	Ar << WorldTransform;
+	Ar << DetailMode;
+}
+
+void FFastGeoComponent::ApplyWorldTransform(const FTransform& InTransform)
+{
+	check(!GetOwnerContainer()->IsRegistered());
+	WorldTransform = LocalTransform * InTransform;
 }
 
 void FFastGeoComponent::OnAsyncCreatePhysicsState()
@@ -142,4 +162,10 @@ void FFastGeoComponent::MarkRenderStateDirty()
 		UFastGeoWorldSubsystem* WorldSubsystem = GetWorld()->GetSubsystemChecked<UFastGeoWorldSubsystem>();
 		WorldSubsystem->AddToComponentsPendingRecreate(this);
 	}
+}
+
+bool FFastGeoComponent::ShouldComponentAddToScene() const
+{
+	// If the detail mode setting allows it, add it to the scene.
+	return DetailMode <= GetCachedScalabilityCVars().DetailMode;
 }
\ No newline at end of file
diff --git a/Engine/Plugins/Experimental/FastGeoStreaming/Source/FastGeoStreaming/Private/FastGeoComponentCluster.cpp b/Engine/Plugins/Experimental/FastGeoStreaming/Source/FastGeoStreaming/Private/FastGeoComponentCluster.cpp
--- a/Engine/Plugins/Experimental/FastGeoStreaming/Source/FastGeoStreaming/Private/FastGeoComponentCluster.cpp
+++ b/Engine/Plugins/Experimental/FastGeoStreaming/Source/FastGeoStreaming/Private/FastGeoComponentCluster.cpp
@@ -28,6 +28,9 @@ FFastGeoComponentCluster::FFastGeoComponentCluster(const FFastGeoComponentCluste
 	InstancedSkinnedMeshComponents = InOther.InstancedSkinnedMeshComponents;
 	ProceduralISMComponents = InOther.ProceduralISMComponents;
 	DecalComponents = InOther.DecalComponents;
+	RectLightComponents = InOther.RectLightComponents;
+	SpotLightComponents = InOther.SpotLightComponents;
+	PointLightComponents = InOther.PointLightComponents;
 }
 
 FArchive& operator<<(FArchive& Ar, FFastGeoComponentCluster& ComponentCluster)
@@ -46,6 +49,9 @@ void FFastGeoComponentCluster::Serialize(FArchive& Ar)
 	Ar << InstancedSkinnedMeshComponents;
 	Ar << ProceduralISMComponents;
 	Ar << DecalComponents;
+	Ar << RectLightComponents;
+	Ar << SpotLightComponents;
+	Ar << PointLightComponents;
 }
 
 void FFastGeoComponentCluster::InitializeDynamicProperties()
@@ -99,6 +105,18 @@ FFastGeoComponent* FFastGeoComponentCluster::GetComponent(uint32 InComponentType
 	{
 		return DecalComponents.IsValidIndex(InComponentIndex) ? &DecalComponents[InComponentIndex] : nullptr;
 	}
+	else if (FFastGeoRectLightComponent::Type.IsSameTypeID(InComponentTypeID))
+	{
+		return RectLightComponents.IsValidIndex(InComponentIndex) ? &RectLightComponents[InComponentIndex] : nullptr;
+	}
+	else if (FFastGeoSpotLightComponent::Type.IsSameTypeID(InComponentTypeID))
+	{
+		return SpotLightComponents.IsValidIndex(InComponentIndex) ? &SpotLightComponents[InComponentIndex] : nullptr;
+	}
+	else if (FFastGeoPointLightComponent::Type.IsSameTypeID(InComponentTypeID))
+	{
+		return PointLightComponents.IsValidIndex(InComponentIndex) ? &PointLightComponents[InComponentIndex] : nullptr;
+	}
 	check(false);
 	return nullptr;
 }
@@ -254,6 +272,18 @@ FFastGeoComponent& FFastGeoComponentCluster::AddComponent(FFastGeoElementType In
 	{
 		NewComponent = &DecalComponents.Emplace_GetRef(DecalComponents.Num());
 	}
+	else if (InComponentType == FFastGeoRectLightComponent::Type)
+	{
+		NewComponent = &RectLightComponents.Emplace_GetRef(RectLightComponents.Num());
+	}
+	else if (InComponentType == FFastGeoSpotLightComponent::Type)
+	{
+		NewComponent = &SpotLightComponents.Emplace_GetRef(SpotLightComponents.Num());
+	}
+	else if (InComponentType == FFastGeoPointLightComponent::Type)
+	{
+		NewComponent = &PointLightComponents.Emplace_GetRef(PointLightComponents.Num());
+	}
 	else
 	{
 		unimplemented(); // Should never be reached, missing type handling
@@ -274,5 +304,8 @@ bool FFastGeoComponentCluster::HasComponents() const
 		   InstancedSkinnedMeshComponents.Num() > 0 || 
 		   SkinnedMeshComponents.Num() > 0 || 
 		   ProceduralISMComponents.Num() > 0 || 
-		   DecalComponents.Num() > 0;
+		   DecalComponents.Num() > 0 ||
+		   RectLightComponents.Num() > 0 ||
+		   SpotLightComponents.Num() > 0 ||
+		   PointLightComponents.Num() > 0;
 }
diff --git a/Engine/Plugins/Experimental/FastGeoStreaming/Source/FastGeoStreaming/Private/FastGeoContainer.cpp b/Engine/Plugins/Experimental/FastGeoStreaming/Source/FastGeoStreaming/Private/FastGeoContainer.cpp
--- a/Engine/Plugins/Experimental/FastGeoStreaming/Source/FastGeoStreaming/Private/FastGeoContainer.cpp
+++ b/Engine/Plugins/Experimental/FastGeoStreaming/Source/FastGeoStreaming/Private/FastGeoContainer.cpp
@@ -13,6 +13,7 @@
 #include "Engine/Level.h"
 #include "Engine/SkeletalMesh.h"
 #include "Engine/Texture.h"
+#include "Engine/TextureLightProfile.h"
 #include "Engine/World.h"
 #include "Misc/Timeout.h"
 #include "NavigationSystem.h"
@@ -683,6 +684,7 @@ public:
 	FFastGeoGatherFastGeoContainerAssetRefsArchive(UFastGeoContainer& Container)
 	{
 		KnownClasses.Add(UTexture::StaticClass());
+		KnownClasses.Add(UTextureLightProfile::StaticClass());
 		KnownClasses.Add(URuntimeVirtualTexture::StaticClass());
 		KnownClasses.Add(UStaticMesh::StaticClass());
 		KnownClasses.Add(UMaterialInterface::StaticClass());
diff --git a/Engine/Plugins/Experimental/FastGeoStreaming/Source/FastGeoStreaming/Private/FastGeoDecalComponent.cpp b/Engine/Plugins/Experimental/FastGeoStreaming/Source/FastGeoStreaming/Private/FastGeoDecalComponent.cpp
--- a/Engine/Plugins/Experimental/FastGeoStreaming/Source/FastGeoStreaming/Private/FastGeoDecalComponent.cpp
+++ b/Engine/Plugins/Experimental/FastGeoStreaming/Source/FastGeoStreaming/Private/FastGeoDecalComponent.cpp
@@ -20,19 +20,15 @@ FFastGeoDecalComponent::FFastGeoDecalComponent(int32 InComponentIndex, FFastGeoE
 	, DecalSize(FVector::ZeroVector)
 	, bIsVisible(true)
 	, SceneProxy(nullptr)
-	, Lock(MakeUnique<FRWLock>())
 {
 }
 
 FFastGeoDecalComponent::FFastGeoDecalComponent(const FFastGeoDecalComponent& Other)
 	: Super(Other)
-	, LocalTransform(Other.LocalTransform)
-	, WorldTransform(Other.WorldTransform)
 	, DecalSize(Other.DecalSize)
 	, bIsVisible(Other.bIsVisible)
 	, SceneProxyDesc(Other.SceneProxyDesc)
 	, SceneProxy(nullptr)
-	, Lock(MakeUnique<FRWLock>())
 #if !WITH_EDITOR && UE_WITH_PSO_PRECACHING
 	, LatestPSOPrecacheJobSetCompleted(0)
 	, LatestPSOPrecacheJobSet(0)
@@ -46,10 +42,7 @@ void FFastGeoDecalComponent::InitializeFromComponent(UActorComponent* Component)
 	Super::InitializeFromComponent(Component);
 
 	UDecalComponent* DecalComponent = CastChecked<UDecalComponent>(Component);
-	DecalComponent->UpdateComponentToWorld();
 	SceneProxyDesc = FDeferredDecalSceneProxyDesc(DecalComponent);
-	LocalTransform = DecalComponent->GetComponentToWorld();
-	WorldTransform = LocalTransform;
 	DecalSize = DecalComponent->DecalSize;
 	bIsVisible = SceneProxyDesc.bDrawInGame;
 	check(SceneProxyDesc.Bounds == CalcBounds());
@@ -60,8 +53,6 @@ void FFastGeoDecalComponent::Serialize(FArchive& Ar)
 {
 	Super::Serialize(Ar);
 
-	Ar << LocalTransform;
-	Ar << WorldTransform;
 	Ar << DecalSize;
 	FArchive_Serialize_BitfieldBool(Ar, bIsVisible);
 
@@ -84,8 +75,8 @@ void FFastGeoDecalComponent::Serialize(FArchive& Ar)
 
 void FFastGeoDecalComponent::ApplyWorldTransform(const FTransform& InTransform)
 {
-	check(!GetOwnerContainer()->IsRegistered());
-	WorldTransform = LocalTransform * InTransform;
+	Super::ApplyWorldTransform(InTransform);
+
 	SceneProxyDesc.Bounds = CalcBounds();
 	SceneProxyDesc.TransformWithDecalScale = GetTransformIncludingDecalSize();
 }
@@ -106,7 +97,7 @@ void FFastGeoDecalComponent::UpdateVisibilityInternal()
 
 bool FFastGeoDecalComponent::ShouldCreateRenderState() const
 {
-	return FApp::CanEverRender() && SceneProxyDesc.bDrawInGame;
+	return FApp::CanEverRender() && ShouldComponentAddToScene() && SceneProxyDesc.bDrawInGame;
 }
 
 void FFastGeoDecalComponent::CreateRenderState(FRegisterComponentContext* Context)
diff --git a/Engine/Plugins/Experimental/FastGeoStreaming/Source/FastGeoStreaming/Private/FastGeoLocalLightComponent.cpp b/Engine/Plugins/Experimental/FastGeoStreaming/Source/FastGeoStreaming/Private/FastGeoLocalLightComponent.cpp
new file mode 100644
--- /dev/null
+++ b/Engine/Plugins/Experimental/FastGeoStreaming/Source/FastGeoStreaming/Private/FastGeoLocalLightComponent.cpp
@@ -0,0 +1,144 @@
+// Copyright Epic Games, Inc. All Rights Reserved.
+
+#include "FastGeoLocalLightComponent.h"
+#include "FastGeoWorldSubsystem.h"
+#include "Components/LocalLightComponent.h"
+#include "PointLightSceneProxy.h"
+#include "PointLightSceneProxyDesc.h"
+#include "LightComponentId.h"
+#include "LightSceneDesc.h"
+#include "RenderUtils.h"
+#include "SceneInterface.h"
+
+const FFastGeoElementType FFastGeoLocalLightComponent::Type(&FFastGeoComponent::Type);
+
+FFastGeoLocalLightComponent::FFastGeoLocalLightComponent(int32 InComponentIndex, FFastGeoElementType InType)
+	: Super(InComponentIndex, InType)
+	, bAffectsWorld(true)
+	, bComponentIsVisible(true)
+	, bIsVisible(true)
+	, SceneProxy(nullptr)
+{
+}
+
+#if WITH_EDITOR
+void FFastGeoLocalLightComponent::InitializeFromComponent(UActorComponent* Component)
+{
+	Super::InitializeFromComponent(Component);
+
+	ULocalLightComponent* LocalLightComponent = CastChecked<ULocalLightComponent>(Component);
+	bAffectsWorld = LocalLightComponent->bAffectsWorld;
+	bComponentIsVisible = LocalLightComponent->IsVisible() && LocalLightComponent->GetOwner() && !LocalLightComponent->GetOwner()->IsHidden();
+}
+#endif
+
+void FFastGeoLocalLightComponent::Serialize(FArchive& Ar)
+{
+	Super::Serialize(Ar);
+
+	// Serialize persistent data
+	Ar << bAffectsWorld;
+	Ar << bComponentIsVisible;
+
+	// Serialize persistent data from FLocalLightSceneProxyDesc
+	FLocalLightSceneProxyDesc& SceneProxyDesc = GetLocalLightSceneProxyDesc();
+	SceneProxyDesc.Serialize(Ar);
+
+	if (Ar.IsPersistent() && Ar.IsLoading())
+	{
+		// Initialize LightComponentId as it's not serialized
+		SceneProxyDesc.LightComponentId = FLightComponentId::GetNextId();
+	}
+}
+
+bool FFastGeoLocalLightComponent::ShouldCreateRenderState() const
+{
+	if (!FApp::CanEverRender())
+	{
+		return false;
+	}
+
+	// If the detail mode setting allows it, add it to the scene.
+	if (!ShouldComponentAddToScene())
+	{
+		return false;
+	}
+
+	return bAffectsWorld && bComponentIsVisible;
+}
+
+void FFastGeoLocalLightComponent::CreateRenderState(FRegisterComponentContext* Context)
+{
+	FWriteScopeLock WriteLock(*Lock.Get());
+	check(ProxyState == EProxyCreationState::None);
+	ProxyState = EProxyCreationState::Creating;
+	bRenderStateDirty = false;
+
+	InitializeSceneProxyDescDynamicProperties();
+	check(!SceneProxy);
+	SceneProxy = CreateSceneProxy();
+	if (SceneProxy)
+	{
+		FLightSceneDesc LightSceneDesc{ .Transform = WorldTransform, .SceneProxy = SceneProxy};
+		FSceneInterface* Scene = GetScene();
+		check(Scene);
+		if (bIsVisible)
+		{
+			Scene->BatchAddLights(MakeArrayView(&LightSceneDesc, 1));
+		}
+		else
+		{
+			// This should be revisited to avoid needing to create an invisible light.
+			// Note that this is currently only possible when FastGeo.Show 0 is set.
+			Scene->BatchAddInvisibleLights(MakeArrayView(&LightSceneDesc, 1));
+		}
+
+		check(LightSceneDesc.SceneProxy);
+		ProxyState = EProxyCreationState::Created;
+	}
+	else
+	{
+		ProxyState = EProxyCreationState::None;
+	}
+}
+
+void FFastGeoLocalLightComponent::DestroyRenderState(FFastGeoDestroyRenderStateContext* Context)
+{
+	FWriteScopeLock WriteLock(*Lock.Get());
+	check(SceneProxy || ProxyState == EProxyCreationState::None);
+	if (SceneProxy)
+	{
+		check(ProxyState == EProxyCreationState::Created);
+		FLightSceneDesc LightSceneDesc{ .Transform = WorldTransform, .SceneProxy = SceneProxy };
+		FSceneInterface* Scene = GetScene();
+		check(Scene);
+		Scene->BatchRemoveLights(MakeArrayView(&LightSceneDesc, 1));
+		SceneProxy = nullptr;
+		ProxyState = EProxyCreationState::None;
+		bRenderStateDirty = false;
+	}
+}
+
+void FFastGeoLocalLightComponent::InitializeSceneProxyDescDynamicProperties()
+{
+	UpdateVisibilityInternal();
+}
+
+void FFastGeoLocalLightComponent::UpdateVisibility()
+{
+	if (UpdateVisibilityInternal())
+	{
+		MarkRenderStateDirty();
+	}
+}
+
+bool FFastGeoLocalLightComponent::UpdateVisibilityInternal()
+{
+	bool bOldIsVisible = bIsVisible;
+#if !UE_BUILD_SHIPPING
+	bIsVisible = UFastGeoWorldSubsystem::IsFastGeoVisible();
+#else
+	bIsVisible = true;
+#endif
+	return bIsVisible != bOldIsVisible;
+}
\ No newline at end of file
diff --git a/Engine/Plugins/Experimental/FastGeoStreaming/Source/FastGeoStreaming/Private/FastGeoPointLightComponent.cpp b/Engine/Plugins/Experimental/FastGeoStreaming/Source/FastGeoStreaming/Private/FastGeoPointLightComponent.cpp
new file mode 100644
--- /dev/null
+++ b/Engine/Plugins/Experimental/FastGeoStreaming/Source/FastGeoStreaming/Private/FastGeoPointLightComponent.cpp
@@ -0,0 +1,64 @@
+// Copyright Epic Games, Inc. All Rights Reserved.
+
+#include "FastGeoPointLightComponent.h"
+#include "Components/PointLightComponent.h"
+#include "PointLightSceneProxy.h"
+#include "PointLightSceneProxyDesc.h"
+#include "RenderUtils.h"
+#include "SceneInterface.h"
+
+const FFastGeoElementType FFastGeoPointLightComponentBase::Type(&FFastGeoLocalLightComponent::Type);
+
+FFastGeoPointLightComponentBase::FFastGeoPointLightComponentBase(int32 InComponentIndex, FFastGeoElementType InType)
+	: Super(InComponentIndex, InType)
+{
+}
+
+FLightSceneProxy* FFastGeoPointLightComponentBase::CreateSceneProxy()
+{
+	if (ShouldComponentAddToScene() && IsPointLightSupported())
+	{
+		return CreateTypedSceneProxy();
+	}
+	return nullptr;
+}
+
+bool FFastGeoPointLightComponentBase::IsPointLightSupported() const
+{
+	const FPointLightSceneProxyDesc& SceneProxyDesc = GetPointLightSceneProxyDesc();
+	if (GMaxRHIFeatureLevel <= ERHIFeatureLevel::ES3_1 && SceneProxyDesc.bMovable)
+	{
+		if (!IsMobileDeferredShadingEnabled(GMaxRHIShaderPlatform))
+		{
+			// if project does not support dynamic point lights on mobile do not add them to the renderer 
+			return MobileForwardEnableLocalLights(GMaxRHIShaderPlatform);
+		}
+	}
+	return true;
+}
+
+const FFastGeoElementType FFastGeoPointLightComponent::Type(&FFastGeoPointLightComponentBase::Type);
+
+FFastGeoPointLightComponent::FFastGeoPointLightComponent(int32 InComponentIndex, FFastGeoElementType InType)
+	: Super(InComponentIndex, InType)
+{
+}
+
+#if WITH_EDITOR
+void FFastGeoPointLightComponent::InitializeFromComponent(UActorComponent* Component)
+{
+	Super::InitializeFromComponent(Component);
+
+	UPointLightComponent* PointLightComponent = CastChecked<UPointLightComponent>(Component);
+	SceneProxyDesc = UPointLightComponent::BuildSceneProxyDesc(*PointLightComponent);
+	SceneProxyDesc.LightComponent = nullptr;
+}
+#endif
+
+FLightSceneProxy* FFastGeoPointLightComponent::CreateTypedSceneProxy()
+{
+	FSceneInterface* Scene = GetScene();
+	check(Scene);
+	SceneProxyDesc.SceneInterface = Scene;
+	return new FPointLightSceneProxy(SceneProxyDesc);
+}
\ No newline at end of file
diff --git a/Engine/Plugins/Experimental/FastGeoStreaming/Source/FastGeoStreaming/Private/FastGeoPrimitiveComponent.cpp b/Engine/Plugins/Experimental/FastGeoStreaming/Source/FastGeoStreaming/Private/FastGeoPrimitiveComponent.cpp
--- a/Engine/Plugins/Experimental/FastGeoStreaming/Source/FastGeoStreaming/Private/FastGeoPrimitiveComponent.cpp
+++ b/Engine/Plugins/Experimental/FastGeoStreaming/Source/FastGeoStreaming/Private/FastGeoPrimitiveComponent.cpp
@@ -35,14 +35,11 @@ FFastGeoPrimitiveComponent::FFastGeoPrimitiveComponent(int32 InComponentIndex, F
 	: Super(InComponentIndex, InType)
 	, LocalBounds(ForceInit)
 	, WorldBounds(ForceInit)
-	, Lock(MakeUnique<FRWLock>())
 {
 }
 
 FFastGeoPrimitiveComponent::FFastGeoPrimitiveComponent(const FFastGeoPrimitiveComponent& Other)
-	: FFastGeoComponent(Other)
-	, LocalTransform(Other.LocalTransform)
-	, WorldTransform(Other.WorldTransform)
+	: Super(Other)
 	, LocalBounds(Other.LocalBounds)
 	, WorldBounds(Other.WorldBounds)
 	, bIsVisible(Other.bIsVisible)
@@ -51,7 +48,6 @@ FFastGeoPrimitiveComponent::FFastGeoPrimitiveComponent(const FFastGeoPrimitiveCo
 	, bRasterizeAsFilledConvexVolume(Other.bRasterizeAsFilledConvexVolume)
 	, bCanEverAffectNavigation(Other.bCanEverAffectNavigation)
 	, CustomPrimitiveData(Other.CustomPrimitiveData)
-	, DetailMode(Other.DetailMode)
 	, bHasCustomNavigableGeometry(Other.bHasCustomNavigableGeometry)
 	, BodyInstance(Other.BodyInstance)
 	, RuntimeVirtualTextures(Other.RuntimeVirtualTextures)
@@ -66,7 +62,6 @@ FFastGeoPrimitiveComponent::FFastGeoPrimitiveComponent(const FFastGeoPrimitiveCo
 	, bPSOPrecacheRequired(false)
 	, PSOPrecacheRequestPriority(EPSOPrecachePriority::Medium)
 #endif
-	, Lock(MakeUnique<FRWLock>())
 {
 }
 
@@ -75,8 +70,6 @@ void FFastGeoPrimitiveComponent::Serialize(FArchive& Ar)
 	Super::Serialize(Ar);
 
 	// Serialize persistent data from FFastGeoPrimitiveComponent
-	Ar << LocalTransform;
-	Ar << WorldTransform;
 	Ar << LocalBounds;
 	Ar << WorldBounds;
 	FArchive_Serialize_BitfieldBool(Ar, bIsVisible);
@@ -85,7 +78,6 @@ void FFastGeoPrimitiveComponent::Serialize(FArchive& Ar)
 	FArchive_Serialize_BitfieldBool(Ar, bRasterizeAsFilledConvexVolume);
 	FArchive_Serialize_BitfieldBool(Ar, bCanEverAffectNavigation);
 	Ar << CustomPrimitiveData.Data;
-	Ar << DetailMode;
 	Ar << bHasCustomNavigableGeometry;
 	Ar << RuntimeVirtualTextures;
 	FBodyInstance::StaticStruct()->SerializeItem(Ar, &BodyInstance, nullptr);
@@ -164,12 +156,6 @@ void FFastGeoPrimitiveComponent::Serialize(FArchive& Ar)
 	Ar << SceneProxyDesc.VirtualTextureMainPassMaxDrawDistance;
 }
 
-void FFastGeoPrimitiveComponent::ApplyWorldTransform(const FTransform& InTransform)
-{
-	check(!GetOwnerContainer()->IsRegistered());
-	WorldTransform = LocalTransform * InTransform;
-}
-
 #if WITH_EDITOR
 
 void FFastGeoPrimitiveComponent::InitializeFromComponent(UActorComponent* Component)
@@ -178,16 +164,12 @@ void FFastGeoPrimitiveComponent::InitializeFromComponent(UActorComponent* Compon
 
 	// Initialize properties not handled by InitializeFromPrimitiveComponent
 	UPrimitiveComponent* PrimitiveComponent = CastChecked<UPrimitiveComponent>(Component);
-	PrimitiveComponent->UpdateComponentToWorld();
-	LocalTransform = PrimitiveComponent->GetComponentToWorld();
-	WorldTransform = LocalTransform;
 	bIsVisible = PrimitiveComponent->IsVisible();
 	bStaticWhenNotMoveable = PrimitiveComponent->GetStaticWhenNotMoveable();
 	bFillCollisionUnderneathForNavmesh = PrimitiveComponent->bFillCollisionUnderneathForNavmesh;
 	bRasterizeAsFilledConvexVolume = PrimitiveComponent->bRasterizeAsFilledConvexVolume;
 	bCanEverAffectNavigation = PrimitiveComponent->CanEverAffectNavigation();
 	CustomPrimitiveData = PrimitiveComponent->GetCustomPrimitiveData();
-	DetailMode = PrimitiveComponent->DetailMode;
 	bHasCustomNavigableGeometry = PrimitiveComponent->bHasCustomNavigableGeometry;
 	BodyInstance.CopyBodyInstancePropertiesFrom(&PrimitiveComponent->BodyInstance);
 	RuntimeVirtualTextures = PrimitiveComponent->GetRuntimeVirtualTextures();
@@ -502,8 +484,7 @@ bool FFastGeoPrimitiveComponent::ShouldCreateRenderState() const
 	}
 
 	// If the detail mode setting allows it, add it to the scene.
-	const bool bDetailModeAllowsRendering = DetailMode <= GetCachedScalabilityCVars().DetailMode;
-	if (!bDetailModeAllowsRendering)
+	if (!ShouldComponentAddToScene())
 	{
 		return false;
 	}
diff --git a/Engine/Plugins/Experimental/FastGeoStreaming/Source/FastGeoStreaming/Private/FastGeoRectLightComponent.cpp b/Engine/Plugins/Experimental/FastGeoStreaming/Source/FastGeoStreaming/Private/FastGeoRectLightComponent.cpp
new file mode 100644
--- /dev/null
+++ b/Engine/Plugins/Experimental/FastGeoStreaming/Source/FastGeoStreaming/Private/FastGeoRectLightComponent.cpp
@@ -0,0 +1,37 @@
+// Copyright Epic Games, Inc. All Rights Reserved.
+
+#include "FastGeoRectLightComponent.h"
+#include "Components/RectLightComponent.h"
+#include "RectLightSceneProxy.h"
+#include "RectLightSceneProxyDesc.h"
+#include "SceneInterface.h"
+
+const FFastGeoElementType FFastGeoRectLightComponent::Type(&FFastGeoLocalLightComponent::Type);
+
+FFastGeoRectLightComponent::FFastGeoRectLightComponent(int32 InComponentIndex, FFastGeoElementType InType)
+	: Super(InComponentIndex, InType)
+{
+}
+
+#if WITH_EDITOR
+void FFastGeoRectLightComponent::InitializeFromComponent(UActorComponent* Component)
+{
+	Super::InitializeFromComponent(Component);
+
+	URectLightComponent* RectLightComponent = CastChecked<URectLightComponent>(Component);
+	SceneProxyDesc = URectLightComponent::BuildSceneProxyDesc(*RectLightComponent);
+	SceneProxyDesc.LightComponent = nullptr;
+}
+#endif
+
+FLightSceneProxy* FFastGeoRectLightComponent::CreateSceneProxy()
+{
+	if (ShouldComponentAddToScene())
+	{
+		FSceneInterface* Scene = GetScene();
+		check(Scene);
+		SceneProxyDesc.SceneInterface = Scene;
+		return new FRectLightSceneProxy(SceneProxyDesc);
+	}
+	return nullptr;
+}
\ No newline at end of file
diff --git a/Engine/Plugins/Experimental/FastGeoStreaming/Source/FastGeoStreaming/Private/FastGeoSpotLightComponent.cpp b/Engine/Plugins/Experimental/FastGeoStreaming/Source/FastGeoStreaming/Private/FastGeoSpotLightComponent.cpp
new file mode 100644
--- /dev/null
+++ b/Engine/Plugins/Experimental/FastGeoStreaming/Source/FastGeoStreaming/Private/FastGeoSpotLightComponent.cpp
@@ -0,0 +1,32 @@
+// Copyright Epic Games, Inc. All Rights Reserved.
+
+#include "FastGeoSpotLightComponent.h"
+#include "Components/SpotLightComponent.h"
+#include "SpotLightSceneProxy.h"
+#include "SpotLightSceneProxyDesc.h"
+
+const FFastGeoElementType FFastGeoSpotLightComponent::Type(&FFastGeoPointLightComponentBase::Type);
+
+FFastGeoSpotLightComponent::FFastGeoSpotLightComponent(int32 InComponentIndex, FFastGeoElementType InType)
+	: Super(InComponentIndex, InType)
+{
+}
+
+#if WITH_EDITOR
+void FFastGeoSpotLightComponent::InitializeFromComponent(UActorComponent* Component)
+{
+	Super::InitializeFromComponent(Component);
+
+	USpotLightComponent* SpotLightComponent = CastChecked<USpotLightComponent>(Component);
+	SceneProxyDesc = USpotLightComponent::BuildSceneProxyDesc(*SpotLightComponent);
+	SceneProxyDesc.LightComponent = nullptr;
+}
+#endif
+
+FLightSceneProxy* FFastGeoSpotLightComponent::CreateTypedSceneProxy()
+{
+	FSceneInterface* Scene = GetScene();
+	check(Scene);
+	SceneProxyDesc.SceneInterface = Scene;
+	return new FSpotLightSceneProxy(SceneProxyDesc);
+}
\ No newline at end of file
diff --git a/Engine/Plugins/Experimental/FastGeoStreaming/Source/FastGeoStreaming/Private/FastGeoWorldPartitionRuntimeCellTransformer.cpp b/Engine/Plugins/Experimental/FastGeoStreaming/Source/FastGeoStreaming/Private/FastGeoWorldPartitionRuntimeCellTransformer.cpp
--- a/Engine/Plugins/Experimental/FastGeoStreaming/Source/FastGeoStreaming/Private/FastGeoWorldPartitionRuntimeCellTransformer.cpp
+++ b/Engine/Plugins/Experimental/FastGeoStreaming/Source/FastGeoStreaming/Private/FastGeoWorldPartitionRuntimeCellTransformer.cpp
@@ -22,6 +22,9 @@
 #include "Components/InstancedSkinnedMeshComponent.h"
 #include "Components/SplineMeshComponent.h"
 #include "Components/DecalComponent.h"
+#include "Components/PointLightComponent.h"
+#include "Components/SpotLightComponent.h"
+#include "Components/RectLightComponent.h"
 #include "GameFramework/WorldSettings.h"
 #include "GameFramework/Info.h"
 #include "WorldPartition/WorldPartition.h"
@@ -112,7 +115,10 @@ namespace FastGeo
 			{ UInstancedStaticMeshComponent::StaticClass(), FFastGeoInstancedStaticMeshComponent::Type },
 			{ USkinnedMeshComponent::StaticClass(), FFastGeoSkinnedMeshComponent::Type },
 			{ UInstancedSkinnedMeshComponent::StaticClass(), FFastGeoInstancedSkinnedMeshComponent::Type },
-			{ UDecalComponent::StaticClass(), FFastGeoDecalComponent::Type }
+			{ UDecalComponent::StaticClass(), FFastGeoDecalComponent::Type },
+			{ UPointLightComponent::StaticClass(), FFastGeoPointLightComponent::Type },
+			{ USpotLightComponent::StaticClass(), FFastGeoSpotLightComponent::Type },
+			{ URectLightComponent::StaticClass(), FFastGeoRectLightComponent::Type },
 		};
 
 		// Walk the component class hierarchy and look for a fast geo mapping
@@ -1047,6 +1053,10 @@ FFastGeoTransformResult UFastGeoWorldPartitionRuntimeCellTransformer::CanTransfo
 			return { EFastGeoTransform::Discard, [PrimitiveComponent] { return FString::Printf(TEXT("Component %s has no collision and is not visible"), *FastGeo::GetComponentShortName(PrimitiveComponent)); } };
 		}
 	}
+	else if (!bIsSceneComponentVisible)
+	{
+		return { EFastGeoTransform::Discard, [SceneComponent] { return FString::Printf(TEXT("Component %s has is not visible"), *FastGeo::GetComponentShortName(SceneComponent)); } };
+	}
 
 	if (CanIgnoreComponent(InComponent))
 	{
diff --git a/Engine/Source/Editor/StatsViewer/Private/StatsPages/PrimitiveStatsPage.cpp b/Engine/Source/Editor/StatsViewer/Private/StatsPages/PrimitiveStatsPage.cpp
--- a/Engine/Source/Editor/StatsViewer/Private/StatsPages/PrimitiveStatsPage.cpp
+++ b/Engine/Source/Editor/StatsViewer/Private/StatsPages/PrimitiveStatsPage.cpp
@@ -198,15 +198,14 @@ struct PrimitiveStatsGenerator
 			// Don't list pending kill components.
 			&&	IsValid(InPrimitiveComponent) )
 		{
-			// Retrieve relevant lights.
-			TArray<const ULightComponent*> RelevantLights;
-			World->Scene->GetRelevantLights( InPrimitiveComponent, &RelevantLights );
+			// Retrieve number of relevant lights.
+			const int32 RelevantLightCount = World->Scene->GetRelevantLightCount(InPrimitiveComponent);
 
 			// Only look for relevant lights if we aren't unlit.
 			if( !bUsesOnlyUnlitMaterials )
 			{
 				// Lightmap and shadow map lights are calculated above, per component type, infer the "other" light count here
-				LightsOtherCount = RelevantLights.Num() >= LightsLMCount ? RelevantLights.Num() - LightsLMCount : 0;
+				LightsOtherCount = RelevantLightCount >= LightsLMCount ? RelevantLightCount - LightsLMCount : 0;
 			}
 
 			// Figure out memory used by light and shadow maps and light/ shadow map resolution.
diff --git a/Engine/Source/Runtime/Engine/Classes/Components/LightComponent.h b/Engine/Source/Runtime/Engine/Classes/Components/LightComponent.h
--- a/Engine/Source/Runtime/Engine/Classes/Components/LightComponent.h
+++ b/Engine/Source/Runtime/Engine/Classes/Components/LightComponent.h
@@ -10,6 +10,7 @@
 #include "SceneTypes.h"
 #include "RenderResource.h"
 #include "Components/LightComponentBase.h"
+#include "LightSceneProxyDesc.h"
 #include "PSOPrecacheFwd.h" 
 #include "PSOPrecacheComponent.h"
 #include "LightComponent.generated.h"
@@ -506,7 +507,24 @@ public:
 	/** Get the PSO precaching data */
 	FPSOPrecacheComponentData& GetPSOPrecacheComponentData() { return PSOPrecacheComponentData; }
 
+	ENGINE_API bool CastsStaticShadow() const;
+	ENGINE_API bool CastsDynamicShadow() const;
+	ENGINE_API bool TransmitsThroughSurfaces() const;
+	ENGINE_API bool GetContactShadowLengthInWS() const;
+	ENGINE_API float GetRayEndBias() const;
+	ENGINE_API const FStaticShadowDepthMap* GetStaticShadowDepthMap() const;
+
+	inline static FLightSceneProxyDesc BuildSceneProxyDesc(const ULightComponent& InLightComponent)
+	{
+		FLightSceneProxyDesc ProxyDesc;
+		BuildSceneProxyDesc(InLightComponent, ProxyDesc);
+		return ProxyDesc;
+	}
+
 protected:
+
+	ENGINE_API static void BuildSceneProxyDesc(const ULightComponent& InLightComponent, FLightSceneProxyDesc& OutProxyDesc);
+
 	//~ Begin UActorComponent Interface
 	ENGINE_API virtual void OnRegister() override;
 	ENGINE_API virtual void CreateRenderState_Concurrent(FRegisterComponentContext* Context) override;
diff --git a/Engine/Source/Runtime/Engine/Classes/Components/LightComponentBase.h b/Engine/Source/Runtime/Engine/Classes/Components/LightComponentBase.h
--- a/Engine/Source/Runtime/Engine/Classes/Components/LightComponentBase.h
+++ b/Engine/Source/Runtime/Engine/Classes/Components/LightComponentBase.h
@@ -5,6 +5,7 @@
 #include "CoreMinimal.h"
 #include "UObject/ObjectMacros.h"
 #include "Misc/Guid.h"
+#include "LightComponentId.h"
 #include "Components/SceneComponent.h"
 #include "LightComponentBase.generated.h"
 
@@ -245,6 +246,13 @@ class ULightComponentBase : public USceneComponent
 
 	/** Get the extent used when placing this component in the editor, used for 'pulling back' hit. */
 	ENGINE_API virtual FBoxSphereBounds GetPlacementExtent() const override;
+
+	const FLightComponentId& GetLightComponentId() const { return LightComponentId; }
+	TEnumAsByte<ECastRayTracedShadow::Type> GetCastRaytracedShadow() const { return CastShadows == 0 ? (TEnumAsByte<ECastRayTracedShadow::Type>) ECastRayTracedShadow::Disabled : CastRaytracedShadow; }
+
+private:
+	/** Runtime unique id used to identify a light component. */
+	FLightComponentId LightComponentId;
 };
 
 
diff --git a/Engine/Source/Runtime/Engine/Classes/Components/LocalLightComponent.h b/Engine/Source/Runtime/Engine/Classes/Components/LocalLightComponent.h
--- a/Engine/Source/Runtime/Engine/Classes/Components/LocalLightComponent.h
+++ b/Engine/Source/Runtime/Engine/Classes/Components/LocalLightComponent.h
@@ -5,6 +5,7 @@
 #include "CoreMinimal.h"
 #include "UObject/ObjectMacros.h"
 #include "Engine/EngineTypes.h"
+#include "LocalLightSceneProxyDesc.h"
 #include "Components/LightComponent.h"
 #include "LocalLightComponent.generated.h"
 
@@ -83,6 +84,17 @@ public:
 #endif // WITH_EDITOR
 	//~ End UObject Interface
 
+	inline static FLocalLightSceneProxyDesc BuildSceneProxyDesc(const ULocalLightComponent& InLightComponent)
+	{
+		FLocalLightSceneProxyDesc ProxyDesc;
+		BuildSceneProxyDesc(InLightComponent, ProxyDesc);
+		return ProxyDesc;
+	}
+
+protected:
+
+	ENGINE_API static void BuildSceneProxyDesc(const ULocalLightComponent& InLightComponent, FLocalLightSceneProxyDesc& OutProxyDesc);
+
 private:
 
 	/** Pushes the value of radius to the rendering thread. */
diff --git a/Engine/Source/Runtime/Engine/Classes/Components/PointLightComponent.h b/Engine/Source/Runtime/Engine/Classes/Components/PointLightComponent.h
--- a/Engine/Source/Runtime/Engine/Classes/Components/PointLightComponent.h
+++ b/Engine/Source/Runtime/Engine/Classes/Components/PointLightComponent.h
@@ -6,6 +6,7 @@
 #include "CoreMinimal.h"
 #include "UObject/ObjectMacros.h"
 #include "Engine/EngineTypes.h"
+#include "PointLightSceneProxyDesc.h"
 #include "Components/LocalLightComponent.h"
 #include "PointLightComponent.generated.h"
 
@@ -95,7 +96,15 @@ public:
 	ENGINE_API virtual void PostEditChangeProperty(FPropertyChangedEvent& PropertyChangedEvent) override;
 #endif // WITH_EDITOR
 	//~ End UObject Interface
-};
 
+	inline static FPointLightSceneProxyDesc BuildSceneProxyDesc(const UPointLightComponent& InLightComponent)
+	{
+		FPointLightSceneProxyDesc ProxyDesc;
+		BuildSceneProxyDesc(InLightComponent, ProxyDesc);
+		return ProxyDesc;
+	}
 
+protected:
 
+	ENGINE_API static void BuildSceneProxyDesc(const UPointLightComponent& InLightComponent, FPointLightSceneProxyDesc& OutProxyDesc);
+};
\ No newline at end of file
diff --git a/Engine/Source/Runtime/Engine/Classes/Components/RectLightComponent.h b/Engine/Source/Runtime/Engine/Classes/Components/RectLightComponent.h
--- a/Engine/Source/Runtime/Engine/Classes/Components/RectLightComponent.h
+++ b/Engine/Source/Runtime/Engine/Classes/Components/RectLightComponent.h
@@ -6,6 +6,7 @@
 #include "UObject/ObjectMacros.h"
 #include "Engine/EngineTypes.h"
 #include "Engine/Scene.h"
+#include "RectLightSceneProxyDesc.h"
 #include "Components/LocalLightComponent.h"
 #include "RectLightComponent.generated.h"
 
@@ -112,7 +113,17 @@ public:
 	ENGINE_API virtual void Serialize(FArchive& Ar) override;
 	ENGINE_API virtual void PostLoad() override;
 
+	inline static FRectLightSceneProxyDesc BuildSceneProxyDesc(const URectLightComponent& InLightComponent)
+	{
+		FRectLightSceneProxyDesc ProxyDesc;
+		BuildSceneProxyDesc(InLightComponent, ProxyDesc);
+		return ProxyDesc;
+	}
+
 protected:
+
+	ENGINE_API static void BuildSceneProxyDesc(const URectLightComponent& InLightComponent, FRectLightSceneProxyDesc& OutProxyDesc);
+
 #if WITH_EDITOR
 	/** Validate source texture is of a valid type. */
 	void ValidateTexture() const;
diff --git a/Engine/Source/Runtime/Engine/Classes/Components/SpotLightComponent.h b/Engine/Source/Runtime/Engine/Classes/Components/SpotLightComponent.h
--- a/Engine/Source/Runtime/Engine/Classes/Components/SpotLightComponent.h
+++ b/Engine/Source/Runtime/Engine/Classes/Components/SpotLightComponent.h
@@ -4,6 +4,7 @@
 
 #include "CoreMinimal.h"
 #include "UObject/ObjectMacros.h"
+#include "SpotLightSceneProxyDesc.h"
 #include "Components/PointLightComponent.h"
 #include "SpotLightComponent.generated.h"
 
@@ -44,6 +45,7 @@ class USpotLightComponent : public UPointLightComponent
 	ENGINE_API virtual FBox GetStreamingBounds() const override;
 #endif
 
+	static FVector2f GetClampedConeAngles(float InInnerConeAngle, float InOuterConeAngle);
 	ENGINE_API FVector2f GetClampedConeAngles() const;
 	ENGINE_API float GetHalfConeAngle() const;
 	ENGINE_API float GetCosHalfConeAngle() const;
@@ -51,7 +53,15 @@ class USpotLightComponent : public UPointLightComponent
 #if WITH_EDITOR
 	ENGINE_API virtual void PostEditChangeProperty( struct FPropertyChangedEvent& PropertyChangedEvent) override;
 #endif
-};
 
+	inline static FSpotLightSceneProxyDesc BuildSceneProxyDesc(const USpotLightComponent& InLightComponent)
+	{
+		FSpotLightSceneProxyDesc ProxyDesc;
+		BuildSceneProxyDesc(InLightComponent, ProxyDesc);
+		return ProxyDesc;
+	}
 
+protected:
 
+	ENGINE_API static void BuildSceneProxyDesc(const USpotLightComponent& InLightComponent, FSpotLightSceneProxyDesc& OutProxyDesc);
+};
\ No newline at end of file
diff --git a/Engine/Source/Runtime/Engine/Private/Components/LightComponent.cpp b/Engine/Source/Runtime/Engine/Private/Components/LightComponent.cpp
--- a/Engine/Source/Runtime/Engine/Private/Components/LightComponent.cpp
+++ b/Engine/Source/Runtime/Engine/Private/Components/LightComponent.cpp
@@ -447,6 +447,7 @@ ULightComponentBase::ULightComponentBase(const FObjectInitializer& ObjectInitial
 	bVisualizeComponent = true;
 #endif
 	DeepShadowLayerDistribution = 0.5f;
+	LightComponentId = FLightComponentId::GetNextId();
 }
 
 ULightComponent::FOnUpdateColorAndBrightness ULightComponent::UpdateColorAndBrightnessEvent;
@@ -525,6 +526,36 @@ bool ULightComponent::IsShadowCast(UPrimitiveComponent* Primitive) const
 	}
 }
 
+bool ULightComponent::CastsStaticShadow() const
+{
+	return CastShadows && CastStaticShadows;
+}
+
+bool ULightComponent::CastsDynamicShadow() const
+{
+	return CastShadows && CastDynamicShadows;
+}
+
+bool ULightComponent::TransmitsThroughSurfaces() const
+{
+	return bTransmission && CastsDynamicShadow() && !CastsStaticShadow();
+}
+
+bool ULightComponent::GetContactShadowLengthInWS() const
+{
+	return ContactShadowLengthInWS ? true : false;
+}
+
+float ULightComponent::GetRayEndBias() const
+{
+	return bOverrideRayEndBias ? FMath::Max(RayEndBias, 0.0f) : -1.0f;
+}
+
+const FStaticShadowDepthMap* ULightComponent::GetStaticShadowDepthMap() const
+{
+	return &StaticShadowDepthMap;
+}
+
 float ULightComponent::ComputeLightBrightness() const
 {
 	float LightBrightness = Intensity;
@@ -1646,6 +1677,76 @@ void ULightComponent::PushSelectionToProxy()
 
 }
 
+void ULightComponent::BuildSceneProxyDesc(const ULightComponent& InLightComponent, FLightSceneProxyDesc& OutProxyDesc)
+{
+	const FLightComponentMapBuildData* MapBuildData = InLightComponent.GetLightComponentMapBuildData();
+
+	OutProxyDesc.LightComponent = &InLightComponent;
+	OutProxyDesc.LightComponentId = InLightComponent.GetLightComponentId();
+	OutProxyDesc.SceneInterface = InLightComponent.GetScene();
+#if ACTOR_HAS_LABELS
+	OutProxyDesc.OwnerNameOrLabel = InLightComponent.GetOwner() ? InLightComponent.GetOwner()->GetActorNameOrLabel() : InLightComponent.GetName();
+#endif
+	OutProxyDesc.bSelected = InLightComponent.IsSelected() || InLightComponent.IsOwnerSelected();
+	OutProxyDesc.bContactShadowLengthInWS = InLightComponent.GetContactShadowLengthInWS();
+	OutProxyDesc.bMovable = InLightComponent.IsMovable();
+	OutProxyDesc.bStaticLighting = InLightComponent.HasStaticLighting();
+	OutProxyDesc.bStaticShadowing = InLightComponent.HasStaticShadowing();
+	OutProxyDesc.bCastDynamicShadow = InLightComponent.CastsDynamicShadow();
+	OutProxyDesc.bCastStaticShadow = InLightComponent.CastsStaticShadow();
+	OutProxyDesc.bCastTranslucentShadows = InLightComponent.CastTranslucentShadows;
+	OutProxyDesc.bTransmission = InLightComponent.TransmitsThroughSurfaces();
+	OutProxyDesc.bCastVolumetricShadow = InLightComponent.bCastVolumetricShadow;
+	OutProxyDesc.bCastHairStrandsDeepShadow = InLightComponent.bCastDeepShadow;
+	OutProxyDesc.bCastShadowsFromCinematicObjectsOnly = InLightComponent.bCastShadowsFromCinematicObjectsOnly;
+	OutProxyDesc.bForceCachedShadowsForMovablePrimitives = InLightComponent.bForceCachedShadowsForMovablePrimitives;
+	OutProxyDesc.bAffectReflection = InLightComponent.bAffectReflection;
+	OutProxyDesc.bAffectGlobalIllumination = InLightComponent.bAffectGlobalIllumination;
+	OutProxyDesc.bAffectTranslucentLighting = InLightComponent.bAffectTranslucentLighting;
+	OutProxyDesc.bUsedAsAtmosphereSunLight = InLightComponent.IsUsedAsAtmosphereSunLight();
+	OutProxyDesc.bUseRayTracedDistanceFieldShadows = InLightComponent.bUseRayTracedDistanceFieldShadows;
+	OutProxyDesc.bAllowMegaLights = InLightComponent.bAllowMegaLights;
+	OutProxyDesc.bIsPrecomputedLightingValid = InLightComponent.IsPrecomputedLightingValid();
+	OutProxyDesc.bEnableLightShaftBloom = InLightComponent.bEnableLightShaftBloom;
+	OutProxyDesc.Color = InLightComponent.GetColoredLightBrightness();
+	OutProxyDesc.IndirectLightingScale = InLightComponent.IndirectLightingIntensity;
+	OutProxyDesc.VolumetricScatteringIntensity = FMath::Max(InLightComponent.VolumetricScatteringIntensity, 0.0f);
+	OutProxyDesc.RayEndBias = InLightComponent.GetRayEndBias();
+	OutProxyDesc.ShadowResolutionScale = InLightComponent.ShadowResolutionScale;
+	OutProxyDesc.ShadowBias = InLightComponent.ShadowBias;
+	OutProxyDesc.ShadowSlopeBias = InLightComponent.ShadowSlopeBias;
+	OutProxyDesc.ShadowSharpen = InLightComponent.ShadowSharpen;
+	OutProxyDesc.ContactShadowLength = InLightComponent.ContactShadowLength;
+	OutProxyDesc.ContactShadowCastingIntensity = InLightComponent.ContactShadowCastingIntensity;
+	OutProxyDesc.ContactShadowNonCastingIntensity = InLightComponent.ContactShadowNonCastingIntensity;
+	OutProxyDesc.SpecularScale = InLightComponent.SpecularScale;
+	OutProxyDesc.DiffuseScale = InLightComponent.DiffuseScale;
+	OutProxyDesc.LightGuid = InLightComponent.LightGuid;
+	OutProxyDesc.MapBuildDataShadowMapChannel = MapBuildData ? MapBuildData->ShadowMapChannel : INDEX_NONE;
+	OutProxyDesc.PreviewShadowMapChannel = InLightComponent.PreviewShadowMapChannel;
+	OutProxyDesc.RayStartOffsetDepthScale = InLightComponent.RayStartOffsetDepthScale;
+	OutProxyDesc.LightFunctionScale = InLightComponent.LightFunctionScale;
+	OutProxyDesc.LightFunctionFadeDistance = InLightComponent.LightFunctionFadeDistance;
+	OutProxyDesc.LightFunctionDisabledBrightness = InLightComponent.DisabledBrightness;
+	OutProxyDesc.LightFunctionMaterial = InLightComponent.LightFunctionMaterial;
+	OutProxyDesc.IESTexture = InLightComponent.IESTexture;
+	OutProxyDesc.CastRaytracedShadow = InLightComponent.GetCastRaytracedShadow();
+	OutProxyDesc.BloomScale = InLightComponent.BloomScale;
+	OutProxyDesc.BloomThreshold = InLightComponent.BloomThreshold;
+	OutProxyDesc.BloomMaxBrightness = InLightComponent.BloomMaxBrightness;
+	OutProxyDesc.BloomTint = InLightComponent.BloomTint;
+	OutProxyDesc.MegaLightsShadowMethod = InLightComponent.MegaLightsShadowMethod;
+	OutProxyDesc.AtmosphereSunLightIndex = InLightComponent.GetAtmosphereSunLightIndex();
+	OutProxyDesc.AtmosphereSunDiskColorScale = InLightComponent.GetAtmosphereSunDiskColorScale();
+	OutProxyDesc.LightType = InLightComponent.GetLightType();
+	OutProxyDesc.LightingChannelMask = GetLightingChannelMaskForStruct(InLightComponent.LightingChannels);
+	OutProxyDesc.ViewLightingChannelMask = InLightComponent.ViewLightingChannels.GetMaskForStruct();
+	OutProxyDesc.ComponentName = InLightComponent.GetFName();
+	OutProxyDesc.LevelName = InLightComponent.GetOwner() ? InLightComponent.GetOwner()->GetLevel()->GetOutermost()->GetFName() : NAME_None;
+	OutProxyDesc.SamplesPerPixel = InLightComponent.SamplesPerPixel;
+	OutProxyDesc.DeepShadowLayerDistribution = InLightComponent.DeepShadowLayerDistribution;
+}
+
 /** Stores a light and a channel it has been assigned to. */
 struct FLightAndChannel
 {
diff --git a/Engine/Source/Runtime/Engine/Private/Components/LocalLightComponent.cpp b/Engine/Source/Runtime/Engine/Private/Components/LocalLightComponent.cpp
--- a/Engine/Source/Runtime/Engine/Private/Components/LocalLightComponent.cpp
+++ b/Engine/Source/Runtime/Engine/Private/Components/LocalLightComponent.cpp
@@ -1,9 +1,5 @@
 // Copyright Epic Games, Inc. All Rights Reserved.
 
-/*=============================================================================
-	PointLightComponent.cpp: PointLightComponent implementation.
-=============================================================================*/
-
 #include "Components/LocalLightComponent.h"
 #include "Engine/Scene.h"
 #include "LocalLightSceneProxy.h"
@@ -229,3 +225,12 @@ float ULocalLightComponent::GetUnitsConversionFactor(ELightUnits SrcUnits, ELigh
 	}
 }
 
+void ULocalLightComponent::BuildSceneProxyDesc(const ULocalLightComponent& InLightComponent, FLocalLightSceneProxyDesc& OutProxyDesc)
+{
+	ULightComponent::BuildSceneProxyDesc(InLightComponent, OutProxyDesc);
+
+	OutProxyDesc.MaxDrawDistance = InLightComponent.MaxDrawDistance;
+	OutProxyDesc.MaxDistanceFadeRange = InLightComponent.MaxDistanceFadeRange;
+	OutProxyDesc.InverseExposureBlend = InLightComponent.InverseExposureBlend;
+	OutProxyDesc.AttenuationRadius = InLightComponent.AttenuationRadius;
+}
\ No newline at end of file
diff --git a/Engine/Source/Runtime/Engine/Private/Components/PointLightComponent.cpp b/Engine/Source/Runtime/Engine/Private/Components/PointLightComponent.cpp
--- a/Engine/Source/Runtime/Engine/Private/Components/PointLightComponent.cpp
+++ b/Engine/Source/Runtime/Engine/Private/Components/PointLightComponent.cpp
@@ -1,9 +1,5 @@
 // Copyright Epic Games, Inc. All Rights Reserved.
 
-/*=============================================================================
-	PointLightComponent.cpp: PointLightComponent implementation.
-=============================================================================*/
-
 #include "Components/PointLightComponent.h"
 #include "RenderUtils.h"
 #include "UObject/ConstructorHelpers.h"
@@ -15,82 +11,6 @@
 
 #include UE_INLINE_GENERATED_CPP_BY_NAME(PointLightComponent)
 
-int32 GAllowPointLightCubemapShadows = 1;
-static FAutoConsoleVariableRef CVarAllowPointLightCubemapShadows(
-	TEXT("r.AllowPointLightCubemapShadows"),
-	GAllowPointLightCubemapShadows,
-	TEXT("When 0, will prevent point light cube map shadows from being used and the light will be unshadowed.")
-	);
-
-void FLocalLightSceneProxy::UpdateRadius_GameThread(float ComponentRadius)
-{
-	FLocalLightSceneProxy* InLightSceneInfo = this;
-	ENQUEUE_RENDER_COMMAND(UpdateRadius)(
-		[InLightSceneInfo, ComponentRadius](FRHICommandList& RHICmdList)
-		{
-			InLightSceneInfo->UpdateRadius(ComponentRadius);
-		});
-}
-
-/** Accesses parameters needed for rendering the light. */
-void FPointLightSceneProxy::GetLightShaderParameters(FLightRenderParameters& LightParameters, uint32 Flags) const
-{
-	LightParameters.WorldPosition = GetOrigin();
-	LightParameters.InvRadius = InvRadius;
-	LightParameters.Color = GetColor();
-	LightParameters.FalloffExponent = FalloffExponent;
-
-	// TODO LWC - GetDirection() seems like it needs to be normalized, somehow accumulating error with large-scale position values
-	LightParameters.Direction = (FVector3f)-GetDirection(); // LWC_TODO: Precision Loss
-	LightParameters.Tangent = FVector3f(WorldToLight.M[0][2], WorldToLight.M[1][2], WorldToLight.M[2][2]);
-	LightParameters.SpotAngles = FVector2f( -2.0f, 1.0f );
-	LightParameters.SpecularScale = FMath::Clamp(SpecularScale, 0.f, 1.f);
-	LightParameters.DiffuseScale = FMath::Clamp(DiffuseScale, 0.f, 1.f);
-	LightParameters.SourceRadius = SourceRadius;
-	LightParameters.SoftSourceRadius = SoftSourceRadius;
-	LightParameters.SourceLength = SourceLength;
-	LightParameters.RectLightBarnCosAngle = 0.0f;
-	LightParameters.RectLightBarnLength = -2.0f;
-	LightParameters.RectLightAtlasUVOffset = FVector2f::ZeroVector;
-	LightParameters.RectLightAtlasUVScale = FVector2f::ZeroVector;
-	LightParameters.RectLightAtlasMaxLevel = FLightRenderParameters::GetRectLightAtlasInvalidMIPLevel();
-	LightParameters.IESAtlasIndex = INDEX_NONE;
-	LightParameters.InverseExposureBlend = InverseExposureBlend;
-	LightParameters.LightFunctionAtlasLightIndex = GetLightFunctionAtlasLightIndex();
-	LightParameters.bAffectsTranslucentLighting = AffectsTranslucentLighting() ? 1 : 0;
-
-	if (IESAtlasId != ~0)
-	{
-		GetSceneInterface()->GetLightIESAtlasSlot(this, &LightParameters);
-	}
-}
-
-/**
-* Sets up a projected shadow initializer for shadows from the entire scene.
-* @return True if the whole-scene projected shadow should be used.
-*/
-bool FPointLightSceneProxy::GetWholeSceneProjectedShadowInitializer(const FSceneViewFamily& ViewFamily, TArray<FWholeSceneProjectedShadowInitializer, TInlineAllocator<6> >& OutInitializers) const
-{
-	if ((ViewFamily.GetFeatureLevel() >= ERHIFeatureLevel::SM5 || DoesRuntimeSupportOnePassPointLightShadows(ViewFamily.GetShaderPlatform()))
-		&& GAllowPointLightCubemapShadows != 0)
-	{
-		FWholeSceneProjectedShadowInitializer& OutInitializer = *new(OutInitializers) FWholeSceneProjectedShadowInitializer;
-		OutInitializer.PreShadowTranslation = -GetLightToWorld().GetOrigin();
-		OutInitializer.WorldToLight = GetWorldToLight().RemoveTranslation();
-		OutInitializer.Scales = FVector2D(1, 1);
-		OutInitializer.SubjectBounds = FBoxSphereBounds(FVector(0, 0, 0),FVector(Radius,Radius,Radius),Radius);
-		OutInitializer.WAxis = FVector4(0,0,1,0);
-		OutInitializer.MinLightW = 0.1f;
-		OutInitializer.MaxDistanceToCastInLightW = Radius;
-		OutInitializer.bOnePassPointLightShadow = true;
-
-		OutInitializer.bRayTracedDistanceField = UseRayTracedDistanceFieldShadows() && DoesPlatformSupportDistanceFieldShadowing(ViewFamily.GetShaderPlatform());
-		return true;
-	}
-		
-	return false;
-}
-
 UPointLightComponent::UPointLightComponent(const FObjectInitializer& ObjectInitializer)
 	: Super(ObjectInitializer)
 {
@@ -356,3 +276,13 @@ void UPointLightComponent::PostEditChangeProperty(FPropertyChangedEvent& Propert
 }
 #endif // WITH_EDITOR
 
+void UPointLightComponent::BuildSceneProxyDesc(const UPointLightComponent& InLightComponent, FPointLightSceneProxyDesc& OutProxyDesc)
+{
+	ULocalLightComponent::BuildSceneProxyDesc(InLightComponent, OutProxyDesc);
+
+	OutProxyDesc.LightFalloffExponent = InLightComponent.LightFalloffExponent;
+	OutProxyDesc.SourceRadius = InLightComponent.SourceRadius;
+	OutProxyDesc.SoftSourceRadius = InLightComponent.SoftSourceRadius;
+	OutProxyDesc.SourceLength = InLightComponent.SourceLength;
+	OutProxyDesc.bUseInverseSquaredFalloff = InLightComponent.bUseInverseSquaredFalloff;
+}
\ No newline at end of file
diff --git a/Engine/Source/Runtime/Engine/Private/Components/RectLightComponent.cpp b/Engine/Source/Runtime/Engine/Private/Components/RectLightComponent.cpp
--- a/Engine/Source/Runtime/Engine/Private/Components/RectLightComponent.cpp
+++ b/Engine/Source/Runtime/Engine/Private/Components/RectLightComponent.cpp
@@ -1,11 +1,6 @@
 // Copyright Epic Games, Inc. All Rights Reserved.
 
-/*=============================================================================
-	PointLightComponent.cpp: PointLightComponent implementation.
-=============================================================================*/
-
 #include "Components/RectLightComponent.h"
-
 #include "DataDrivenShaderPlatformInfo.h"
 #include "Engine/Texture2D.h"
 #include "Logging/MessageLog.h"
@@ -20,8 +15,6 @@
 
 #define LOCTEXT_NAMESPACE "RectLightComponent"
 
-extern int32 GAllowPointLightCubemapShadows;
-
 float GetRectLightBarnDoorMaxAngle()
 {
 	return 88.f;
@@ -372,117 +365,18 @@ void URectLightComponent::PostLoad()
 #endif
 }
 
-
-FRectLightSceneProxy::FRectLightSceneProxy(const URectLightComponent* Component)
-	: FLocalLightSceneProxy(Component)
-	, SourceWidth(Component->SourceWidth)
-	, SourceHeight(Component->SourceHeight)
-	, BarnDoorAngle(FMath::Clamp(Component->BarnDoorAngle, 0.f, GetRectLightBarnDoorMaxAngle()))
-	, BarnDoorLength(FMath::Max(0.1f, Component->BarnDoorLength))
-	, SourceTexture(Component->SourceTexture)
-	, LightFunctionConeAngleTangent(Component->LightFunctionConeAngle > 0 ? FMath::Tan(FMath::Clamp(Component->LightFunctionConeAngle, 0.0f, 89.0f) * (float)UE_PI / 180.0f) : 0.0f)
-	, SourceTextureScaleOffset(FVector4f(
-	  FMath::Clamp(Component->SourceTextureScale.X, 0.f, 1.f), 
-	  FMath::Clamp(Component->SourceTextureScale.Y, 0.f, 1.f), 
-	  FMath::Clamp(Component->SourceTextureOffset.X, 0.f, 1.f), 
-	  FMath::Clamp(Component->SourceTextureOffset.Y, 0.f, 1.f)))
-{
-	RectAtlasId = ~0u;
-}
-
-FRectLightSceneProxy::~FRectLightSceneProxy() {}
-
-bool FRectLightSceneProxy::IsRectLight() const
+void URectLightComponent::BuildSceneProxyDesc(const URectLightComponent& InLightComponent, FRectLightSceneProxyDesc& OutProxyDesc)
 {
-	return true;
-}
-
-bool FRectLightSceneProxy::HasSourceTexture() const
-{
-	return SourceTexture != nullptr;
-}
-
-/** Accesses parameters needed for rendering the light. */
-void FRectLightSceneProxy::GetLightShaderParameters(FLightRenderParameters& LightParameters, uint32 Flags) const
-{
-	FLinearColor LightColor = GetColor();
-	LightColor /= 0.5f * SourceWidth * SourceHeight;
-	LightParameters.WorldPosition = GetOrigin();
-	LightParameters.InvRadius = InvRadius;
-	LightParameters.Color = LightColor;
-	LightParameters.FalloffExponent = 0.0f;
-
-	LightParameters.Direction = FVector3f(-GetDirection());
-	LightParameters.Tangent = FVector3f(WorldToLight.M[0][2], WorldToLight.M[1][2], WorldToLight.M[2][2]);
-	LightParameters.SpotAngles = FVector2f(-2.0f, 1.0f);
-	LightParameters.SpecularScale = FMath::Clamp(SpecularScale, 0.f, 1.f);
-	LightParameters.DiffuseScale = FMath::Clamp(DiffuseScale, 0.f, 1.f);
-	LightParameters.SourceRadius = SourceWidth * 0.5f;
-	LightParameters.SoftSourceRadius = 0.0f;
-	LightParameters.SourceLength = SourceHeight * 0.5f;
-	LightParameters.RectLightBarnCosAngle = FMath::Cos(FMath::DegreesToRadians(BarnDoorAngle));
-	LightParameters.RectLightBarnLength = BarnDoorLength;
-	LightParameters.RectLightAtlasUVOffset = FVector2f::ZeroVector;
-	LightParameters.RectLightAtlasUVScale = FVector2f::ZeroVector;
-	LightParameters.RectLightAtlasMaxLevel = FLightRenderParameters::GetRectLightAtlasInvalidMIPLevel();
-	LightParameters.IESAtlasIndex = INDEX_NONE;
-	LightParameters.InverseExposureBlend = InverseExposureBlend;
-	LightParameters.LightFunctionAtlasLightIndex = GetLightFunctionAtlasLightIndex();
-	LightParameters.bAffectsTranslucentLighting = AffectsTranslucentLighting() ? 1 : 0;
-
-	if (IESAtlasId != ~0)
-	{
-		GetSceneInterface()->GetLightIESAtlasSlot(this, &LightParameters);
-	}
-
-	if (RectAtlasId != ~0u)
-	{
-		GetSceneInterface()->GetRectLightAtlasSlot(this, &LightParameters);
-	}
-	
-	// Render RectLight approximately as SpotLight if the requester does not support rect light (e.g., translucent light grid or mobile)
-	if (!!(Flags & ELightShaderParameterFlags::RectAsSpotLight))
-	{
-		float ClampedOuterConeAngle = FMath::DegreesToRadians(89.001f);
-		float ClampedInnerConeAngle = FMath::DegreesToRadians(70.0f);
-		float CosOuterCone = FMath::Cos(ClampedOuterConeAngle);
-		float CosInnerCone = FMath::Cos(ClampedInnerConeAngle);
-		float InvCosConeDifference = 1.0f / (CosInnerCone - CosOuterCone);
-
-		LightParameters.Color = GetColor();
-		LightParameters.FalloffExponent = 8.0f;
-		LightParameters.SpotAngles = FVector2f(CosOuterCone, InvCosConeDifference);
-		LightParameters.SourceRadius = (SourceWidth + SourceHeight) * 0.5 * 0.5f;
-		LightParameters.SourceLength = 0.0f;
-		LightParameters.RectLightBarnCosAngle = 0.0f;
-		LightParameters.RectLightBarnLength = -2.0f;
-	}
-}
-
-/**
-* Sets up a projected shadow initializer for shadows from the entire scene.
-* @return True if the whole-scene projected shadow should be used.
-*/
-bool FRectLightSceneProxy::GetWholeSceneProjectedShadowInitializer(const FSceneViewFamily& ViewFamily, TArray<FWholeSceneProjectedShadowInitializer, TInlineAllocator<6> >& OutInitializers) const
-{
-	if (ViewFamily.GetFeatureLevel() >= ERHIFeatureLevel::SM5
-		&& GAllowPointLightCubemapShadows != 0)
-	{
-		FWholeSceneProjectedShadowInitializer& OutInitializer = *new(OutInitializers) FWholeSceneProjectedShadowInitializer;
-		OutInitializer.PreShadowTranslation = -GetLightToWorld().GetOrigin();
-		OutInitializer.WorldToLight = GetWorldToLight().RemoveTranslation();
-		OutInitializer.Scales = FVector2D(1, 1);
-		OutInitializer.SubjectBounds = FBoxSphereBounds(FVector(0, 0, 0), FVector(Radius, Radius, Radius), Radius);
-		OutInitializer.WAxis = FVector4(0, 0, 1, 0);
-		OutInitializer.MinLightW = 0.1f;
-		OutInitializer.MaxDistanceToCastInLightW = Radius;
-		OutInitializer.bOnePassPointLightShadow = true;
-
-		OutInitializer.bRayTracedDistanceField = UseRayTracedDistanceFieldShadows() && DoesPlatformSupportDistanceFieldShadowing(ViewFamily.GetShaderPlatform());
-		return true;
-	}
-
-	return false;
+	ULocalLightComponent::BuildSceneProxyDesc(InLightComponent, OutProxyDesc);
+
+	OutProxyDesc.SourceWidth = InLightComponent.SourceWidth;
+	OutProxyDesc.SourceHeight = InLightComponent.SourceHeight;
+	OutProxyDesc.BarnDoorAngle = InLightComponent.BarnDoorAngle;
+	OutProxyDesc.BarnDoorLength = InLightComponent.BarnDoorLength;
+	OutProxyDesc.LightFunctionConeAngle = InLightComponent.LightFunctionConeAngle;
+	OutProxyDesc.SourceTexture = InLightComponent.SourceTexture;
+	OutProxyDesc.SourceTextureScale = InLightComponent.SourceTextureScale;
+	OutProxyDesc.SourceTextureOffset = InLightComponent.SourceTextureOffset;
 }
 
 #undef LOCTEXT_NAMESPACE
diff --git a/Engine/Source/Runtime/Engine/Private/Components/SpotLightComponent.cpp b/Engine/Source/Runtime/Engine/Private/Components/SpotLightComponent.cpp
--- a/Engine/Source/Runtime/Engine/Private/Components/SpotLightComponent.cpp
+++ b/Engine/Source/Runtime/Engine/Private/Components/SpotLightComponent.cpp
@@ -1,186 +1,17 @@
 // Copyright Epic Games, Inc. All Rights Reserved.
 
-/*=============================================================================
-	SpotLightComponent.cpp: LightComponent implementation.
-=============================================================================*/
-
 #include "Components/SpotLightComponent.h"
 #include "RenderUtils.h"
 #include "UObject/ConstructorHelpers.h"
 #include "Engine/Texture2D.h"
 #include "PointLightSceneProxy.h"
+#include "SpotLightSceneProxy.h"
 #include "UObject/UnrealType.h"
 #include "SceneInterface.h"
 #include "SceneView.h"
 
 #include UE_INLINE_GENERATED_CPP_BY_NAME(SpotLightComponent)
 
-/**
- * The scene info for a spot light.
- */
-class FSpotLightSceneProxy : public FPointLightSceneProxy
-{
-public:
-
-	/** Outer cone angle in radians, clamped to a valid range. */
-	float OuterConeAngle;
-
-	/** Cosine of the spot light's inner cone angle. */
-	float CosInnerCone;
-
-	/** Cosine of the spot light's outer cone angle. */
-	float CosOuterCone;
-
-	/** 1 / (CosInnerCone - CosOuterCone) */
-	float InvCosConeDifference;
-
-	/** Sine of the spot light's outer cone angle. */
-	float SinOuterCone;
-
-	/** 1 / Tangent of the spot light's outer cone angle. */
-	float InvTanOuterCone;
-
-	/** Initialization constructor. */
-	FSpotLightSceneProxy(const USpotLightComponent* Component)
-	:	FPointLightSceneProxy(Component)
-	{
-		const FVector2f ClampedConeAngles = Component->GetClampedConeAngles();
-
-		OuterConeAngle = ClampedConeAngles.Y;
-		CosOuterCone = FMath::Cos(ClampedConeAngles.Y);
-		SinOuterCone = FMath::Sin(ClampedConeAngles.Y);
-		CosInnerCone = FMath::Cos(ClampedConeAngles.X);
-		InvCosConeDifference = 1.0f / (CosInnerCone - CosOuterCone);
-		InvTanOuterCone = 1.0f / FMath::Tan(ClampedConeAngles.Y);
-	}
-
-	/** Accesses parameters needed for rendering the light. */
-	virtual void GetLightShaderParameters(FLightRenderParameters& LightParameters, uint32 Flags=0) const override
-	{
-		LightParameters.WorldPosition = GetOrigin();
-		LightParameters.InvRadius = InvRadius;
-		LightParameters.Color = GetColor();
-		LightParameters.FalloffExponent = FalloffExponent;
-		LightParameters.Direction = FVector3f(-GetDirection());
-		LightParameters.Tangent = FVector3f(WorldToLight.M[0][2], WorldToLight.M[1][2], WorldToLight.M[2][2]);
-		LightParameters.SpotAngles = FVector2f(CosOuterCone, InvCosConeDifference);
-		LightParameters.SpecularScale = FMath::Clamp(SpecularScale, 0.f, 1.f);
-		LightParameters.DiffuseScale = FMath::Clamp(DiffuseScale, 0.f, 1.f);
-		LightParameters.SourceRadius = SourceRadius;
-		LightParameters.SoftSourceRadius = SoftSourceRadius;
-		LightParameters.SourceLength = SourceLength;
-		LightParameters.RectLightBarnCosAngle = 0.0f;
-		LightParameters.RectLightBarnLength = -2.0f;
-		LightParameters.RectLightAtlasUVOffset = FVector2f::ZeroVector;
-		LightParameters.RectLightAtlasUVScale = FVector2f::ZeroVector;
-		LightParameters.RectLightAtlasMaxLevel = FLightRenderParameters::GetRectLightAtlasInvalidMIPLevel();
-		LightParameters.IESAtlasIndex = INDEX_NONE;
-		LightParameters.LightFunctionAtlasLightIndex = GetLightFunctionAtlasLightIndex();
-		LightParameters.bAffectsTranslucentLighting = AffectsTranslucentLighting() ? 1 : 0;
-		LightParameters.InverseExposureBlend = InverseExposureBlend;
-
-		if (IESAtlasId != uint32(INDEX_NONE))
-		{
-			GetSceneInterface()->GetLightIESAtlasSlot(this, &LightParameters);
-		}
-	}
-
-	// FLightSceneInfo interface.
-	virtual bool AffectsBounds(const FBoxSphereBounds& Bounds) const override
-	{
-		if(!FLocalLightSceneProxy::AffectsBounds(Bounds))
-		{
-			return false;
-		}
-
-		FVector	U = GetOrigin() - (Bounds.SphereRadius / SinOuterCone) * GetDirection(),
-				D = Bounds.Origin - U;
-		float	dsqr = D | D,
-				E = GetDirection() | D;
-		if(E > 0.0f && E * E >= dsqr * FMath::Square(CosOuterCone))
-		{
-			D = Bounds.Origin - GetOrigin();
-			dsqr = D | D;
-			E = -(GetDirection() | D);
-			if(E > 0.0f && E * E >= dsqr * FMath::Square(SinOuterCone))
-				return dsqr <= FMath::Square(Bounds.SphereRadius);
-			else
-				return true;
-		}
-
-		return false;
-	}
-	
-	/**
-	 * Sets up a projected shadow initializer for shadows from the entire scene.
-	 * @return True if the whole-scene projected shadow should be used.
-	 */
-	virtual bool GetWholeSceneProjectedShadowInitializer(const FSceneViewFamily& ViewFamily, TArray<FWholeSceneProjectedShadowInitializer, TInlineAllocator<6> >& OutInitializers) const override
-	{
-		FWholeSceneProjectedShadowInitializer& OutInitializer = OutInitializers.AddDefaulted_GetRef();
-		OutInitializer.PreShadowTranslation = -GetLightToWorld().GetOrigin();
-		OutInitializer.WorldToLight = GetWorldToLight().RemoveTranslation();
-		OutInitializer.Scales = FVector2D(InvTanOuterCone,InvTanOuterCone);
-
-		const FSphere AbsoluteBoundingSphere = FSpotLightSceneProxy::GetBoundingSphere();
-		OutInitializer.SubjectBounds = FBoxSphereBounds(
-			AbsoluteBoundingSphere.Center - GetOrigin(),
-			FVector(AbsoluteBoundingSphere.W, AbsoluteBoundingSphere.W, AbsoluteBoundingSphere.W),
-			AbsoluteBoundingSphere.W
-			);
-
-		OutInitializer.WAxis = FVector4(0,0,1,0);
-		OutInitializer.MinLightW = 0.1f;
-		OutInitializer.MaxDistanceToCastInLightW = Radius;
-		OutInitializer.bRayTracedDistanceField = UseRayTracedDistanceFieldShadows() && DoesPlatformSupportDistanceFieldShadowing(ViewFamily.GetShaderPlatform());
-		return true;
-	}
-
-	virtual float GetOuterConeAngle() const override { return OuterConeAngle; }
-
-	virtual FSphere GetBoundingSphere() const override
-	{
-		return FMath::ComputeBoundingSphereForCone(GetOrigin(), GetDirection(), (FSphere::FReal)Radius, (FSphere::FReal)CosOuterCone, (FSphere::FReal)SinOuterCone);
-	}
-	
-	virtual float GetEffectiveScreenRadius(const FViewMatrices& ShadowViewMatrices, const FIntPoint& CameraViewRectSize) const 
-	{
-		// Heuristic: use the radius of the inscribed sphere at the cone's end as the light's effective screen radius
-		// We do so because we do not want to use the light's radius directly, which will make us overestimate the shadow map resolution greatly for a spot light
-
-		// In the correct form,
-		//   InscribedSpherePosition = GetOrigin() + GetDirection() * GetRadius() / CosOuterCone
-		//   InscribedSphereRadius = GetRadius() / SinOuterCone
-		// Do it incorrectly to avoid division which is more expensive and risks division by zero
-		const FVector InscribedSpherePosition = GetOrigin() + GetDirection() * GetRadius() * CosOuterCone;
-		const float InscribedSphereRadius = GetRadius() * SinOuterCone;
-
-		const float SphereDistanceFromViewOrigin = (InscribedSpherePosition - ShadowViewMatrices.GetViewOrigin()).Size();
-
-		const FVector2D &ProjectionScale = ShadowViewMatrices.GetProjectionScale();
-		const float ScreenScale = FMath::Max(CameraViewRectSize.X * 0.5f * ProjectionScale.X, CameraViewRectSize.Y * 0.5f * ProjectionScale.Y);
-
-		return ScreenScale * InscribedSphereRadius / FMath::Max(SphereDistanceFromViewOrigin, 1.0f);
-	}
-
-	virtual float GetEffectiveScreenRadius(const FViewMatrices& ShadowViewMatrices) const override
-	{
-		// Heuristic: use the radius of the inscribed sphere at the cone's end as the light's effective screen radius
-		// We do so because we do not want to use the light's radius directly, which will make us overestimate the shadow map resolution greatly for a spot light
-
-		// In the correct form,
-		//   InscribedSpherePosition = GetOrigin() + GetDirection() * GetRadius() / CosOuterCone
-		//   InscribedSphereRadius = GetRadius() / SinOuterCone
-		// Do it incorrectly to avoid division which is more expensive and risks division by zero
-		const FVector InscribedSpherePosition = GetOrigin() + GetDirection() * GetRadius() * CosOuterCone;
-		const float InscribedSphereRadius = GetRadius() * SinOuterCone;
-
-		const float SphereDistanceFromViewOrigin = (InscribedSpherePosition - ShadowViewMatrices.GetViewOrigin()).Size();
-
-		return ShadowViewMatrices.GetScreenScale() * InscribedSphereRadius / FMath::Max(SphereDistanceFromViewOrigin, 1.0f);
-	}
-};
-
 USpotLightComponent::USpotLightComponent(const FObjectInitializer& ObjectInitializer)
 	: Super(ObjectInitializer)
 {
@@ -201,14 +32,19 @@ USpotLightComponent::USpotLightComponent(const FObjectInitializer& ObjectInitial
 	OuterConeAngle = 44.0f;
 }
 
-FVector2f USpotLightComponent::GetClampedConeAngles() const
+FVector2f USpotLightComponent::GetClampedConeAngles(float InInnerConeAngle, float InOuterConeAngle)
 {
 	FVector2f ConeAngles;
-	ConeAngles.X = FMath::Clamp(InnerConeAngle, 0.0f, 88.9f) * (float)UE_PI / 180.0f;
-	ConeAngles.Y = FMath::Clamp(OuterConeAngle * (float)UE_PI / 180.0f, ConeAngles.X + 0.001f, 88.9f * (float)UE_PI / 180.0f + 0.001f);
+	ConeAngles.X = FMath::Clamp(InInnerConeAngle, 0.0f, 88.9f) * (float)UE_PI / 180.0f;
+	ConeAngles.Y = FMath::Clamp(InOuterConeAngle * (float)UE_PI / 180.0f, ConeAngles.X + 0.001f, 88.9f * (float)UE_PI / 180.0f + 0.001f);
 	return ConeAngles;
 }
 
+FVector2f USpotLightComponent::GetClampedConeAngles() const
+{
+	return GetClampedConeAngles(InnerConeAngle, OuterConeAngle);
+}
+
 float USpotLightComponent::GetHalfConeAngle() const
 {
 	const FVector2f ClampedConeAngles = GetClampedConeAngles();
@@ -410,3 +246,10 @@ void USpotLightComponent::PostEditChangeProperty( FPropertyChangedEvent& Propert
 
 #endif	// WITH_EDITOR
 
+void USpotLightComponent::BuildSceneProxyDesc(const USpotLightComponent& InLightComponent, FSpotLightSceneProxyDesc& OutProxyDesc)
+{
+	UPointLightComponent::BuildSceneProxyDesc(InLightComponent, OutProxyDesc);
+
+	OutProxyDesc.InnerConeAngle = InLightComponent.InnerConeAngle;
+	OutProxyDesc.OuterConeAngle = InLightComponent.OuterConeAngle;
+}
\ No newline at end of file
diff --git a/Engine/Source/Runtime/Engine/Private/LightComponentId.cpp b/Engine/Source/Runtime/Engine/Private/LightComponentId.cpp
new file mode 100644
--- /dev/null
+++ b/Engine/Source/Runtime/Engine/Private/LightComponentId.cpp
@@ -0,0 +1,12 @@
+// Copyright Epic Games, Inc. All Rights Reserved.
+
+#include "LightComponentId.h"
+#include "HAL/ThreadSafeCounter.h"
+#include "AutoRTFM.h"
+
+FThreadSafeCounter FLightComponentId::NextLightId = 0;
+
+FLightComponentId FLightComponentId::GetNextId()
+{
+	return FLightComponentId{ AutoRTFM::Open([]() -> uint32 { return NextLightId.Increment(); }) };
+}
\ No newline at end of file
diff --git a/Engine/Source/Runtime/Engine/Private/LightSceneProxy.cpp b/Engine/Source/Runtime/Engine/Private/LightSceneProxy.cpp
--- a/Engine/Source/Runtime/Engine/Private/LightSceneProxy.cpp
+++ b/Engine/Source/Runtime/Engine/Private/LightSceneProxy.cpp
@@ -8,6 +8,7 @@
 #include "Engine/MapBuildDataRegistry.h"
 #include "Engine/TextureLightProfile.h"
 #include "EngineDefines.h"
+#include "LightSceneProxyDesc.h"
 #include "MaterialDomain.h"
 #include "Materials/Material.h"
 #include "Materials/MaterialInterface.h"
@@ -17,86 +18,104 @@
 #include "UObject/Package.h"
 
 FLightSceneProxy::FLightSceneProxy(const ULightComponent* InLightComponent)
-	: LightComponent(InLightComponent)
-	, SceneInterface(InLightComponent->GetScene())
-	, IndirectLightingScale(InLightComponent->IndirectLightingIntensity)
-	, VolumetricScatteringIntensity(FMath::Max(InLightComponent->VolumetricScatteringIntensity, 0.0f))
-	, RayEndBias(InLightComponent->bOverrideRayEndBias ? FMath::Max(InLightComponent->RayEndBias, 0.0f) : -1.0f)
-	, ShadowResolutionScale(InLightComponent->ShadowResolutionScale)
-	, ShadowBias(InLightComponent->ShadowBias)
-	, ShadowSlopeBias(InLightComponent->ShadowSlopeBias)
-	, ShadowSharpen(InLightComponent->ShadowSharpen)
-	, ContactShadowLength(InLightComponent->ContactShadowLength)
-	, ContactShadowCastingIntensity(InLightComponent->ContactShadowCastingIntensity)
-	, ContactShadowNonCastingIntensity(InLightComponent->ContactShadowNonCastingIntensity)
-	, SpecularScale(InLightComponent->SpecularScale)
-	, DiffuseScale(InLightComponent->DiffuseScale)
-	, LightGuid(InLightComponent->LightGuid)
-	, RayStartOffsetDepthScale(InLightComponent->RayStartOffsetDepthScale)
-	, IESTexture(0)
-	, bContactShadowLengthInWS(InLightComponent->ContactShadowLengthInWS ? true : false)
-	, bMovable(InLightComponent->IsMovable())
-	, bStaticLighting(InLightComponent->HasStaticLighting())
-	, bStaticShadowing(InLightComponent->HasStaticShadowing())
-	, bCastDynamicShadow(InLightComponent->CastShadows&& InLightComponent->CastDynamicShadows)
-	, bCastStaticShadow(InLightComponent->CastShadows&& InLightComponent->CastStaticShadows)
-	, bCastTranslucentShadows(InLightComponent->CastTranslucentShadows)
-	, bTransmission(InLightComponent->bTransmission && bCastDynamicShadow && !bStaticShadowing)
-	, bCastVolumetricShadow(InLightComponent->bCastVolumetricShadow)
-	, bCastHairStrandsDeepShadow(InLightComponent->bCastDeepShadow)
-	, bCastShadowsFromCinematicObjectsOnly(InLightComponent->bCastShadowsFromCinematicObjectsOnly)
-	, bForceCachedShadowsForMovablePrimitives(InLightComponent->bForceCachedShadowsForMovablePrimitives)
-	, CastRaytracedShadow(InLightComponent->CastShadows == 0 ? (TEnumAsByte<ECastRayTracedShadow::Type>) ECastRayTracedShadow::Disabled : InLightComponent->CastRaytracedShadow)
-	, bAffectReflection(InLightComponent->bAffectReflection)
-	, bAffectGlobalIllumination(InLightComponent->bAffectGlobalIllumination)
-	, bAffectTranslucentLighting(InLightComponent->bAffectTranslucentLighting)
-	, bUsedAsAtmosphereSunLight(InLightComponent->IsUsedAsAtmosphereSunLight())
-	, bUseRayTracedDistanceFieldShadows(InLightComponent->bUseRayTracedDistanceFieldShadows)
-	, bUseVirtualShadowMaps(false)	// See below
+	: FLightSceneProxy(ULightComponent::BuildSceneProxyDesc(*InLightComponent))
+{
+}
+
+FLightSceneProxy::FLightSceneProxy(const FLightSceneProxyDesc& InLightDesc)
+	: LightComponent(InLightDesc.LightComponent)
+	, LightComponentId(InLightDesc.LightComponentId)
+	, SceneInterface(InLightDesc.SceneInterface)
+	, Color(InLightDesc.Color)
+	, IndirectLightingScale(InLightDesc.IndirectLightingScale)
+	, VolumetricScatteringIntensity(FMath::Max(InLightDesc.VolumetricScatteringIntensity, 0.0f))
+	, RayEndBias(InLightDesc.RayEndBias)
+	, ShadowResolutionScale(InLightDesc.ShadowResolutionScale)
+	, ShadowBias(InLightDesc.ShadowBias)
+	, ShadowSlopeBias(InLightDesc.ShadowSlopeBias)
+	, ShadowSharpen(InLightDesc.ShadowSharpen)
+	, ContactShadowLength(InLightDesc.ContactShadowLength)
+	, ContactShadowCastingIntensity(InLightDesc.ContactShadowCastingIntensity)
+	, ContactShadowNonCastingIntensity(InLightDesc.ContactShadowNonCastingIntensity)
+	, SpecularScale(InLightDesc.SpecularScale)
+	, DiffuseScale(InLightDesc.DiffuseScale)
+	, LightGuid(InLightDesc.LightGuid)
+	, ShadowMapChannel(INDEX_NONE) // Initialized below
+	, PreviewShadowMapChannel(INDEX_NONE) // Initialized below
+	, RayStartOffsetDepthScale(InLightDesc.RayStartOffsetDepthScale)
+	, StaticShadowDepthMap(InLightDesc.GetStaticShadowDepthMap())
+	, LightFunctionScale(InLightDesc.LightFunctionScale)
+	, LightFunctionFadeDistance(InLightDesc.LightFunctionFadeDistance)
+	, LightFunctionDisabledBrightness(InLightDesc.LightFunctionDisabledBrightness)
+	, LightFunctionMaterial(nullptr) // Initialized below
+	, IESTexture(InLightDesc.IESTexture)
+	, bContactShadowLengthInWS(InLightDesc.bContactShadowLengthInWS)
+	, bMovable(InLightDesc.bMovable)
+	, bStaticLighting(InLightDesc.bStaticLighting)
+	, bStaticShadowing(InLightDesc.bStaticShadowing)
+	, bCastDynamicShadow(InLightDesc.bCastDynamicShadow)
+	, bCastStaticShadow(InLightDesc.bCastStaticShadow)
+	, bCastTranslucentShadows(InLightDesc.bCastTranslucentShadows)
+	, bTransmission(InLightDesc.bTransmission)
+	, bCastVolumetricShadow(InLightDesc.bCastVolumetricShadow)
+	, bCastHairStrandsDeepShadow(InLightDesc.bCastHairStrandsDeepShadow)
+	, bCastShadowsFromCinematicObjectsOnly(InLightDesc.bCastShadowsFromCinematicObjectsOnly)
+	, bForceCachedShadowsForMovablePrimitives(InLightDesc.bForceCachedShadowsForMovablePrimitives)
+	, CastRaytracedShadow(InLightDesc.CastRaytracedShadow)
+	, bAffectReflection(InLightDesc.bAffectReflection)
+	, bAffectGlobalIllumination(InLightDesc.bAffectGlobalIllumination)
+	, bAffectTranslucentLighting(InLightDesc.bAffectTranslucentLighting)
+	, bUsedAsAtmosphereSunLight(InLightDesc.bUsedAsAtmosphereSunLight)
+	, bUseRayTracedDistanceFieldShadows(InLightDesc.bUseRayTracedDistanceFieldShadows)
+	, bUseVirtualShadowMaps(false) // Initialized below
 	, bCastModulatedShadows(false)
 	, bUseWholeSceneCSMForMovableObjects(false)
-	, bSelected(InLightComponent->IsSelected() || InLightComponent->IsOwnerSelected())
-	, bAllowMegaLights(InLightComponent->bAllowMegaLights)
-	, MegaLightsShadowMethod(InLightComponent->MegaLightsShadowMethod)
+	, bSelected(InLightDesc.bSelected)
+	, bAllowMegaLights(InLightDesc.bAllowMegaLights)
+	, bIsPrecomputedLightingValid(InLightDesc.bIsPrecomputedLightingValid)
+	, bEnableLightShaftBloom(InLightDesc.bEnableLightShaftBloom)
+	, BloomScale(InLightDesc.BloomScale)
+	, BloomThreshold(InLightDesc.BloomThreshold)
+	, BloomMaxBrightness(InLightDesc.BloomMaxBrightness)
+	, BloomTint(InLightDesc.BloomTint)
+	, MegaLightsShadowMethod(InLightDesc.MegaLightsShadowMethod)
 	, LightFunctionAtlasLightIndex(0)
-	, AtmosphereSunLightIndex(InLightComponent->GetAtmosphereSunLightIndex())
-	, AtmosphereSunDiskColorScale(InLightComponent->GetAtmosphereSunDiskColorScale())
-	, LightType(InLightComponent->GetLightType())
-	, LightingChannelMask(GetLightingChannelMaskForStruct(InLightComponent->LightingChannels))
-	, ViewLightingChannelMask(InLightComponent->ViewLightingChannels.GetMaskForStruct())
-	, StatId(InLightComponent->GetStatID(true))
-	, ComponentName(InLightComponent->GetFName())
-	, LevelName(InLightComponent->GetOwner() ? InLightComponent->GetOwner()->GetLevel()->GetOutermost()->GetFName() : NAME_None)
+	, AtmosphereSunLightIndex(InLightDesc.AtmosphereSunLightIndex)
+	, AtmosphereSunDiskColorScale(InLightDesc.AtmosphereSunDiskColorScale)
+	, LightType(InLightDesc.LightType)
+	, LightingChannelMask(InLightDesc.LightingChannelMask)
+	, ViewLightingChannelMask(InLightDesc.ViewLightingChannelMask)
+	, StatId(InLightDesc.GetStatID())
+	, ComponentName(InLightDesc.ComponentName)
+	, LevelName(InLightDesc.LevelName)
+	, VSMTexelDitherScale(1.0f)
+	, VSMResolutionLodBias(0.0f)
 	, FarShadowDistance(0)
 	, FarShadowCascadeCount(0)
+	, ModulatedShadowColor(EForceInit::ForceInitToZero)
 	, ShadowAmount(1.0f)
-	, SamplesPerPixel(1)
-	, DeepShadowLayerDistribution(InLightComponent->DeepShadowLayerDistribution)
+	, SamplesPerPixel(InLightDesc.SamplesPerPixel)
+	, DeepShadowLayerDistribution(InLightDesc.DeepShadowLayerDistribution)
 	, IESAtlasId(~0u)
 #if ACTOR_HAS_LABELS
-	, OwnerNameOrLabel(InLightComponent->GetOwner() ? InLightComponent->GetOwner()->GetActorNameOrLabel() : InLightComponent->GetName())
+	, OwnerNameOrLabel(InLightDesc.OwnerNameOrLabel)
 #endif
 {
 	check(SceneInterface);
+	check(LightComponentId.IsValid());
 
 	// Currently we use virtual shadows maps for all lights when the global setting is enabled
 	bUseVirtualShadowMaps = ::UseVirtualShadowMaps(SceneInterface->GetShaderPlatform(), SceneInterface->GetFeatureLevel());
+	
 	// Treat stationary lights as movable when VSMs are enabled
-	if (bUseVirtualShadowMaps)
+	if (bUseVirtualShadowMaps || !IsStaticLightingAllowed())
 	{
 		bStaticShadowing = bStaticLighting;
 	}
 
-	if (!IsStaticLightingAllowed())
+	if ((InLightDesc.MapBuildDataShadowMapChannel != INDEX_NONE) && bStaticShadowing && !bStaticLighting)
 	{
-		bStaticShadowing = bStaticLighting;
-	}
-
-	const FLightComponentMapBuildData* MapBuildData = InLightComponent->GetLightComponentMapBuildData();
-
-	if (MapBuildData && bStaticShadowing && !bStaticLighting)
-	{
-		ShadowMapChannel = MapBuildData->ShadowMapChannel;
+		ShadowMapChannel = InLightDesc.MapBuildDataShadowMapChannel;
 	}
 	else
 	{
@@ -104,33 +123,14 @@ FLightSceneProxy::FLightSceneProxy(const ULightComponent* InLightComponent)
 	}
 
 	// Use the preview channel if valid, otherwise fallback to the lighting build channel
-	PreviewShadowMapChannel = InLightComponent->PreviewShadowMapChannel != INDEX_NONE ? InLightComponent->PreviewShadowMapChannel : ShadowMapChannel;
-
-	StaticShadowDepthMap = &LightComponent->StaticShadowDepthMap;
+	PreviewShadowMapChannel = InLightDesc.PreviewShadowMapChannel != INDEX_NONE ? InLightDesc.PreviewShadowMapChannel : ShadowMapChannel;
 
-	if (LightComponent->IESTexture)
+	if (InLightDesc.LightFunctionMaterial &&
+		InLightDesc.LightFunctionMaterial->GetMaterial()->MaterialDomain == MD_LightFunction)
 	{
-		IESTexture = LightComponent->IESTexture;
+		LightFunctionMaterial = InLightDesc.LightFunctionMaterial->GetRenderProxy();
 	}
 
-	Color = LightComponent->GetColoredLightBrightness();
-
-	if (LightComponent->LightFunctionMaterial &&
-		LightComponent->LightFunctionMaterial->GetMaterial()->MaterialDomain == MD_LightFunction)
-	{
-		LightFunctionMaterial = LightComponent->LightFunctionMaterial->GetRenderProxy();
-	}
-	else
-	{
-		LightFunctionMaterial = NULL;
-	}
-
-	LightFunctionScale = LightComponent->LightFunctionScale;
-	LightFunctionFadeDistance = LightComponent->LightFunctionFadeDistance;
-	LightFunctionDisabledBrightness = LightComponent->DisabledBrightness;
-
-	SamplesPerPixel = LightComponent->SamplesPerPixel;
-
 	if (bCastDynamicShadow && IsMobilePlatform(SceneInterface->GetShaderPlatform()))
 	{
 		if ((GetLightType() == LightType_Point && !DoesRuntimeSupportOnePassPointLightShadows(SceneInterface->GetShaderPlatform()))
@@ -142,8 +142,6 @@ FLightSceneProxy::FLightSceneProxy(const ULightComponent* InLightComponent)
 
 		bTransmission = false;
 	}
-	VSMTexelDitherScale = 1.0f;
-	VSMResolutionLodBias = 0.0f;
 }
 
 FLightSceneProxy::~FLightSceneProxy() = default;
diff --git a/Engine/Source/Runtime/Engine/Private/LightSceneProxyDesc.cpp b/Engine/Source/Runtime/Engine/Private/LightSceneProxyDesc.cpp
new file mode 100644
--- /dev/null
+++ b/Engine/Source/Runtime/Engine/Private/LightSceneProxyDesc.cpp
@@ -0,0 +1,78 @@
+// Copyright Epic Games, Inc. All Rights Reserved.
+
+#include "LightSceneProxyDesc.h"
+#include "Engine/Level.h"
+#include "Engine/MapBuildDataRegistry.h"
+#include "Components/LightComponent.h"
+#include "UObject/Package.h"
+
+void FLightSceneProxyDesc::Serialize(FArchive& Ar)
+{
+	FArchive_Serialize_BitfieldBool(Ar, bContactShadowLengthInWS);
+	FArchive_Serialize_BitfieldBool(Ar, bMovable);
+	FArchive_Serialize_BitfieldBool(Ar, bStaticLighting);
+	FArchive_Serialize_BitfieldBool(Ar, bStaticShadowing);
+	FArchive_Serialize_BitfieldBool(Ar, bCastDynamicShadow);
+	FArchive_Serialize_BitfieldBool(Ar, bCastStaticShadow);
+	FArchive_Serialize_BitfieldBool(Ar, bCastTranslucentShadows);
+	FArchive_Serialize_BitfieldBool(Ar, bTransmission);
+	FArchive_Serialize_BitfieldBool(Ar, bCastVolumetricShadow);
+	FArchive_Serialize_BitfieldBool(Ar, bCastHairStrandsDeepShadow);
+	FArchive_Serialize_BitfieldBool(Ar, bCastShadowsFromCinematicObjectsOnly);
+	FArchive_Serialize_BitfieldBool(Ar, bForceCachedShadowsForMovablePrimitives);
+	FArchive_Serialize_BitfieldBool(Ar, bAffectReflection);
+	FArchive_Serialize_BitfieldBool(Ar, bAffectGlobalIllumination);
+	FArchive_Serialize_BitfieldBool(Ar, bAffectTranslucentLighting);
+	FArchive_Serialize_BitfieldBool(Ar, bUsedAsAtmosphereSunLight);
+	FArchive_Serialize_BitfieldBool(Ar, bUseRayTracedDistanceFieldShadows);
+	FArchive_Serialize_BitfieldBool(Ar, bAllowMegaLights);
+	FArchive_Serialize_BitfieldBool(Ar, bIsPrecomputedLightingValid);
+	FArchive_Serialize_BitfieldBool(Ar, bEnableLightShaftBloom);
+	Ar << Color;
+	Ar << IndirectLightingScale;
+	Ar << VolumetricScatteringIntensity;
+	Ar << RayEndBias;
+	Ar << ShadowResolutionScale;
+	Ar << ShadowBias;
+	Ar << ShadowSlopeBias;
+	Ar << ShadowSharpen;
+	Ar << ContactShadowLength;
+	Ar << ContactShadowCastingIntensity;
+	Ar << ContactShadowNonCastingIntensity;
+	Ar << SpecularScale;
+	Ar << DiffuseScale;
+	Ar << LightGuid;
+	Ar << MapBuildDataShadowMapChannel;
+	Ar << PreviewShadowMapChannel;
+	Ar << RayStartOffsetDepthScale;
+	Ar << LightFunctionScale;
+	Ar << LightFunctionFadeDistance;
+	Ar << LightFunctionDisabledBrightness;
+	Ar << LightFunctionMaterial;
+	Ar << IESTexture;
+	Ar << CastRaytracedShadow;
+	Ar << BloomScale;
+	Ar << BloomThreshold;
+	Ar << BloomMaxBrightness;
+	Ar << BloomTint;
+	Ar << MegaLightsShadowMethod;
+	Ar << AtmosphereSunLightIndex;
+	Ar << AtmosphereSunDiskColorScale;
+	Ar << LightType;
+	Ar << LightingChannelMask;
+	Ar << ViewLightingChannelMask;
+	Ar << ComponentName;
+	Ar << LevelName;
+	Ar << SamplesPerPixel;
+	Ar << DeepShadowLayerDistribution;
+}
+
+const FStaticShadowDepthMap* FLightSceneProxyDesc::GetStaticShadowDepthMap() const
+{
+	return LightComponent ? LightComponent->GetStaticShadowDepthMap() : nullptr;
+}
+
+TStatId FLightSceneProxyDesc::GetStatID() const
+{
+	return LightComponent ? LightComponent->GetStatID(true) : TStatId();
+}
\ No newline at end of file
diff --git a/Engine/Source/Runtime/Engine/Private/LocalLightSceneProxy.cpp b/Engine/Source/Runtime/Engine/Private/LocalLightSceneProxy.cpp
--- a/Engine/Source/Runtime/Engine/Private/LocalLightSceneProxy.cpp
+++ b/Engine/Source/Runtime/Engine/Private/LocalLightSceneProxy.cpp
@@ -5,6 +5,7 @@
 =============================================================================*/
 
 #include "LocalLightSceneProxy.h"
+#include "LocalLightSceneProxyDesc.h"
 #include "Components/LocalLightComponent.h"
 #include "Math/InverseRotationMatrix.h"
 #include "SceneManagement.h"
@@ -12,21 +13,29 @@
 
 /** Initialization constructor. */
 FLocalLightSceneProxy::FLocalLightSceneProxy(const ULocalLightComponent* Component)
-	: FLightSceneProxy(Component)
-	, MaxDrawDistance(Component->MaxDrawDistance)
-	, FadeRange(Component->MaxDistanceFadeRange)
-	, InverseExposureBlend(Component->InverseExposureBlend)
+	: FLocalLightSceneProxy(ULocalLightComponent::BuildSceneProxyDesc(*Component))
 {
-	UpdateRadius(Component->AttenuationRadius);
 }
 
-/**
-* Called on the light scene info after it has been passed to the rendering thread to update the rendering thread's cached info when
-* the light's radius changes.
-*/
-void UpdateRadius_GameThread(float Radius);
+FLocalLightSceneProxy::FLocalLightSceneProxy(const FLocalLightSceneProxyDesc& LightDesc)
+	: FLightSceneProxy(LightDesc)
+	, MaxDrawDistance(LightDesc.MaxDrawDistance)
+	, FadeRange(LightDesc.MaxDistanceFadeRange)
+	, InverseExposureBlend(LightDesc.InverseExposureBlend)
+{
+	UpdateRadius(LightDesc.AttenuationRadius);
+}
+
+void FLocalLightSceneProxy::UpdateRadius_GameThread(float InRadius)
+{
+	FLocalLightSceneProxy* InLightSceneInfo = this;
+	ENQUEUE_RENDER_COMMAND(UpdateRadius)(
+		[InLightSceneInfo, InRadius](FRHICommandList& RHICmdList)
+		{
+			InLightSceneInfo->UpdateRadius(InRadius);
+		});
+}
 
-// FLightSceneInfo interface.
 float FLocalLightSceneProxy::GetMaxDrawDistance() const
 {
 	return MaxDrawDistance;
@@ -98,16 +107,6 @@ float FLocalLightSceneProxy::GetEffectiveScreenRadius(const FViewMatrices& Shado
 	return ScreenScale * GetRadius() / FMath::Max(LightDistance, 1.0f);
 }
 
-float FLocalLightSceneProxy::GetEffectiveScreenRadius(const FViewMatrices& ShadowViewMatrices) const
-{
-	// Use the distance from the view origin to the light to approximate perspective projection
-	// We do not use projected screen position since it causes problems when the light is behind the camera
-
-	const float LightDistance = (GetOrigin() - ShadowViewMatrices.GetViewOrigin()).Size();
-
-	return ShadowViewMatrices.GetScreenScale() * GetRadius() / FMath::Max(LightDistance, 1.0f);
-}
-
 FVector FLocalLightSceneProxy::GetPerObjectProjectedShadowProjectionPoint(const FBoxSphereBounds& SubjectBounds) const
 {
 	return GetOrigin();
diff --git a/Engine/Source/Runtime/Engine/Private/LocalLightSceneProxyDesc.cpp b/Engine/Source/Runtime/Engine/Private/LocalLightSceneProxyDesc.cpp
new file mode 100644
--- /dev/null
+++ b/Engine/Source/Runtime/Engine/Private/LocalLightSceneProxyDesc.cpp
@@ -0,0 +1,13 @@
+// Copyright Epic Games, Inc. All Rights Reserved.
+
+#include "LocalLightSceneProxyDesc.h"
+
+void FLocalLightSceneProxyDesc::Serialize(FArchive& Ar)
+{
+	Super::Serialize(Ar);
+
+	Ar << MaxDrawDistance;
+	Ar << MaxDistanceFadeRange;
+	Ar << InverseExposureBlend;
+	Ar << AttenuationRadius;
+}
\ No newline at end of file
diff --git a/Engine/Source/Runtime/Engine/Private/PointLightSceneProxy.cpp b/Engine/Source/Runtime/Engine/Private/PointLightSceneProxy.cpp
new file mode 100644
--- /dev/null
+++ b/Engine/Source/Runtime/Engine/Private/PointLightSceneProxy.cpp
@@ -0,0 +1,103 @@
+// Copyright Epic Games, Inc. All Rights Reserved.
+
+#include "PointLightSceneProxy.h"
+#include "PointLightSceneProxyDesc.h"
+#include "Components/PointLightComponent.h"
+#include "SceneInterface.h"
+#include "SceneManagement.h"
+#include "SceneView.h"
+
+int32 GAllowPointLightCubemapShadows = 1;
+static FAutoConsoleVariableRef CVarAllowPointLightCubemapShadows(
+	TEXT("r.AllowPointLightCubemapShadows"),
+	GAllowPointLightCubemapShadows,
+	TEXT("When 0, will prevent point light cube map shadows from being used and the light will be unshadowed.")
+);
+
+FPointLightSceneProxy::FPointLightSceneProxy(const UPointLightComponent* Component)
+	: FPointLightSceneProxy(UPointLightComponent::BuildSceneProxyDesc(*Component))
+{
+}
+
+FPointLightSceneProxy::FPointLightSceneProxy(const FPointLightSceneProxyDesc& LightDesc)
+	: FLocalLightSceneProxy(LightDesc)
+	, FalloffExponent(LightDesc.LightFalloffExponent)
+	, SourceRadius(LightDesc.SourceRadius)
+	, SoftSourceRadius(LightDesc.SoftSourceRadius)
+	, SourceLength(LightDesc.SourceLength)
+	, bInverseSquared(LightDesc.bUseInverseSquaredFalloff)
+{
+}
+
+float FPointLightSceneProxy::GetSourceRadius() const 
+{
+	return SourceRadius;
+}
+
+bool FPointLightSceneProxy::IsInverseSquared() const 
+{
+	return bInverseSquared;
+}
+
+/** Accesses parameters needed for rendering the light. */
+void FPointLightSceneProxy::GetLightShaderParameters(FLightRenderParameters& LightParameters, uint32 Flags) const
+{
+	LightParameters.WorldPosition = GetOrigin();
+	LightParameters.InvRadius = InvRadius;
+	LightParameters.Color = GetColor();
+	LightParameters.FalloffExponent = FalloffExponent;
+
+	// TODO LWC - GetDirection() seems like it needs to be normalized, somehow accumulating error with large-scale position values
+	LightParameters.Direction = (FVector3f)-GetDirection(); // LWC_TODO: Precision Loss
+	LightParameters.Tangent = FVector3f(WorldToLight.M[0][2], WorldToLight.M[1][2], WorldToLight.M[2][2]);
+	LightParameters.SpotAngles = FVector2f(-2.0f, 1.0f);
+	LightParameters.SpecularScale = FMath::Clamp(SpecularScale, 0.f, 1.f);
+	LightParameters.DiffuseScale = FMath::Clamp(DiffuseScale, 0.f, 1.f);
+	LightParameters.SourceRadius = SourceRadius;
+	LightParameters.SoftSourceRadius = SoftSourceRadius;
+	LightParameters.SourceLength = SourceLength;
+	LightParameters.RectLightBarnCosAngle = 0.0f;
+	LightParameters.RectLightBarnLength = -2.0f;
+	LightParameters.RectLightAtlasUVOffset = FVector2f::ZeroVector;
+	LightParameters.RectLightAtlasUVScale = FVector2f::ZeroVector;
+	LightParameters.RectLightAtlasMaxLevel = FLightRenderParameters::GetRectLightAtlasInvalidMIPLevel();
+	LightParameters.IESAtlasIndex = INDEX_NONE;
+	LightParameters.InverseExposureBlend = InverseExposureBlend;
+	LightParameters.LightFunctionAtlasLightIndex = GetLightFunctionAtlasLightIndex();
+	LightParameters.bAffectsTranslucentLighting = AffectsTranslucentLighting() ? 1 : 0;
+
+	if (IESAtlasId != ~0)
+	{
+		GetSceneInterface()->GetLightIESAtlasSlot(this, &LightParameters);
+	}
+}
+
+/**
+* Sets up a projected shadow initializer for shadows from the entire scene.
+* @return True if the whole-scene projected shadow should be used.
+*/
+bool FPointLightSceneProxy::GetWholeSceneProjectedShadowInitializer(const FSceneViewFamily& ViewFamily, TArray<FWholeSceneProjectedShadowInitializer, TInlineAllocator<6> >& OutInitializers) const
+{
+	if ((ViewFamily.GetFeatureLevel() >= ERHIFeatureLevel::SM5 || DoesRuntimeSupportOnePassPointLightShadows(ViewFamily.GetShaderPlatform()))
+		&& GAllowPointLightCubemapShadows != 0)
+	{
+		FWholeSceneProjectedShadowInitializer& OutInitializer = *new(OutInitializers) FWholeSceneProjectedShadowInitializer;
+		OutInitializer.PreShadowTranslation = -GetLightToWorld().GetOrigin();
+		OutInitializer.WorldToLight = GetWorldToLight().RemoveTranslation();
+		OutInitializer.Scales = FVector2D(1, 1);
+		OutInitializer.SubjectBounds = FBoxSphereBounds(FVector(0, 0, 0), FVector(Radius, Radius, Radius), Radius);
+		OutInitializer.WAxis = FVector4(0, 0, 1, 0);
+		OutInitializer.MinLightW = 0.1f;
+		OutInitializer.MaxDistanceToCastInLightW = Radius;
+		OutInitializer.bOnePassPointLightShadow = true;
+		OutInitializer.bRayTracedDistanceField = UseRayTracedDistanceFieldShadows() && DoesPlatformSupportDistanceFieldShadowing(ViewFamily.GetShaderPlatform());
+		return true;
+	}
+
+	return false;
+}
+
+FVector FPointLightSceneProxy::GetPerObjectProjectedShadowProjectionPoint(const FBoxSphereBounds& SubjectBounds) const
+{
+	return FMath::ClosestPointOnSegment(SubjectBounds.Origin, GetOrigin() - GetDirection() * SourceLength * 0.5f, GetOrigin() + GetDirection() * SourceLength * 0.5f);
+}
\ No newline at end of file
diff --git a/Engine/Source/Runtime/Engine/Private/PointLightSceneProxy.h b/Engine/Source/Runtime/Engine/Private/PointLightSceneProxy.h
deleted file mode 100644
diff --git a/Engine/Source/Runtime/Engine/Private/PointLightSceneProxyDesc.cpp b/Engine/Source/Runtime/Engine/Private/PointLightSceneProxyDesc.cpp
new file mode 100644
--- /dev/null
+++ b/Engine/Source/Runtime/Engine/Private/PointLightSceneProxyDesc.cpp
@@ -0,0 +1,14 @@
+// Copyright Epic Games, Inc. All Rights Reserved.
+
+#include "PointLightSceneProxyDesc.h"
+
+void FPointLightSceneProxyDesc::Serialize(FArchive& Ar)
+{
+	Super::Serialize(Ar);
+
+	Ar << LightFalloffExponent;
+	Ar << SourceRadius;
+	Ar << SoftSourceRadius;
+	Ar << SourceLength;
+	Ar << bUseInverseSquaredFalloff;
+}
\ No newline at end of file
diff --git a/Engine/Source/Runtime/Engine/Private/RectLightSceneProxy.cpp b/Engine/Source/Runtime/Engine/Private/RectLightSceneProxy.cpp
new file mode 100644
--- /dev/null
+++ b/Engine/Source/Runtime/Engine/Private/RectLightSceneProxy.cpp
@@ -0,0 +1,127 @@
+// Copyright Epic Games, Inc. All Rights Reserved.
+
+#include "RectLightSceneProxy.h"
+#include "RectLightSceneProxyDesc.h"
+#include "Components/RectLightComponent.h"
+#include "SceneManagement.h"
+#include "SceneInterface.h"
+#include "SceneView.h"
+
+extern int32 GAllowPointLightCubemapShadows;
+
+FRectLightSceneProxy::FRectLightSceneProxy(const URectLightComponent* Component)
+	: FRectLightSceneProxy(URectLightComponent::BuildSceneProxyDesc(*Component))
+{
+}
+
+FRectLightSceneProxy::FRectLightSceneProxy(const FRectLightSceneProxyDesc& LightDesc)
+	: FLocalLightSceneProxy(LightDesc)
+	, SourceWidth(FMath::Max(1, LightDesc.SourceWidth))
+	, SourceHeight(FMath::Max(1, LightDesc.SourceHeight))
+	, BarnDoorAngle(FMath::Clamp(LightDesc.BarnDoorAngle, 0.f, GetRectLightBarnDoorMaxAngle()))
+	, BarnDoorLength(FMath::Max(0.1f, LightDesc.BarnDoorLength))
+	, SourceTexture(LightDesc.SourceTexture)
+	, RectAtlasId(~0u)
+	, LightFunctionConeAngleTangent(LightDesc.LightFunctionConeAngle > 0 ? FMath::Tan(FMath::Clamp(LightDesc.LightFunctionConeAngle, 0.0f, 89.0f) * (float)UE_PI / 180.0f) : 0.0f)
+	, SourceTextureScaleOffset(FVector4f(
+		FMath::Clamp(LightDesc.SourceTextureScale.X, 0.f, 1.f),
+		FMath::Clamp(LightDesc.SourceTextureScale.Y, 0.f, 1.f),
+		FMath::Clamp(LightDesc.SourceTextureOffset.X, 0.f, 1.f),
+		FMath::Clamp(LightDesc.SourceTextureOffset.Y, 0.f, 1.f)))
+{
+}
+
+FRectLightSceneProxy::~FRectLightSceneProxy() {}
+
+bool FRectLightSceneProxy::IsRectLight() const
+{
+	return true;
+}
+
+bool FRectLightSceneProxy::HasSourceTexture() const
+{
+	return SourceTexture != nullptr;
+}
+
+/** Accesses parameters needed for rendering the light. */
+void FRectLightSceneProxy::GetLightShaderParameters(FLightRenderParameters& LightParameters, uint32 Flags) const
+{
+	FLinearColor LightColor = GetColor();
+	LightColor /= 0.5f * SourceWidth * SourceHeight;
+	LightParameters.WorldPosition = GetOrigin();
+	LightParameters.InvRadius = InvRadius;
+	LightParameters.Color = LightColor;
+	LightParameters.FalloffExponent = 0.0f;
+
+	LightParameters.Direction = FVector3f(-GetDirection());
+	LightParameters.Tangent = FVector3f(WorldToLight.M[0][2], WorldToLight.M[1][2], WorldToLight.M[2][2]);
+	LightParameters.SpotAngles = FVector2f(-2.0f, 1.0f);
+	LightParameters.SpecularScale = FMath::Clamp(SpecularScale, 0.f, 1.f);
+	LightParameters.DiffuseScale = FMath::Clamp(DiffuseScale, 0.f, 1.f);
+	LightParameters.SourceRadius = SourceWidth * 0.5f;
+	LightParameters.SoftSourceRadius = 0.0f;
+	LightParameters.SourceLength = SourceHeight * 0.5f;
+	LightParameters.RectLightBarnCosAngle = FMath::Cos(FMath::DegreesToRadians(BarnDoorAngle));
+	LightParameters.RectLightBarnLength = BarnDoorLength;
+	LightParameters.RectLightAtlasUVOffset = FVector2f::ZeroVector;
+	LightParameters.RectLightAtlasUVScale = FVector2f::ZeroVector;
+	LightParameters.RectLightAtlasMaxLevel = FLightRenderParameters::GetRectLightAtlasInvalidMIPLevel();
+	LightParameters.IESAtlasIndex = INDEX_NONE;
+	LightParameters.InverseExposureBlend = InverseExposureBlend;
+	LightParameters.LightFunctionAtlasLightIndex = GetLightFunctionAtlasLightIndex();
+	LightParameters.bAffectsTranslucentLighting = AffectsTranslucentLighting() ? 1 : 0;
+
+	if (IESAtlasId != ~0)
+	{
+		GetSceneInterface()->GetLightIESAtlasSlot(this, &LightParameters);
+	}
+
+	if (RectAtlasId != ~0u)
+	{
+		GetSceneInterface()->GetRectLightAtlasSlot(this, &LightParameters);
+	}
+
+	// Render RectLight approximately as SpotLight if the requester does not support rect light (e.g., translucent light grid or mobile)
+	if (!!(Flags & ELightShaderParameterFlags::RectAsSpotLight))
+	{
+		float ClampedOuterConeAngle = FMath::DegreesToRadians(89.001f);
+		float ClampedInnerConeAngle = FMath::DegreesToRadians(70.0f);
+		float CosOuterCone = FMath::Cos(ClampedOuterConeAngle);
+		float CosInnerCone = FMath::Cos(ClampedInnerConeAngle);
+		float InvCosConeDifference = 1.0f / (CosInnerCone - CosOuterCone);
+
+		LightParameters.Color = GetColor();
+		LightParameters.FalloffExponent = 8.0f;
+		LightParameters.SpotAngles = FVector2f(CosOuterCone, InvCosConeDifference);
+		LightParameters.SourceRadius = (SourceWidth + SourceHeight) * 0.5 * 0.5f;
+		LightParameters.SourceLength = 0.0f;
+		LightParameters.RectLightBarnCosAngle = 0.0f;
+		LightParameters.RectLightBarnLength = -2.0f;
+	}
+}
+
+/**
+* Sets up a projected shadow initializer for shadows from the entire scene.
+* @return True if the whole-scene projected shadow should be used.
+*/
+bool FRectLightSceneProxy::GetWholeSceneProjectedShadowInitializer(const FSceneViewFamily& ViewFamily, TArray<FWholeSceneProjectedShadowInitializer, TInlineAllocator<6> >& OutInitializers) const
+{
+	if (ViewFamily.GetFeatureLevel() >= ERHIFeatureLevel::SM5
+		&& GAllowPointLightCubemapShadows != 0)
+	{
+		FWholeSceneProjectedShadowInitializer& OutInitializer = *new(OutInitializers) FWholeSceneProjectedShadowInitializer;
+		OutInitializer.PreShadowTranslation = -GetLightToWorld().GetOrigin();
+		OutInitializer.WorldToLight = GetWorldToLight().RemoveTranslation();
+		OutInitializer.Scales = FVector2D(1, 1);
+		OutInitializer.SubjectBounds = FBoxSphereBounds(FVector(0, 0, 0), FVector(Radius, Radius, Radius), Radius);
+		OutInitializer.WAxis = FVector4(0, 0, 1, 0);
+		OutInitializer.MinLightW = 0.1f;
+		OutInitializer.MaxDistanceToCastInLightW = Radius;
+		OutInitializer.bOnePassPointLightShadow = true;
+
+		OutInitializer.bRayTracedDistanceField = UseRayTracedDistanceFieldShadows() && DoesPlatformSupportDistanceFieldShadowing(ViewFamily.GetShaderPlatform());
+		return true;
+	}
+
+	return false;
+}
\ No newline at end of file
diff --git a/Engine/Source/Runtime/Engine/Private/RectLightSceneProxyDesc.cpp b/Engine/Source/Runtime/Engine/Private/RectLightSceneProxyDesc.cpp
new file mode 100644
--- /dev/null
+++ b/Engine/Source/Runtime/Engine/Private/RectLightSceneProxyDesc.cpp
@@ -0,0 +1,17 @@
+// Copyright Epic Games, Inc. All Rights Reserved.
+
+#include "RectLightSceneProxyDesc.h"
+
+void FRectLightSceneProxyDesc::Serialize(FArchive& Ar)
+{
+	FLocalLightSceneProxyDesc::Serialize(Ar);
+
+	Ar << SourceWidth;
+	Ar << SourceHeight;
+	Ar << BarnDoorAngle;
+	Ar << BarnDoorLength;
+	Ar << LightFunctionConeAngle;
+	Ar << SourceTexture;
+	Ar << SourceTextureScale;
+	Ar << SourceTextureOffset;
+}
diff --git a/Engine/Source/Runtime/Engine/Private/SpotLightSceneProxy.cpp b/Engine/Source/Runtime/Engine/Private/SpotLightSceneProxy.cpp
new file mode 100644
--- /dev/null
+++ b/Engine/Source/Runtime/Engine/Private/SpotLightSceneProxy.cpp
@@ -0,0 +1,135 @@
+// Copyright Epic Games, Inc. All Rights Reserved.
+
+#include "SpotLightSceneProxy.h"
+#include "SpotLightSceneProxyDesc.h"
+#include "Components/SpotLightComponent.h"
+#include "SceneInterface.h"
+#include "SceneManagement.h"
+#include "SceneView.h"
+
+FSpotLightSceneProxy::FSpotLightSceneProxy(const USpotLightComponent* Component)
+	: FSpotLightSceneProxy(USpotLightComponent::BuildSceneProxyDesc(*Component))
+{
+}
+
+FSpotLightSceneProxy::FSpotLightSceneProxy(const FSpotLightSceneProxyDesc& LightDesc)
+	: FPointLightSceneProxy(LightDesc)
+{
+	const FVector2f ClampedConeAngles = USpotLightComponent::GetClampedConeAngles(LightDesc.InnerConeAngle, LightDesc.OuterConeAngle);
+	OuterConeAngle = ClampedConeAngles.Y;
+	CosOuterCone = FMath::Cos(ClampedConeAngles.Y);
+	SinOuterCone = FMath::Sin(ClampedConeAngles.Y);
+	CosInnerCone = FMath::Cos(ClampedConeAngles.X);
+	InvCosConeDifference = 1.0f / (CosInnerCone - CosOuterCone);
+	InvTanOuterCone = 1.0f / FMath::Tan(ClampedConeAngles.Y);
+}
+
+void FSpotLightSceneProxy::GetLightShaderParameters(FLightRenderParameters& LightParameters, uint32 Flags) const
+{
+	LightParameters.WorldPosition = GetOrigin();
+	LightParameters.InvRadius = InvRadius;
+	LightParameters.Color = GetColor();
+	LightParameters.FalloffExponent = FalloffExponent;
+	LightParameters.Direction = FVector3f(-GetDirection());
+	LightParameters.Tangent = FVector3f(WorldToLight.M[0][2], WorldToLight.M[1][2], WorldToLight.M[2][2]);
+	LightParameters.SpotAngles = FVector2f(CosOuterCone, InvCosConeDifference);
+	LightParameters.SpecularScale = FMath::Clamp(SpecularScale, 0.f, 1.f);
+	LightParameters.DiffuseScale = FMath::Clamp(DiffuseScale, 0.f, 1.f);
+	LightParameters.SourceRadius = SourceRadius;
+	LightParameters.SoftSourceRadius = SoftSourceRadius;
+	LightParameters.SourceLength = SourceLength;
+	LightParameters.RectLightBarnCosAngle = 0.0f;
+	LightParameters.RectLightBarnLength = -2.0f;
+	LightParameters.RectLightAtlasUVOffset = FVector2f::ZeroVector;
+	LightParameters.RectLightAtlasUVScale = FVector2f::ZeroVector;
+	LightParameters.RectLightAtlasMaxLevel = FLightRenderParameters::GetRectLightAtlasInvalidMIPLevel();
+	LightParameters.IESAtlasIndex = INDEX_NONE;
+	LightParameters.LightFunctionAtlasLightIndex = GetLightFunctionAtlasLightIndex();
+	LightParameters.bAffectsTranslucentLighting = AffectsTranslucentLighting() ? 1 : 0;
+	LightParameters.InverseExposureBlend = InverseExposureBlend;
+
+	if (IESAtlasId != uint32(INDEX_NONE))
+	{
+		GetSceneInterface()->GetLightIESAtlasSlot(this, &LightParameters);
+	}
+}
+
+bool FSpotLightSceneProxy::AffectsBounds(const FBoxSphereBounds& Bounds) const
+{
+	if (!FLocalLightSceneProxy::AffectsBounds(Bounds))
+	{
+		return false;
+	}
+
+	FVector	U = GetOrigin() - (Bounds.SphereRadius / SinOuterCone) * GetDirection(),
+		D = Bounds.Origin - U;
+	float	dsqr = D | D,
+		E = GetDirection() | D;
+	if (E > 0.0f && E * E >= dsqr * FMath::Square(CosOuterCone))
+	{
+		D = Bounds.Origin - GetOrigin();
+		dsqr = D | D;
+		E = -(GetDirection() | D);
+		if (E > 0.0f && E * E >= dsqr * FMath::Square(SinOuterCone))
+			return dsqr <= FMath::Square(Bounds.SphereRadius);
+		else
+			return true;
+	}
+
+	return false;
+}
+
+/**
+ * Sets up a projected shadow initializer for shadows from the entire scene.
+ * @return True if the whole-scene projected shadow should be used.
+ */
+bool FSpotLightSceneProxy::GetWholeSceneProjectedShadowInitializer(const FSceneViewFamily& ViewFamily, TArray<FWholeSceneProjectedShadowInitializer, TInlineAllocator<6> >& OutInitializers) const
+{
+	FWholeSceneProjectedShadowInitializer& OutInitializer = OutInitializers.AddDefaulted_GetRef();
+	OutInitializer.PreShadowTranslation = -GetLightToWorld().GetOrigin();
+	OutInitializer.WorldToLight = GetWorldToLight().RemoveTranslation();
+	OutInitializer.Scales = FVector2D(InvTanOuterCone, InvTanOuterCone);
+
+	const FSphere AbsoluteBoundingSphere = FSpotLightSceneProxy::GetBoundingSphere();
+	OutInitializer.SubjectBounds = FBoxSphereBounds(
+		AbsoluteBoundingSphere.Center - GetOrigin(),
+		FVector(AbsoluteBoundingSphere.W, AbsoluteBoundingSphere.W, AbsoluteBoundingSphere.W),
+		AbsoluteBoundingSphere.W
+	);
+
+	OutInitializer.WAxis = FVector4(0, 0, 1, 0);
+	OutInitializer.MinLightW = 0.1f;
+	OutInitializer.MaxDistanceToCastInLightW = Radius;
+	OutInitializer.bRayTracedDistanceField = UseRayTracedDistanceFieldShadows() && DoesPlatformSupportDistanceFieldShadowing(ViewFamily.GetShaderPlatform());
+	return true;
+}
+
+float FSpotLightSceneProxy::GetOuterConeAngle() const 
+{
+	return OuterConeAngle;
+}
+
+FSphere FSpotLightSceneProxy::GetBoundingSphere() const
+{
+	return FMath::ComputeBoundingSphereForCone(GetOrigin(), GetDirection(), (FSphere::FReal)Radius, (FSphere::FReal)CosOuterCone, (FSphere::FReal)SinOuterCone);
+}
+
+float FSpotLightSceneProxy::GetEffectiveScreenRadius(const FViewMatrices& ShadowViewMatrices, const FIntPoint& CameraViewRectSize) const
+{
+	// Heuristic: use the radius of the inscribed sphere at the cone's end as the light's effective screen radius
+	// We do so because we do not want to use the light's radius directly, which will make us overestimate the shadow map resolution greatly for a spot light
+
+	// In the correct form,
+	//   InscribedSpherePosition = GetOrigin() + GetDirection() * GetRadius() / CosOuterCone
+	//   InscribedSphereRadius = GetRadius() / SinOuterCone
+	// Do it incorrectly to avoid division which is more expensive and risks division by zero
+	const FVector InscribedSpherePosition = GetOrigin() + GetDirection() * GetRadius() * CosOuterCone;
+	const float InscribedSphereRadius = GetRadius() * SinOuterCone;
+
+	const float SphereDistanceFromViewOrigin = (InscribedSpherePosition - ShadowViewMatrices.GetViewOrigin()).Size();
+
+	const FVector2D& ProjectionScale = ShadowViewMatrices.GetProjectionScale();
+	const float ScreenScale = FMath::Max(CameraViewRectSize.X * 0.5f * ProjectionScale.X, CameraViewRectSize.Y * 0.5f * ProjectionScale.Y);
+
+	return ScreenScale * InscribedSphereRadius / FMath::Max(SphereDistanceFromViewOrigin, 1.0f);
+}
\ No newline at end of file
diff --git a/Engine/Source/Runtime/Engine/Private/SpotLightSceneProxyDesc.cpp b/Engine/Source/Runtime/Engine/Private/SpotLightSceneProxyDesc.cpp
new file mode 100644
--- /dev/null
+++ b/Engine/Source/Runtime/Engine/Private/SpotLightSceneProxyDesc.cpp
@@ -0,0 +1,11 @@
+// Copyright Epic Games, Inc. All Rights Reserved.
+
+#include "SpotLightSceneProxyDesc.h"
+
+void FSpotLightSceneProxyDesc::Serialize(FArchive& Ar)
+{
+	Super::Serialize(Ar);
+
+	Ar << InnerConeAngle;
+	Ar << OuterConeAngle;
+}
\ No newline at end of file
diff --git a/Engine/Source/Runtime/Engine/Public/LightComponentId.h b/Engine/Source/Runtime/Engine/Public/LightComponentId.h
new file mode 100644
--- /dev/null
+++ b/Engine/Source/Runtime/Engine/Public/LightComponentId.h
@@ -0,0 +1,50 @@
+// Copyright Epic Games, Inc. All Rights Reserved.
+
+#pragma once
+
+#include "Templates/TypeHash.h"
+
+/**
+ * Class used to identify ULightComponents on the rendering thread without having to pass the pointer around,
+ * Which would make it easy for people to access game thread variables from the rendering thread.
+ */
+class FLightComponentId
+{
+public:
+	FLightComponentId()
+		: IDValue(0)
+	{
+	}
+
+	inline bool IsValid() const
+	{
+		return IDValue > 0;
+	}
+
+	inline bool operator==(FLightComponentId OtherId) const
+	{
+		return IDValue == OtherId.IDValue;
+	}
+	
+	inline bool operator<(FLightComponentId OtherId) const
+	{
+		return IDValue < OtherId.IDValue;
+	}
+
+	friend uint32 GetTypeHash(FLightComponentId Id)
+	{
+		return GetTypeHash(Id.IDValue);
+	}
+
+	ENGINE_API static FLightComponentId GetNextId();
+
+private:
+	uint32 IDValue;
+
+	static ENGINE_API FThreadSafeCounter NextLightId;
+
+	FLightComponentId(uint32 InIDValue)
+	: IDValue(InIDValue)
+	{
+	}
+};
\ No newline at end of file
diff --git a/Engine/Source/Runtime/Engine/Public/LightSceneDesc.h b/Engine/Source/Runtime/Engine/Public/LightSceneDesc.h
new file mode 100644
--- /dev/null
+++ b/Engine/Source/Runtime/Engine/Public/LightSceneDesc.h
@@ -0,0 +1,26 @@
+// Copyright Epic Games, Inc. All Rights Reserved.
+
+#pragma once
+
+#include "LightSceneProxy.h"
+
+struct FLightSceneDesc
+{
+	const FTransform& GetTransform() const
+	{
+		return Transform;
+	}
+
+	FVector4 GetPosition() const
+	{
+		return FVector4(GetTransform().GetLocation(), 1);
+	}
+
+	FLightSceneProxy* GetSceneProxy() const
+	{
+		return SceneProxy;
+	}
+
+	FTransform Transform = FTransform::Identity;
+	FLightSceneProxy* SceneProxy = nullptr;
+};
\ No newline at end of file
diff --git a/Engine/Source/Runtime/Engine/Public/LightSceneProxy.h b/Engine/Source/Runtime/Engine/Public/LightSceneProxy.h
--- a/Engine/Source/Runtime/Engine/Public/LightSceneProxy.h
+++ b/Engine/Source/Runtime/Engine/Public/LightSceneProxy.h
@@ -6,12 +6,11 @@
 #include "Math/Color.h"
 #include "Math/MathFwd.h"
 #include "Math/Vector.h"
+#include "LightComponentId.h"
 #include "UObject/NameTypes.h"
 
 class ULightComponent;
 class FMaterialRenderProxy;
-struct FViewMatrices;
-struct FLightRenderParameters;
 class FSceneView;
 class FSceneViewFamily;
 class FWholeSceneProjectedShadowInitializer;
@@ -21,6 +20,9 @@ class FRHICommandList;
 class FLightSceneInfo;
 class UTextureLightProfile;
 class FTexture;
+struct FViewMatrices;
+struct FLightRenderParameters;
+struct FLightSceneProxyDesc;
 
 enum ELightShaderParameterFlags
 {
@@ -43,8 +45,9 @@ class FLightSceneProxy : public FLightSceneProxyFixLayout
 {
 public:
 
-	/** Initialization constructor. */
+	/** Initialization constructors. */
 	ENGINE_API FLightSceneProxy(const ULightComponent* InLightComponent);
+	ENGINE_API FLightSceneProxy(const FLightSceneProxyDesc& Desc);
 	ENGINE_API virtual ~FLightSceneProxy();
 
 	/**
@@ -183,7 +186,9 @@ public:
 	 * The game thread owns UObject state and may be writing to them at any time.
 	 * Mirror the data in the scene proxy and access that instead.
 	 */
+	UE_DEPRECATED(5.8, "Use GetLightComponentId instead.")
 	inline const ULightComponent* GetLightComponent() const { return LightComponent; }
+	inline const FLightComponentId& GetLightComponentId() const { return LightComponentId; }
 	inline FSceneInterface* GetSceneInterface() const { return SceneInterface; }
 	inline FLightSceneInfo* GetLightSceneInfo() const { return LightSceneInfo; }
 	inline const FMatrix& GetWorldToLight() const { return WorldToLight; }
@@ -240,6 +245,14 @@ public:
 	inline bool IsSelected() const { return bSelected; }
 	inline void SetSelected(bool bNewSelected) { bSelected = bNewSelected; }
 
+	// Necessary to build FLightSceneInfo
+	inline bool IsPrecomputedLightingValid() const { return bIsPrecomputedLightingValid; }
+	inline bool IsLightShaftBloomEnabled() const { return bEnableLightShaftBloom; }
+	inline float GetBloomScale() const { return BloomScale; }
+	inline float GetBloomThreshold() const { return BloomThreshold; }
+	inline float GetBloomMaxBrightness() const { return BloomMaxBrightness; }
+	inline const FColor& GetBloomTint() const { return BloomTint; }
+
 	/**
 	 * Use to get the owning actor label (or component name as fallback, if the owner is null or ENABLE_DEBUG_LABELS is off) for diagnostic messages, debug or profiling.
 	 * The actor label is what is shown in the UI (as opposed to the the FName).
@@ -318,9 +331,12 @@ protected:
 	friend class FScene;
 	friend class FLightSceneInfo;
 
-	/** The light component. */
+	/** The light component. (Only necessary for FScene::GetRelevantLights) */
 	const ULightComponent* LightComponent;
 
+	/** The runtime unique id used to identify a light component. */
+	const FLightComponentId LightComponentId;
+
 	/** The scene the primitive is in. */
 	FSceneInterface* SceneInterface;
 
@@ -472,6 +488,24 @@ protected:
 	/** Whether the light should be rendered with MegaLights. */
 	const uint8 bAllowMegaLights : 1;
 
+	/** Whether the precomputed lighting is valid. */
+	uint8 bIsPrecomputedLightingValid : 1;
+
+	/** Whether to render light shaft bloom from this light. */
+	uint8 bEnableLightShaftBloom : 1;
+
+	/** Scales the additive color. */
+	float BloomScale = 1.0f;
+
+	/** Scene color must be larger than this to create bloom in the light shafts. */
+	float BloomThreshold = 1.0f;
+
+	/** After exposure is applied, scene color brightness larger than BloomMaxBrightness will be rescaled down to BloomMaxBrightness. */
+	float BloomMaxBrightness = 1.0f;
+
+	/** Multiplies against scene color to create the bloom color. */
+	FColor BloomTint = FColor(EForceInit::ForceInitToZero);
+
 	/** Whether the light shadows are computed with shadow-mapping or ray-tracing (when available). */
 	const TEnumAsByte<EMegaLightsShadowMethod::Type> MegaLightsShadowMethod;
 
@@ -542,4 +576,18 @@ private:
 	// May store the label or name of the actor containing the component or if there is no actor the name of the component itself
 	FString OwnerNameOrLabel;
 #endif
+
+	friend struct FLightSceneProxyComponentAccessor;
+};
+
+/** Accessor used by FScene::GetRelevantLights. */
+struct FLightSceneProxyComponentAccessor
+{
+private:
+	static const ULightComponent* Get(const FLightSceneProxy& LightSceneProxy)
+	{
+		return LightSceneProxy.LightComponent;
+	}
+	friend class FScene;
 };
+
diff --git a/Engine/Source/Runtime/Engine/Public/LightSceneProxyDesc.h b/Engine/Source/Runtime/Engine/Public/LightSceneProxyDesc.h
new file mode 100644
--- /dev/null
+++ b/Engine/Source/Runtime/Engine/Public/LightSceneProxyDesc.h
@@ -0,0 +1,217 @@
+// Copyright Epic Games, Inc. All Rights Reserved.
+
+#pragma once
+
+#include "CoreMinimal.h"
+#include "LightComponentId.h"
+#include "Engine/EngineTypes.h"
+#include "GameFramework/Actor.h"
+#include "Math/Color.h"
+#include "UObject/ObjectPtr.h"
+
+class FSceneInterface;
+class ULightComponent;
+class FStaticShadowDepthMap;
+
+struct FLightSceneProxyDesc
+{
+	FLightSceneProxyDesc() = default;
+	virtual ~FLightSceneProxyDesc() = default;
+	ENGINE_API virtual void Serialize(FArchive& Ar);
+	ENGINE_API const FStaticShadowDepthMap* GetStaticShadowDepthMap() const;
+	ENGINE_API TStatId GetStatID() const;
+
+	//
+	// Transient properties
+	//
+
+	/** The light component. */
+	const ULightComponent* LightComponent = nullptr;
+
+	/** A runtime unique id that can be used to identify light source **/
+	FLightComponentId LightComponentId;
+
+	/** The scene the light is in. */
+	FSceneInterface* SceneInterface = nullptr;
+
+#if ACTOR_HAS_LABELS
+	// May store the label or name of the actor containing the component or if there is no actor the name of the component itself
+	FString OwnerNameOrLabel;
+#endif
+
+	/** Whether the light is selected in the editor. */
+	uint8 bSelected : 1 = false;
+
+	//
+	// Persistent properties
+	// 
+
+	/** True: length of screen space ray trace for sharp contact shadows is in world space. False: in screen space. */
+	uint8 bContactShadowLengthInWS : 1 = false;
+
+	/* True if the light's Mobility is set to Movable. */
+	uint8 bMovable : 1 = false;
+
+	/** True if a light's parameters as well as its position is static during gameplay, and can thus use static lighting. */
+	uint8 bStaticLighting : 1 = false;
+
+	/** Whether the light has static direct shadowing. */
+	uint8 bStaticShadowing : 1 = false;
+
+	/** True if the light casts dynamic shadows. */
+	uint8 bCastDynamicShadow : 1 = false;
+
+	/** True if the light casts static shadows. */
+	uint8 bCastStaticShadow : 1 = false;
+
+	/** Whether the light is allowed to cast dynamic shadows from translucency. */
+	uint8 bCastTranslucentShadows : 1 = false;
+
+	/** Whether light from this light transmits through surfaces with subsurface scattering profiles. Requires light to be movable. */
+	uint8 bTransmission : 1 = false;
+
+	/** Whether the light shadows volumetric fog.  Disabling this can save GPU time. */
+	uint8 bCastVolumetricShadow : 1 = false;
+
+	/** Whether the light should cast high quality hair-strands self-shadowing. When this option is enabled, an extra GPU cost for this light. */
+	uint8 bCastHairStrandsDeepShadow : 1 = false;
+
+	/** Whether the light should only cast shadows from components marked as bCastCinematicShadows. */
+	uint8 bCastShadowsFromCinematicObjectsOnly : 1 = false;
+
+	/** Enables cached shadows for movable primitives for this light even if r.shadow.cachedshadowscastfrommovableprimitives is 0. */
+	uint8 bForceCachedShadowsForMovablePrimitives : 1 = false;
+
+	/** Whether the light affects objects in reflections, when ray-traced reflection is enabled. */
+	uint8 bAffectReflection : 1 = false;
+
+	/** Whether the light affects global illumination, when ray-traced global illumination is enabled. */
+	uint8 bAffectGlobalIllumination : 1 = false;
+
+	/** Whether the light affects translucency or not.  Disabling this can save GPU time when there are many small lights. */
+	uint8 bAffectTranslucentLighting : 1 = false;
+
+	/** Whether to consider light as a sunlight for atmospheric scattering and exponential height fog. */
+	uint8 bUsedAsAtmosphereSunLight : 1 = false;
+
+	/** Whether to use ray traced distance field area shadows. */
+	uint8 bUseRayTracedDistanceFieldShadows : 1 = false;
+
+	/** Whether the light should be rendered with MegaLights. */
+	uint8 bAllowMegaLights : 1 = false;
+
+	/** Whether the precomputed lighting is valid. */
+	uint8 bIsPrecomputedLightingValid : 1 = false;
+
+	/** Whether to render light shaft bloom from this light. */
+	uint8 bEnableLightShaftBloom : 1 = false;
+
+	/** The light color. */
+	FLinearColor Color = FLinearColor(EForceInit::ForceInitToZero);
+
+	/** Scale for indirect lighting from this light.  When 0, indirect lighting is disabled. */
+	float IndirectLightingScale = 0.0f;
+
+	/** Scales this light's intensity for volumetric scattering. */
+	float VolumetricScatteringIntensity = 1.0f;
+
+	/** Bias applied to end of shadow rays (-1 if global value should be used) */
+	float RayEndBias = 0.0f;
+
+	float ShadowResolutionScale = 1.0f;
+
+	/** User setting from light component, 0:no bias, 0.5:reasonable, larger object might appear to float */
+	float ShadowBias = 0.0f;
+
+	/** User setting from light component, 0:no bias, 0.5:reasonable, larger object might appear to float */
+	float ShadowSlopeBias = 0.0f;
+
+	/** Sharpen shadow filtering */
+	float ShadowSharpen = 0.0f;
+
+	/** Length of screen space ray trace for sharp contact shadows. */
+	float ContactShadowLength = 0.0f;
+
+	/** Intensity of the shadows cast by primitives with "cast contact shadow" enabled. 0 = no shadow, 1 (default) = fully shadowed. */
+	float ContactShadowCastingIntensity = 1.0f;
+
+	/** Intensity of the shadows cast by primitives with "cast contact shadow" disabled. 0 (default) = no shadow, 1 = fully shadowed. */
+	float ContactShadowNonCastingIntensity = 0.0f;
+
+	/** Specular scale */
+	float SpecularScale = 1.0f;
+
+	/** Diffuse scale */
+	float DiffuseScale = 1.0f;
+
+	/** The light's persistent shadowing GUID. */
+	FGuid LightGuid;
+
+	/** Shadow map channel which is used to match up with the appropriate static shadowing during a deferred shading pass. This is generated during a lighting build. */
+	int32 MapBuildDataShadowMapChannel = INDEX_NONE;
+
+	/** Transient shadowmap channel used to preview the results of stationary light shadowmap packing. */
+	int32 PreviewShadowMapChannel = INDEX_NONE;
+
+	/** Controls how large of an offset ray traced shadows have from the receiving surface as the camera gets further away. */
+	float RayStartOffsetDepthScale = 0.0f;
+
+	/** Light function parameters. */
+	FVector	LightFunctionScale = FVector::ZeroVector;
+
+	/** Distance at which the light function should be completely faded to DisabledBrightness. */
+	float LightFunctionFadeDistance = 0.0f;
+
+	/** Brightness factor applied to the light when the light function is specified but disabled, for example in scene captures that use SceneCapView_LitNoShadows. */
+	float LightFunctionDisabledBrightness = 0.0f;
+
+	/** The light function material to be applied to this light. */
+	TObjectPtr<class UMaterialInterface> LightFunctionMaterial = nullptr;
+
+	/** IES texture (light profiles from real world measured data) */
+	TObjectPtr<class UTextureLightProfile> IESTexture = nullptr;
+
+	/** Whether the light shadows are computed with shadow-mapping or ray-tracing (when available). */
+	TEnumAsByte<ECastRayTracedShadow::Type> CastRaytracedShadow = ECastRayTracedShadow::Disabled;
+
+	/** Scales the additive color. */
+	float BloomScale = 1.0f;
+
+	/** Scene color must be larger than this to create bloom in the light shafts. */
+	float BloomThreshold = 1.0f;
+
+	/** After exposure is applied, scene color brightness larger than BloomMaxBrightness will be rescaled down to BloomMaxBrightness. */
+	float BloomMaxBrightness = 1.0f;
+
+	/** Multiplies against scene color to create the bloom color. */
+	FColor BloomTint = FColor(EForceInit::ForceInitToZero);
+
+	/** Whether the light shadows are computed with shadow-mapping or ray-tracing (when available). */
+	TEnumAsByte<EMegaLightsShadowMethod::Type> MegaLightsShadowMethod = EMegaLightsShadowMethod::Default;
+
+	/** The index of the atmospheric light. Multiple lights can be considered when computing the sky/atmospheric scattering. */
+	uint8 AtmosphereSunLightIndex = 0xFF;
+
+	FLinearColor AtmosphereSunDiskColorScale = FLinearColor(EForceInit::ForceInitToZero);
+
+	/** The light type (ELightComponentType) */
+	uint8 LightType = 0xFF;
+
+	/** Channels that this light should affect. */
+	uint8 LightingChannelMask = 0;
+
+	/** View / light masking support.  Controls which views this light should affect. */
+	uint8 ViewLightingChannelMask = 0;
+
+	/** The name of the light component. */
+	FName ComponentName = NAME_None;
+
+	/** The name of the level the light is in. */
+	FName LevelName = NAME_None;
+
+	/** Samples per pixel for ray tracing */
+	uint32 SamplesPerPixel = 1;
+
+	/** Deep shadow layer distribution. */
+	float DeepShadowLayerDistribution = 0.5f;
+};
diff --git a/Engine/Source/Runtime/Engine/Public/LocalLightSceneProxy.h b/Engine/Source/Runtime/Engine/Public/LocalLightSceneProxy.h
--- a/Engine/Source/Runtime/Engine/Public/LocalLightSceneProxy.h
+++ b/Engine/Source/Runtime/Engine/Public/LocalLightSceneProxy.h
@@ -9,6 +9,7 @@
 #include "LightSceneProxy.h"
 
 class ULocalLightComponent;
+struct FLocalLightSceneProxyDesc;
 
 /** The parts of the point light scene info that aren't dependent on the light policy type. */
 class FLocalLightSceneProxy : public FLightSceneProxy
@@ -21,46 +22,37 @@ public:
 	/** One over the light's radius. */
 	float InvRadius;
 
-	/** Initialization constructor. */
+	/** Initialization constructors. */
 	FLocalLightSceneProxy(const ULocalLightComponent* Component);
+	ENGINE_API FLocalLightSceneProxy(const FLocalLightSceneProxyDesc& LightDesc);
+
+	//~Begin FLightSceneProxy Interface
+	ENGINE_API virtual float GetMaxDrawDistance() const final override;
+	ENGINE_API virtual float GetFadeRange() const final override;
+	ENGINE_API virtual float GetRadius() const override;
+	ENGINE_API virtual bool AffectsBounds(const FBoxSphereBounds& Bounds) const override;
+	ENGINE_API virtual bool GetScissorRect(FIntRect& ScissorRect, const FSceneView& View, const FIntRect& ViewRect) const override;
+	ENGINE_API virtual bool SetScissorRect(FRHICommandList& RHICmdList, const FSceneView& View, const FIntRect& ViewRect, FIntRect* OutScissorRect = nullptr) const override;
+	ENGINE_API virtual FSphere GetBoundingSphere() const;
+	ENGINE_API virtual float GetEffectiveScreenRadius(const FViewMatrices& ShadowViewMatrices, const FIntPoint& CameraViewRectSize) const override;
+	ENGINE_API virtual bool GetPerObjectProjectedShadowInitializer(const FBoxSphereBounds& SubjectBounds, class FPerObjectProjectedShadowInitializer& OutInitializer) const override;
+	ENGINE_API virtual bool IsLocalLight() const override;
+	//~End FLightSceneProxy Interface
 
 	/**
-	* Called on the light scene info after it has been passed to the rendering thread to update the rendering thread's cached info when
-	* the light's radius changes.
-	*/
+	 * Called on the light scene info after it has been passed to the rendering thread to update the rendering thread's cached info when
+	 * the light's radius changes.
+	 */
 	void UpdateRadius_GameThread(float Radius);
 
-	// FLightSceneInfo interface.
-	virtual float GetMaxDrawDistance() const final override;
-
-	virtual float GetFadeRange() const final override;
-
-	/** @return radius of the light or 0 if no radius */
-	virtual float GetRadius() const override;
-
-	virtual bool AffectsBounds(const FBoxSphereBounds& Bounds) const override;
-
-	virtual bool GetScissorRect(FIntRect& ScissorRect, const FSceneView& View, const FIntRect& ViewRect) const override;
-
-	virtual bool SetScissorRect(FRHICommandList& RHICmdList, const FSceneView& View, const FIntRect& ViewRect, FIntRect* OutScissorRect = nullptr) const override;
-
-	virtual FSphere GetBoundingSphere() const;
-
-	virtual float GetEffectiveScreenRadius(const FViewMatrices& ShadowViewMatrices, const FIntPoint& CameraViewRectSize) const override;
-	virtual float GetEffectiveScreenRadius(const FViewMatrices& ShadowViewMatrices) const override;
-
 	virtual FVector GetPerObjectProjectedShadowProjectionPoint(const FBoxSphereBounds& SubjectBounds) const;
 
-	virtual bool GetPerObjectProjectedShadowInitializer(const FBoxSphereBounds& SubjectBounds, class FPerObjectProjectedShadowInitializer& OutInitializer) const override;
-
-	virtual bool IsLocalLight() const override;
-
 protected:
 
 	/** Updates the light scene info's radius from the component. */
-	void UpdateRadius(float ComponentRadius);
+	ENGINE_API void UpdateRadius(float ComponentRadius);
 
 	float MaxDrawDistance;
 	float FadeRange;
 	float InverseExposureBlend;
-};
+};
\ No newline at end of file
diff --git a/Engine/Source/Runtime/Engine/Public/LocalLightSceneProxyDesc.h b/Engine/Source/Runtime/Engine/Public/LocalLightSceneProxyDesc.h
new file mode 100644
--- /dev/null
+++ b/Engine/Source/Runtime/Engine/Public/LocalLightSceneProxyDesc.h
@@ -0,0 +1,26 @@
+// Copyright Epic Games, Inc. All Rights Reserved.
+
+#pragma once
+
+#include "LightSceneProxyDesc.h"
+
+struct FLocalLightSceneProxyDesc : public FLightSceneProxyDesc
+{
+	typedef FLightSceneProxyDesc Super;
+
+	FLocalLightSceneProxyDesc() = default;
+	virtual ~FLocalLightSceneProxyDesc() = default;
+	ENGINE_API virtual void Serialize(FArchive& Ar) override;
+
+	/** World distance to where the light stops rendering (lighting & shadows). 0.0 or less means no max draw distance. */
+	float MaxDrawDistance = 0.0f;
+
+	/** Fade the light out as it approaches MaxDrawDistance. */
+	float MaxDistanceFadeRange = 0.0f;
+
+	/** Blend Factor used to blend between Intensity and Intensity/Exposure. */
+	float InverseExposureBlend = 0.0f;
+
+	/** Bounds the light's visible influence. */
+	float AttenuationRadius = 0.0f;
+};
\ No newline at end of file
diff --git a/Engine/Source/Runtime/Engine/Public/PointLightSceneProxy.h b/Engine/Source/Runtime/Engine/Public/PointLightSceneProxy.h
new file mode 100644
--- /dev/null
+++ b/Engine/Source/Runtime/Engine/Public/PointLightSceneProxy.h
@@ -0,0 +1,47 @@
+// Copyright Epic Games, Inc. All Rights Reserved.
+
+/*=============================================================================
+	PointLightSceneProxy.h: Point light scene info definition.
+=============================================================================*/
+
+#pragma once
+
+#include "CoreMinimal.h"
+#include "LocalLightSceneProxy.h"
+
+class UPointLightComponent;
+struct FPointLightSceneProxyDesc;
+
+class FPointLightSceneProxy : public FLocalLightSceneProxy
+{
+public:
+	/** The light falloff exponent. */
+	float FalloffExponent;
+
+	/** Radius of light source shape */
+	float SourceRadius;
+
+	/** Soft radius of light source shape */
+	float SoftSourceRadius;
+
+	/** Length of light source shape */
+	float SourceLength;
+
+	/** Whether light uses inverse squared falloff. */
+	const uint32 bInverseSquared : 1;
+
+	/** Initialization constructor. */
+	FPointLightSceneProxy(const UPointLightComponent* Component);
+	ENGINE_API FPointLightSceneProxy(const FPointLightSceneProxyDesc& LightDesc);
+
+	//~Begin FLightSceneProxy Interface
+	virtual float GetSourceRadius() const override;
+	virtual bool IsInverseSquared() const override;
+	virtual void GetLightShaderParameters(FLightRenderParameters& LightParameters, uint32 Flags=0) const override;
+	virtual bool GetWholeSceneProjectedShadowInitializer(const FSceneViewFamily& ViewFamily, TArray<FWholeSceneProjectedShadowInitializer, TInlineAllocator<6> >& OutInitializers) const override;
+	//~End FLightSceneProxy Interface
+
+	//~Begin FLocalLightSceneProxy Interface
+	virtual FVector GetPerObjectProjectedShadowProjectionPoint(const FBoxSphereBounds& SubjectBounds) const override;
+	//~End FLocalLightSceneProxy Interface
+};
\ No newline at end of file
diff --git a/Engine/Source/Runtime/Engine/Public/PointLightSceneProxyDesc.h b/Engine/Source/Runtime/Engine/Public/PointLightSceneProxyDesc.h
new file mode 100644
--- /dev/null
+++ b/Engine/Source/Runtime/Engine/Public/PointLightSceneProxyDesc.h
@@ -0,0 +1,29 @@
+// Copyright Epic Games, Inc. All Rights Reserved.
+
+#pragma once
+
+#include "LocalLightSceneProxyDesc.h"
+
+struct FPointLightSceneProxyDesc : public FLocalLightSceneProxyDesc
+{
+	typedef FLocalLightSceneProxyDesc Super;
+
+	FPointLightSceneProxyDesc() = default;
+	virtual ~FPointLightSceneProxyDesc() = default;
+	ENGINE_API virtual void Serialize(FArchive& Ar) override;
+
+	/** The light falloff exponent. */
+	float LightFalloffExponent = 0.0f;
+
+	/** Radius of light source shape */
+	float SourceRadius = 0.0f;
+
+	/** Soft radius of light source shape */
+	float SoftSourceRadius = 0.0f;
+
+	/** Length of light source shape */
+	float SourceLength = 0.0f;
+
+	/** Whether light uses inverse squared falloff. */
+	bool bUseInverseSquaredFalloff = false;
+};
\ No newline at end of file
diff --git a/Engine/Source/Runtime/Engine/Public/RectLightSceneProxy.h b/Engine/Source/Runtime/Engine/Public/RectLightSceneProxy.h
--- a/Engine/Source/Runtime/Engine/Public/RectLightSceneProxy.h
+++ b/Engine/Source/Runtime/Engine/Public/RectLightSceneProxy.h
@@ -11,6 +11,7 @@
 #include "LocalLightSceneProxy.h"
 #include "SceneManagement.h"
 
+struct FRectLightSceneProxyDesc;
 
 class FRectLightSceneProxy : public FLocalLightSceneProxy
 {
@@ -25,18 +26,14 @@ public:
 	FVector4f	SourceTextureScaleOffset;
 
 	FRectLightSceneProxy(const URectLightComponent* Component);
+	ENGINE_API FRectLightSceneProxy(const FRectLightSceneProxyDesc& Desc);
 	virtual ~FRectLightSceneProxy();
 
+	//~Begin FLightSceneProxy Interface
 	virtual bool IsRectLight() const override;
-
 	virtual bool HasSourceTexture() const override;
-
-	/** Accesses parameters needed for rendering the light. */
 	virtual void GetLightShaderParameters(FLightRenderParameters& OutLightParameters, uint32 Flags=0) const override;
+	virtual bool GetWholeSceneProjectedShadowInitializer(const FSceneViewFamily& ViewFamily, TArray<FWholeSceneProjectedShadowInitializer, TInlineAllocator<6> >& OutInitializers) const override;
+	//~End FLightSceneProxy Interface
+};
 
-	/**
-	* Sets up a projected shadow initializer for shadows from the entire scene.
-	* @return True if the whole-scene projected shadow should be used.
-	*/
-	virtual bool GetWholeSceneProjectedShadowInitializer(const FSceneViewFamily& ViewFamily, TArray<FWholeSceneProjectedShadowInitializer, TInlineAllocator<6> >& OutInitializers) const;
-};
\ No newline at end of file
diff --git a/Engine/Source/Runtime/Engine/Public/RectLightSceneProxyDesc.h b/Engine/Source/Runtime/Engine/Public/RectLightSceneProxyDesc.h
new file mode 100644
--- /dev/null
+++ b/Engine/Source/Runtime/Engine/Public/RectLightSceneProxyDesc.h
@@ -0,0 +1,36 @@
+// Copyright Epic Games, Inc. All Rights Reserved.
+
+#pragma once
+
+#include "LocalLightSceneProxyDesc.h"
+
+struct FRectLightSceneProxyDesc : public FLocalLightSceneProxyDesc
+{
+	FRectLightSceneProxyDesc() = default;
+	virtual ~FRectLightSceneProxyDesc() = default;
+	ENGINE_API virtual void Serialize(FArchive& Ar) override;
+
+	/** Width of light source rect. */
+	float SourceWidth = 0.0f;
+
+	/** Height of light source rect. */
+	float SourceHeight = 0.0f;
+
+	/** Angle of barn door attached to the light source rect. */
+	float BarnDoorAngle = 0.0f;
+
+	/** Length of barn door attached to the light source rect. */
+	float BarnDoorLength = 0.0f;
+
+	/** Aperture of cone angle for the perspective projection of the light function material. If 0, an orthographic projection is used instead. */
+	float LightFunctionConeAngle = 0.0f;
+
+	/** Texture mapped to the light source rectangle */
+	TObjectPtr<class UTexture> SourceTexture = nullptr;
+
+	/** Scales the source texture. Value in 0..1. */
+	FVector2f SourceTextureScale = FVector2f(1.0f, 1.0f);
+
+	/** Offsets the source texture. Value in 0..1. */
+	FVector2f SourceTextureOffset = FVector2f::ZeroVector;
+};
\ No newline at end of file
diff --git a/Engine/Source/Runtime/Engine/Public/SceneInterface.h b/Engine/Source/Runtime/Engine/Public/SceneInterface.h
--- a/Engine/Source/Runtime/Engine/Public/SceneInterface.h
+++ b/Engine/Source/Runtime/Engine/Public/SceneInterface.h
@@ -180,18 +180,42 @@ public:
 	 * @param Light - light component to add
 	 */
 	virtual void AddLight(ULightComponent* Light) = 0;
+
 	/** 
 	 * Removes a light component from the scene
 	 * 
 	 * @param Light - light component to remove
 	 */
 	virtual void RemoveLight(ULightComponent* Light) = 0;
+
 	/** 
 	 * Adds a new light component to the scene which is currently invisible, but needed for editor previewing
 	 * 
 	 * @param Light - light component to add
 	 */
 	virtual void AddInvisibleLight(ULightComponent* Light) = 0;
+
+	/**
+	 * Adds new lights to the scene
+	 *
+	 * @param InLights - lights to add
+	 */
+	virtual void BatchAddLights(TArrayView<FLightSceneDesc> InLights) = 0;
+
+	/**
+	 * Removes lights from the scene
+	 *
+	 * @param InLights - lights to remove
+	 */
+	virtual void BatchRemoveLights(TArrayView<FLightSceneDesc> InLights) = 0;
+
+	/**
+	 * Adds new invisible lights to the scene
+	 *
+	 * @param InLights - lights to add
+	 */
+	virtual void BatchAddInvisibleLights(TArrayView<FLightSceneDesc> InLights) = 0;
+
 	virtual void SetSkyLight(FSkyLightSceneProxy* Light) = 0;
 	virtual void DisableSkyLight(FSkyLightSceneProxy* Light) = 0;
 
@@ -578,6 +602,13 @@ public:
 	 * @param	RelevantLights	[out]	Array of lights interacting with primitive
 	 */
 	virtual void GetRelevantLights( UPrimitiveComponent* Primitive, TArray<const ULightComponent*>* RelevantLights ) const = 0;
+
+	/**
+	 * Retrieves the number of lights interacting with the passed in primitive.
+	 *
+	 * @param	Primitive				Primitive to retrieve interacting lights for
+	 */
+	virtual int32 GetRelevantLightCount(UPrimitiveComponent* Primitive) const = 0;
 	/**
 	 * Indicates if hit proxies should be processed by this scene
 	 *
diff --git a/Engine/Source/Runtime/Engine/Public/SpotLightSceneProxy.h b/Engine/Source/Runtime/Engine/Public/SpotLightSceneProxy.h
new file mode 100644
--- /dev/null
+++ b/Engine/Source/Runtime/Engine/Public/SpotLightSceneProxy.h
@@ -0,0 +1,51 @@
+// Copyright Epic Games, Inc. All Rights Reserved.
+
+#pragma once
+
+#include "PointLightSceneProxy.h"
+#include "SceneInterface.h"
+#include "SceneView.h"
+
+class USpotLightComponent;
+struct FSpotLightSceneProxyDesc;
+
+/**
+ * The scene info for a spot light.
+ */
+class FSpotLightSceneProxy : public FPointLightSceneProxy
+{
+public:
+
+	/** Initialization constructors. */
+	FSpotLightSceneProxy(const USpotLightComponent* Component);
+	ENGINE_API FSpotLightSceneProxy(const FSpotLightSceneProxyDesc& LightDesc);
+
+	//~Begin FLightSceneProxy Interface
+	virtual bool AffectsBounds(const FBoxSphereBounds& Bounds) const override;
+	virtual bool GetWholeSceneProjectedShadowInitializer(const FSceneViewFamily& ViewFamily, TArray<FWholeSceneProjectedShadowInitializer, TInlineAllocator<6> >& OutInitializers) const override;
+	virtual float GetOuterConeAngle() const override;
+	virtual FSphere GetBoundingSphere() const override;
+	virtual float GetEffectiveScreenRadius(const FViewMatrices& ShadowViewMatrices, const FIntPoint& CameraViewRectSize) const override;
+	virtual void GetLightShaderParameters(FLightRenderParameters& LightParameters, uint32 Flags = 0) const override;
+	//~End FLightSceneProxy Interface
+
+protected:
+
+	/** Outer cone angle in radians, clamped to a valid range. */
+	float OuterConeAngle;
+
+	/** Cosine of the spot light's inner cone angle. */
+	float CosInnerCone;
+
+	/** Cosine of the spot light's outer cone angle. */
+	float CosOuterCone;
+
+	/** 1 / (CosInnerCone - CosOuterCone) */
+	float InvCosConeDifference;
+
+	/** Sine of the spot light's outer cone angle. */
+	float SinOuterCone;
+
+	/** 1 / Tangent of the spot light's outer cone angle. */
+	float InvTanOuterCone;
+};
\ No newline at end of file
diff --git a/Engine/Source/Runtime/Engine/Public/SpotLightSceneProxyDesc.h b/Engine/Source/Runtime/Engine/Public/SpotLightSceneProxyDesc.h
new file mode 100644
--- /dev/null
+++ b/Engine/Source/Runtime/Engine/Public/SpotLightSceneProxyDesc.h
@@ -0,0 +1,20 @@
+// Copyright Epic Games, Inc. All Rights Reserved.
+
+#pragma once
+
+#include "PointLightSceneProxyDesc.h"
+
+struct FSpotLightSceneProxyDesc : public FPointLightSceneProxyDesc
+{
+	typedef FPointLightSceneProxyDesc Super;
+
+	FSpotLightSceneProxyDesc() = default;
+	virtual ~FSpotLightSceneProxyDesc() = default;
+	ENGINE_API virtual void Serialize(FArchive& Ar) override;
+
+	/** Inner cone angle (degree) */
+	float InnerConeAngle = 0.0f;
+
+	/** Outer cone angle (degree) */
+	float OuterConeAngle = 0.0f;
+};
\ No newline at end of file
diff --git a/Engine/Source/Runtime/Renderer/Private/LightGridInjection.cpp b/Engine/Source/Runtime/Renderer/Private/LightGridInjection.cpp
--- a/Engine/Source/Runtime/Renderer/Private/LightGridInjection.cpp
+++ b/Engine/Source/Runtime/Renderer/Private/LightGridInjection.cpp
@@ -28,6 +28,7 @@ LightGridInjection.cpp
 #include "VolumetricFog.h"
 #include "VolumetricCloudRendering.h"
 #include "Components/LightComponent.h"
+#include "Components/RectLightComponent.h"
 #include "Engine/MapBuildDataRegistry.h"
 #include "PixelShaderUtils.h"
 #include "ShaderPrint.h"
diff --git a/Engine/Source/Runtime/Renderer/Private/LightSceneInfo.cpp b/Engine/Source/Runtime/Renderer/Private/LightSceneInfo.cpp
--- a/Engine/Source/Runtime/Renderer/Private/LightSceneInfo.cpp
+++ b/Engine/Source/Runtime/Renderer/Private/LightSceneInfo.cpp
@@ -64,20 +64,17 @@ FLightSceneInfo::FLightSceneInfo(FLightSceneProxy* InProxy, bool InbVisible)
 	, Type((ELightComponentType)InProxy->GetLightType())
 	, Id(INDEX_NONE)
 	, DynamicShadowMapChannel(-1)
-	, bPrecomputedLightingIsValid(InProxy->GetLightComponent()->IsPrecomputedLightingValid())
+	, bPrecomputedLightingIsValid(InProxy->IsPrecomputedLightingValid())
 	, bVisible(InbVisible)
-	, bEnableLightShaftBloom(InProxy->GetLightComponent()->bEnableLightShaftBloom)
-	, BloomScale(InProxy->GetLightComponent()->BloomScale)
-	, BloomThreshold(InProxy->GetLightComponent()->BloomThreshold)
-	, BloomMaxBrightness(InProxy->GetLightComponent()->BloomMaxBrightness)
-	, BloomTint(InProxy->GetLightComponent()->BloomTint)
+	, bEnableLightShaftBloom(InProxy->IsLightShaftBloomEnabled())
+	, BloomScale(InProxy->GetBloomScale())
+	, BloomThreshold(InProxy->GetBloomThreshold())
+	, BloomMaxBrightness(InProxy->GetBloomMaxBrightness())
+	, BloomTint(InProxy->GetBloomTint())
 	, NumUnbuiltInteractions(0)
 	, bCreatePerObjectShadowsForDynamicObjects(Proxy->ShouldCreatePerObjectShadowsForDynamicObjects())
-	, Scene(InProxy->GetLightComponent()->GetScene()->GetRenderScene())
+	, Scene(static_cast<FScene*>(InProxy->GetSceneInterface()))
 {
-	// Only visible lights can be added in game
-	check(bVisible || GIsEditor);
-
 }
 
 void FLightSceneInfo::AddToScene()
diff --git a/Engine/Source/Runtime/Renderer/Private/LightShaftRendering.cpp b/Engine/Source/Runtime/Renderer/Private/LightShaftRendering.cpp
--- a/Engine/Source/Runtime/Renderer/Private/LightShaftRendering.cpp
+++ b/Engine/Source/Runtime/Renderer/Private/LightShaftRendering.cpp
@@ -612,11 +612,11 @@ void FDeferredShadingSceneRenderer::RenderLightShaftBloom(
 							if (View.State)
 							{
 								FSceneViewState* ViewState = static_cast<FSceneViewState*>(View.State);
-								TUniquePtr<FTemporalAAHistory>* Entry = ViewState->LightShaftBloomHistoryRTs.Find(LightSceneProxy.GetLightComponent());
+								TUniquePtr<FTemporalAAHistory>* Entry = ViewState->LightShaftBloomHistoryRTs.Find(LightSceneProxy.GetLightComponentId());
 								if (Entry == nullptr)
 								{
 									TemporalHistory = new FTemporalAAHistory;
-									ViewState->LightShaftBloomHistoryRTs.Emplace(LightSceneProxy.GetLightComponent(), TemporalHistory);
+									ViewState->LightShaftBloomHistoryRTs.Emplace(LightSceneProxy.GetLightComponentId(), TemporalHistory);
 								}
 								else
 								{
@@ -666,7 +666,7 @@ void FDeferredShadingSceneRenderer::RenderLightShaftBloom(
 // InScene is passed in, as the Scene pointer in the class itself may be null, if it was allocated without a scene.
 void FSceneViewState::TrimHistoryRenderTargets(const FScene* InScene)
 {
-	for (TMap<const ULightComponent*, TUniquePtr<FTemporalAAHistory>>::TIterator It(LightShaftBloomHistoryRTs); It; ++It)
+	for (TMap<FLightComponentId, TUniquePtr<FTemporalAAHistory>>::TIterator It(LightShaftBloomHistoryRTs); It; ++It)
 	{
 		bool bLightIsUsed = false;
 
@@ -674,7 +674,7 @@ void FSceneViewState::TrimHistoryRenderTargets(const FScene* InScene)
 		{
 			const FLightSceneInfo* const LightSceneInfo = LightIt->LightSceneInfo;
 
-			if (LightSceneInfo->Proxy->GetLightComponent() == It.Key())
+			if (LightSceneInfo->Proxy->GetLightComponentId() == It.Key())
 			{
 				bLightIsUsed = true;
 				break;
diff --git a/Engine/Source/Runtime/Renderer/Private/RendererScene.cpp b/Engine/Source/Runtime/Renderer/Private/RendererScene.cpp
--- a/Engine/Source/Runtime/Renderer/Private/RendererScene.cpp
+++ b/Engine/Source/Runtime/Renderer/Private/RendererScene.cpp
@@ -48,6 +48,7 @@
 #include "MobileBasePassRendering.h"
 #include "PrimitiveSceneDesc.h"
 #include "InstancedStaticMeshSceneProxyDesc.h"
+#include "LightSceneDesc.h"
 #include "ScenePrivate.h"
 #include "SceneProxies/DeferredDecalProxy.h"
 #include "SceneProxies/ReflectionCaptureProxy.h"
@@ -1160,7 +1161,7 @@ FScene::FScene(UWorld* InWorld, bool bInRequiresHitProxies, bool bInIsEditorScen
 ,	RayTracingDynamicGeometryUpdateManager(nullptr)
 #endif
 ,	PersistentViewStateDebugFlags(0)
-,	NumVisibleLights_GameThread(0)
+,	NumVisibleLights_AnyThread(0)
 ,	NumEnabledSkylights_GameThread(0)
 ,	SceneFrameNumber(0)
 ,	SceneFrameNumberRenderThread(0)
@@ -2310,65 +2311,105 @@ void FScene::AddLight(ULightComponent* Light)
 {
 	LLM_SCOPE(ELLMTag::SceneRender);
 
-	// Create the light's scene proxy.
-	FLightSceneProxy* Proxy = Light->CreateSceneProxy();
-	if(Proxy)
+	if (!Light->SceneProxy)
 	{
-		// Associate the proxy with the light.
-		Light->SceneProxy = Proxy;
+		// Create the light's scene proxy.
+		Light->SceneProxy = Light->CreateSceneProxy();
+	}
+
+	FLightSceneDesc LightSceneDesc{ .Transform = Light->GetComponentTransform(), .SceneProxy = Light->SceneProxy };
+	BatchAddLights(MakeArrayView(&LightSceneDesc, 1));
+}
+
+void FScene::BatchAddLights(TArrayView<FLightSceneDesc> InLights)
+{
+	LLM_SCOPE(ELLMTag::SceneRender);
 
-		// Update the light's transform and position.
-		Proxy->SetTransform(Light->GetComponentTransform().ToMatrixNoScale(), Light->GetLightPosition());
+	TArray<FLightSceneInfo*, SceneRenderingAllocator> LightsToAdd;
+	LightsToAdd.Reserve(InLights.Num());
 
-		// Create the light scene info.
-		Proxy->LightSceneInfo = new FLightSceneInfo(Proxy, true);
+	for (const FLightSceneDesc& LightDesc : InLights)
+	{
+		if (FLightSceneProxy* Proxy = LightDesc.GetSceneProxy())
+		{
+			// Update the light's transform and position.
+			Proxy->SetTransform(LightDesc.GetTransform().ToMatrixNoScale(), LightDesc.GetPosition());
 
-		INC_DWORD_STAT(STAT_SceneLights);
+			// Create the light scene info.
+			Proxy->LightSceneInfo = new FLightSceneInfo(Proxy, true);
 
-		// Adding a new light
-		++NumVisibleLights_GameThread;
+			INC_DWORD_STAT(STAT_SceneLights);
 
-		// Send a command to the rendering thread to add the light to the scene.
-		FLightSceneInfo* LightSceneInfo = Proxy->LightSceneInfo;
-		ENQUEUE_RENDER_COMMAND(FAddLightCommand)(
-			[this, LightSceneInfo] (FRHICommandListBase&)
+			// Adding a new light
+			++NumVisibleLights_AnyThread;
+
+			LightsToAdd.Emplace(Proxy->LightSceneInfo);
+		}
+	}
+
+	if (!LightsToAdd.IsEmpty())
+	{
+		// Send a command to the rendering thread to add the lights to the scene.
+		ENQUEUE_RENDER_COMMAND(BatchAddLightsCommand)(
+			[this, LightsToAdd = MoveTemp(LightsToAdd)](FRHICommandListBase& RHICmdList)
 			{
 				CSV_SCOPED_TIMING_STAT_EXCLUSIVE(Scene_AddLight);
-				FScopeCycleCounter Context(LightSceneInfo->Proxy->GetStatId());
-
-				// It should not have been added to the scene already
-				check(LightSceneInfo->Id == INDEX_NONE);
-				SceneLightInfoUpdates->EnqueueAdd(LightSceneInfo);
+				for (FLightSceneInfo* LightSceneInfo : LightsToAdd)
+				{
+					FScopeCycleCounter Context(LightSceneInfo->Proxy->GetStatId());
+					// It should not have been added to the scene already
+					check(LightSceneInfo->Id == INDEX_NONE);
+					SceneLightInfoUpdates->EnqueueAdd(LightSceneInfo);
+				}
 			});
 	}
 }
 
 void FScene::AddInvisibleLight(ULightComponent* Light)
 {
-	// Create the light's scene proxy.
-	FLightSceneProxy* Proxy = Light->CreateSceneProxy();
+	if (!Light->SceneProxy)
+	{
+		// Create the light's scene proxy.
+		Light->SceneProxy = Light->CreateSceneProxy();
+	}
 
-	if(Proxy)
+	FLightSceneDesc LightSceneDesc{ .Transform = Light->GetComponentTransform(), .SceneProxy = Light->SceneProxy };
+	BatchAddInvisibleLights(MakeArrayView(&LightSceneDesc, 1));
+}
+
+void FScene::BatchAddInvisibleLights(TArrayView<FLightSceneDesc> InLights)
+{
+	TArray<FLightSceneInfo*, SceneRenderingAllocator> LightsToAdd;
+	LightsToAdd.Reserve(InLights.Num());
+
+	for (const FLightSceneDesc& LightDesc : InLights)
 	{
-		// Associate the proxy with the light.
-		Light->SceneProxy = Proxy;
+		if (FLightSceneProxy* Proxy = LightDesc.GetSceneProxy())
+		{
+			// Update the light's transform and position.
+			Proxy->SetTransform(LightDesc.GetTransform().ToMatrixNoScale(), LightDesc.GetPosition());
 
-		// Update the light's transform and position.
-		Proxy->SetTransform(Light->GetComponentTransform().ToMatrixNoScale(),Light->GetLightPosition());
+			// Create the light scene info.
+			Proxy->LightSceneInfo = new FLightSceneInfo(Proxy, false);
 
-		// Create the light scene info.
-		Proxy->LightSceneInfo = new FLightSceneInfo(Proxy, false);
+			INC_DWORD_STAT(STAT_SceneLights);
 
-		INC_DWORD_STAT(STAT_SceneLights);
+			// Send a command to the rendering thread to add the light to the scene.
+			LightsToAdd.Emplace(Proxy->LightSceneInfo);
+		}
+	}
 
-		// Send a command to the rendering thread to add the light to the scene.
-		FScene* Scene = this;
-		FLightSceneInfo* LightSceneInfo = Proxy->LightSceneInfo;
-		ENQUEUE_RENDER_COMMAND(FAddLightCommand)(
-			[Scene, LightSceneInfo] (FRHICommandListBase&)
+	if (!LightsToAdd.IsEmpty())
+	{
+		// Send a command to the rendering thread to add the lights to the scene.
+		ENQUEUE_RENDER_COMMAND(BatchAddInvisibleLightsCommand)(
+			[this, LightsToAdd = MoveTemp(LightsToAdd)](FRHICommandListBase& RHICmdList)
 			{
-				FScopeCycleCounter Context(LightSceneInfo->Proxy->GetStatId());
-				LightSceneInfo->Id = Scene->InvisibleLights.Add(FLightSceneInfoCompact(LightSceneInfo));
+				for (FLightSceneInfo* LightSceneInfo : LightsToAdd)
+				{
+					FScopeCycleCounter Context(LightSceneInfo->Proxy->GetStatId());
+					LightSceneInfo->Id = InvisibleLights.Add(FLightSceneInfoCompact(LightSceneInfo));
+				}
 			});
 	}
 }
@@ -3644,41 +3685,63 @@ void FScene::RemoveLightSceneInfo_RenderThread(FLightSceneInfo* LightSceneInfo)
 
 void FScene::RemoveLight(ULightComponent* Light)
 {
-	if(Light->SceneProxy)
+	FLightSceneDesc LightSceneDesc{ .SceneProxy = Light->SceneProxy };
+	BatchRemoveLights(MakeArrayView(&LightSceneDesc, 1));
+
+	// Disassociate the primitive's render info.
+	Light->SceneProxy = nullptr;
+}
+
+void FScene::BatchRemoveLights(TArrayView<FLightSceneDesc> InLights)
+{
+	TArray<FLightSceneInfo*, SceneRenderingAllocator> LightsToRemove;
+	LightsToRemove.Reserve(InLights.Num());
+
+	for (const FLightSceneDesc& LightDesc : InLights)
 	{
-		FLightSceneInfo* LightSceneInfo = Light->SceneProxy->GetLightSceneInfo();
+		if (LightDesc.SceneProxy)
+		{
+			FLightSceneInfo* LightSceneInfo = LightDesc.SceneProxy->GetLightSceneInfo();
 
-		DEC_DWORD_STAT(STAT_SceneLights);
+			DEC_DWORD_STAT(STAT_SceneLights);
 
-		// Removing one visible light
-		--NumVisibleLights_GameThread;
+			if (LightSceneInfo->bVisible)
+			{
+				// Removing one visible light
+				--NumVisibleLights_AnyThread;
+			}
 
-		// Disassociate the primitive's render info.
-		Light->SceneProxy = nullptr;
+			LightsToRemove.Emplace(LightSceneInfo);
+		}
+	}
 
+	if (!LightsToRemove.IsEmpty())
+	{
 		// Send a command to the rendering thread to queue the light for removal from the scene.
-		ENQUEUE_RENDER_COMMAND(FQueueRemoveLightCommand)(
-			[this, LightSceneInfo] (FRHICommandListBase&)
+		ENQUEUE_RENDER_COMMAND(BatchRemoveLightsCommand)(
+			[this, LightsToRemove = MoveTemp(LightsToRemove)](FRHICommandListBase& RHICmdList)
 			{
-				FScopeCycleCounter Context(LightSceneInfo->Proxy->GetStatId());
-
-				if (LightSceneInfo->bVisible)
+				for (FLightSceneInfo* LightSceneInfo : LightsToRemove)
 				{
-					SceneLightInfoUpdates->EnqueueDelete(LightSceneInfo);
-					// Note: all the deletes are deferred now (except the invisible ones) as this simpifies the update buffer.
-					return;
-				}
-				else
-				{
-					// There should never be updates queued for lights that are not visible
-					check(!SceneLightInfoUpdates->HasCommand(LightSceneInfo));
-					// The "invisible lights" are removed at once.
-					InvisibleLights.RemoveAt(LightSceneInfo->Id);
-				}
+					FScopeCycleCounter Context(LightSceneInfo->Proxy->GetStatId());
 
-				// Free the light scene info and proxy.
-				delete LightSceneInfo->Proxy;
-				delete LightSceneInfo;
+					if (LightSceneInfo->bVisible)
+					{
+						// Note: all the deletes are deferred now (except the invisible ones) as this simpifies the update buffer.
+						SceneLightInfoUpdates->EnqueueDelete(LightSceneInfo);
+					}
+					else
+					{
+						// There should never be updates queued for lights that are not visible
+						check(!SceneLightInfoUpdates->HasCommand(LightSceneInfo));
+						// The "invisible lights" are removed at once.
+						InvisibleLights.RemoveAt(LightSceneInfo->Id);
+
+						// Free the light scene info and proxy.
+						delete LightSceneInfo->Proxy;
+						delete LightSceneInfo;
+					}
+				}
 			});
 	}
 }
@@ -4108,7 +4171,11 @@ void FScene::GetRelevantLights( UPrimitiveComponent* Primitive, TArray<const ULi
 				GetRelevantLights_RenderThread( PrimitiveSceneProxy, RelevantLightProxies );
 				for (const FLightSceneProxy* LightSceneProxy : RelevantLightProxies)
 				{
-					RelevantLights->Add(LightSceneProxy->GetLightComponent());
+					check(LightSceneProxy);
+					if (const ULightComponent* LightComponent = FLightSceneProxyComponentAccessor::Get(*LightSceneProxy))
+					{
+						RelevantLights->Add(LightComponent);
+					}
 				}
 			});
 
@@ -4117,6 +4184,33 @@ void FScene::GetRelevantLights( UPrimitiveComponent* Primitive, TArray<const ULi
 	}
 }
 
+int32 FScene::GetRelevantLightCount(UPrimitiveComponent* Primitive) const
+{
+	int32 Count = 0;
+	if (Primitive && Primitive->GetSceneProxy() != nullptr)
+	{
+		// Add interacting lights to the array.
+		ENQUEUE_RENDER_COMMAND(FGetRelevantLightsCommand)(
+			[this, PrimitiveSceneProxy = Primitive->GetSceneProxy(), &Count](FRHICommandListBase&)
+			{
+				TArray<const FLightSceneProxy*> RelevantLightProxies;
+				GetRelevantLights_RenderThread(PrimitiveSceneProxy, RelevantLightProxies);
+				for (const FLightSceneProxy* LightSceneProxy : RelevantLightProxies)
+				{
+					if (LightSceneProxy->GetLightComponentId().IsValid())
+					{
+						++Count;
+					}
+				}
+			});
+	}
+
+	// Flush rendering commands to wait for command to execute
+	FlushRenderingCommands();
+
+	return Count;
+}
+
 /** Sets the precomputed visibility handler for the scene, or NULL to clear the current one. */
 void FScene::SetPrecomputedVisibility(const FPrecomputedVisibilityHandler* NewPrecomputedVisibilityHandler)
 {
@@ -6692,6 +6786,9 @@ public:
 	virtual void AddLight(ULightComponent* Light) override {}
 	virtual void RemoveLight(ULightComponent* Light) override {}
 	virtual void AddInvisibleLight(ULightComponent* Light) override {}
+	virtual void BatchAddLights(TArrayView<FLightSceneDesc> InLights) override {}
+	virtual void BatchAddInvisibleLights(TArrayView<FLightSceneDesc> InLights) override {}
+	virtual void BatchRemoveLights(TArrayView<FLightSceneDesc> InLights) override {}
 	virtual void SetSkyLight(FSkyLightSceneProxy* Light) override {}
 	virtual void DisableSkyLight(FSkyLightSceneProxy* Light) override {}
 	virtual bool HasSkyLightRequiringLightingBuild() const { return false; }
@@ -6770,6 +6867,11 @@ public:
 	 */
 	virtual void GetRelevantLights( UPrimitiveComponent* Primitive, TArray<const ULightComponent*>* RelevantLights ) const override {}
 
+	virtual int32 GetRelevantLightCount(UPrimitiveComponent* Primitive) const override
+	{
+		return 0;
+	}
+
 	/**
 	 * @return		true if hit proxies should be rendered in this scene.
 	 */
diff --git a/Engine/Source/Runtime/Renderer/Private/ScenePrivate.h b/Engine/Source/Runtime/Renderer/Private/ScenePrivate.h
--- a/Engine/Source/Runtime/Renderer/Private/ScenePrivate.h
+++ b/Engine/Source/Runtime/Renderer/Private/ScenePrivate.h
@@ -32,6 +32,7 @@
 #include "TextureLayout.h"
 #include "SceneRendering.h"
 #include "LightMapRendering.h"
+#include "LightComponentId.h"
 #include "VelocityRendering.h"
 #include "VolumeRendering.h"
 #include "CommonRenderResources.h"
@@ -753,14 +754,17 @@ public:
 
 		inline bool operator == (const FProjectedShadowKey &Other) const
 		{
-			return (PrimitiveId == Other.PrimitiveId && Light == Other.Light && ShadowSplitIndex == Other.ShadowSplitIndex && bTranslucentShadow == Other.bTranslucentShadow);
+			return ((PrimitiveId == Other.PrimitiveId) && 
+					(LightComponentId == Other.LightComponentId) && 
+					(ShadowSplitIndex == Other.ShadowSplitIndex) && 
+					(bTranslucentShadow == Other.bTranslucentShadow));
 		}
 
 		FProjectedShadowKey(const FProjectedShadowInfo& ProjectedShadowInfo);
 
-		FProjectedShadowKey(FPrimitiveComponentId InPrimitiveId, const ULightComponent* InLight, int32 InSplitIndex, bool bInTranslucentShadow)
+		FProjectedShadowKey(FPrimitiveComponentId InPrimitiveId, const FLightComponentId& InLightComponentId, int32 InSplitIndex, bool bInTranslucentShadow)
 			: PrimitiveId(InPrimitiveId)
-			, Light(InLight)
+			, LightComponentId(InLightComponentId)
 			, ShadowSplitIndex(InSplitIndex)
 			, bTranslucentShadow(bInTranslucentShadow)
 		{
@@ -768,12 +772,12 @@ public:
 
 		friend inline uint32 GetTypeHash(const FSceneViewState::FProjectedShadowKey& Key)
 		{
-			return PointerHash(Key.Light,GetTypeHash(Key.PrimitiveId));
+			return HashCombineFast(GetTypeHash(Key.LightComponentId), GetTypeHash(Key.PrimitiveId));
 		}
 
 	private:
 		FPrimitiveComponentId PrimitiveId;
-		const ULightComponent* Light;
+		FLightComponentId LightComponentId;
 		int32 ShadowSplitIndex;
 		bool bTranslucentShadow;
 	};
@@ -993,7 +997,7 @@ public:
 	// Temporal AA result for light shafts of last frame
 	FTemporalAAHistory LightShaftOcclusionHistory;
 	// Temporal AA result for light shafts of last frame
-	TMap<const ULightComponent*, TUniquePtr<FTemporalAAHistory> > LightShaftBloomHistoryRTs;
+	TMap<FLightComponentId, TUniquePtr<FTemporalAAHistory> > LightShaftBloomHistoryRTs;
 
 	FIntRect DistanceFieldAOHistoryViewRect;
 	TRefCountPtr<IPooledRenderTarget> DistanceFieldAOHistoryRT;
@@ -3426,6 +3430,9 @@ public:
 	virtual void AddLight(ULightComponent* Light) override;
 	virtual void RemoveLight(ULightComponent* Light) override;
 	virtual void AddInvisibleLight(ULightComponent* Light) override;
+	virtual void BatchAddLights(TArrayView<FLightSceneDesc> InLights) override;
+	virtual void BatchAddInvisibleLights(TArrayView<FLightSceneDesc> InLights) override;
+	virtual void BatchRemoveLights(TArrayView<FLightSceneDesc> InLights) override;
 	virtual void SetSkyLight(FSkyLightSceneProxy* Light) override;
 	virtual void DisableSkyLight(FSkyLightSceneProxy* Light) override;
 	virtual bool HasSkyLightRequiringLightingBuild() const override;
@@ -3603,6 +3610,9 @@ public:
 	 */
 	virtual void GetRelevantLights( UPrimitiveComponent* Primitive, TArray<const ULightComponent*>* RelevantLights ) const override;
 
+	/** Retrieves the number of lights interacting with the passed in primitive. */
+	virtual int32 GetRelevantLightCount(UPrimitiveComponent* Primitive) const override;
+
 	/** Sets the precomputed visibility handler for the scene, or NULL to clear the current one. */
 	virtual void SetPrecomputedVisibility(const FPrecomputedVisibilityHandler* InPrecomputedVisibilityHandler) override;
 
@@ -3720,7 +3730,7 @@ public:
 	virtual bool HasAnyLights() const override 
 	{ 
 		check(IsInGameThread());
-		return NumVisibleLights_GameThread > 0 || NumEnabledSkylights_GameThread > 0; 
+		return NumVisibleLights_AnyThread > 0 || NumEnabledSkylights_GameThread > 0;
 	}
 
 	virtual bool IsEditorScene() const override { return bIsEditorScene; }
@@ -4089,9 +4099,8 @@ private:
 
 	/** 
 	 * The number of visible lights in the scene
-	 * Note: This is tracked on the game thread!
 	 */
-	int32 NumVisibleLights_GameThread;
+	std::atomic<int32> NumVisibleLights_AnyThread;
 
 	/** 
 	 * Whether the scene has a valid sky light.
diff --git a/Engine/Source/Runtime/Renderer/Private/SceneRendering.h b/Engine/Source/Runtime/Renderer/Private/SceneRendering.h
--- a/Engine/Source/Runtime/Renderer/Private/SceneRendering.h
+++ b/Engine/Source/Runtime/Renderer/Private/SceneRendering.h
@@ -25,6 +25,7 @@
 #include "PrimitiveSceneInfo.h"
 #include "PrimitiveViewRelevance.h"
 #include "LightShaftRendering.h"
+#include "LightComponentId.h"
 #include "StaticBoundShaderState.h"
 #include "Templates/UniquePtr.h"
 #include "MeshDrawCommands.h"
@@ -1074,7 +1075,7 @@ struct FPreviousViewInfo
 	FScreenSpaceDenoiserHistory ReflectedSkyLightHistory;
 
 	// History for shadow denoising.
-	TMap<const ULightComponent*, TSharedPtr<FScreenSpaceDenoiserHistory>> ShadowHistories;
+	TMap<FLightComponentId, TSharedPtr<FScreenSpaceDenoiserHistory>> ShadowHistories;
 
 	// History for denoising all lights penumbra at once.
 	FScreenSpaceDenoiserHistory PolychromaticPenumbraHarmonicsHistory;
diff --git a/Engine/Source/Runtime/Renderer/Private/SceneViewState.cpp b/Engine/Source/Runtime/Renderer/Private/SceneViewState.cpp
--- a/Engine/Source/Runtime/Renderer/Private/SceneViewState.cpp
+++ b/Engine/Source/Runtime/Renderer/Private/SceneViewState.cpp
@@ -7,6 +7,7 @@
 #include "Lumen/LumenSceneData.h"
 #include "RendererModule.h"
 #include "ScenePrivate.h"
+#include "LightComponentId.h"
 #include "SceneViewStateSystemMemory.h"
 #include "VirtualShadowMaps/VirtualShadowMapCacheManager.h"
 
@@ -204,7 +205,7 @@ uint64 FPreviousViewInfo::GetGPUSizeBytes(bool bLogSizes) const
 		GetRenderTargetGPUSizeBytes(MobileBloomSetup_EyeAdaptation, bLogSizes) +
 		GetRenderTargetGPUSizeBytes(VisualizeMotionVectors, bLogSizes);
 
-	for (const TPair<const ULightComponent*, TSharedPtr<FScreenSpaceDenoiserHistory>>& ShadowHistoryIt : ShadowHistories)
+	for (const TPair<FLightComponentId, TSharedPtr<FScreenSpaceDenoiserHistory>>& ShadowHistoryIt : ShadowHistories)
 	{
 		if (ShadowHistoryIt.Value.IsValid())
 		{
@@ -454,7 +455,7 @@ uint64 FSceneViewState::GetGPUSizeBytes(bool bLogSizes) const
 	TotalSize += GetRenderTargetGPUSizeBytes(CombinedLUTRenderTarget, bLogSizes);
 	TotalSize += PrevFrameViewInfo.GetGPUSizeBytes(bLogSizes);
 	TotalSize += LightShaftOcclusionHistory.GetGPUSizeBytes(bLogSizes);
-	for (const TPair<const ULightComponent*, TUniquePtr<FTemporalAAHistory>>& LightShaftBloomIt : LightShaftBloomHistoryRTs)
+	for (const TPair<FLightComponentId, TUniquePtr<FTemporalAAHistory>>& LightShaftBloomIt : LightShaftBloomHistoryRTs)
 	{
 		if (LightShaftBloomIt.Value.IsValid())
 		{
diff --git a/Engine/Source/Runtime/Renderer/Private/ScreenSpaceDenoise.cpp b/Engine/Source/Runtime/Renderer/Private/ScreenSpaceDenoise.cpp
--- a/Engine/Source/Runtime/Renderer/Private/ScreenSpaceDenoise.cpp
+++ b/Engine/Source/Runtime/Renderer/Private/ScreenSpaceDenoise.cpp
@@ -7,6 +7,7 @@
 #include "ScreenSpaceDenoise.h"
 #include "DataDrivenShaderPlatformInfo.h"
 #include "LightSceneProxy.h"
+#include "LightComponentId.h"
 #include "StaticBoundShaderState.h"
 #include "SceneUtils.h"
 #include "PostProcess/SceneRenderTargets.h"
@@ -2446,19 +2447,19 @@ public:
 			// Get the packed penumbra and hit distance in Penumbra texture.
 			InputSignal.Textures[BatchedSignalId] = Parameters.InputTextures.Mask;
 
-			const ULightComponent* LightComponent = Settings.LightSceneInfo[BatchedSignalId]->Proxy->GetLightComponent();
-			TSharedPtr<FScreenSpaceDenoiserHistory>* PrevHistoryEntry = PreviousViewInfos->ShadowHistories.Find(LightComponent);
+			const FLightComponentId& LightComponentId = Settings.LightSceneInfo[BatchedSignalId]->Proxy->GetLightComponentId();
+			TSharedPtr<FScreenSpaceDenoiserHistory>* PrevHistoryEntry = PreviousViewInfos->ShadowHistories.Find(LightComponentId);
 			PrevHistories[BatchedSignalId] = PrevHistoryEntry ? PrevHistoryEntry->Get() : nullptr;
 			NewHistories[BatchedSignalId] = nullptr;
 				
 			if (!View.bStatePrevViewInfoIsReadOnly)
 			{
 				check(View.ViewState);
-				TSharedPtr<FScreenSpaceDenoiserHistory>* NewHistoryEntry = View.ViewState->PrevFrameViewInfo.ShadowHistories.Find(LightComponent);
+				TSharedPtr<FScreenSpaceDenoiserHistory>* NewHistoryEntry = View.ViewState->PrevFrameViewInfo.ShadowHistories.Find(LightComponentId);
 				if (NewHistoryEntry == nullptr)
 				{
 					FScreenSpaceDenoiserHistory* NewHistory = new FScreenSpaceDenoiserHistory;
-					View.ViewState->PrevFrameViewInfo.ShadowHistories.Emplace(LightComponent, NewHistory);
+					View.ViewState->PrevFrameViewInfo.ShadowHistories.Emplace(LightComponentId, NewHistory);
 					NewHistories[BatchedSignalId] = NewHistory;
 				}
 				else
diff --git a/Engine/Source/Runtime/Renderer/Private/ShadowSetup.cpp b/Engine/Source/Runtime/Renderer/Private/ShadowSetup.cpp
--- a/Engine/Source/Runtime/Renderer/Private/ShadowSetup.cpp
+++ b/Engine/Source/Runtime/Renderer/Private/ShadowSetup.cpp
@@ -3052,7 +3052,7 @@ void FProjectedShadowInfo::ComputeScissorRectOptim()
 
 FSceneViewState::FProjectedShadowKey::FProjectedShadowKey(const FProjectedShadowInfo& ProjectedShadowInfo)
 	: PrimitiveId(ProjectedShadowInfo.GetParentSceneInfo() ? ProjectedShadowInfo.GetParentSceneInfo()->PrimitiveComponentId : FPrimitiveComponentId())
-	, Light(ProjectedShadowInfo.GetLightSceneInfo().Proxy->GetLightComponent())
+	, LightComponentId(ProjectedShadowInfo.GetLightSceneInfo().Proxy->GetLightComponentId())
 	, ShadowSplitIndex(ProjectedShadowInfo.CascadeSettings.ShadowSplitIndex)
 	, bTranslucentShadow(ProjectedShadowInfo.bTranslucentShadow)
 {
@@ -3310,7 +3310,7 @@ void FSceneRenderer::CreatePerObjectProjectedShadow(
 		// Check if the subject primitive is shadow relevant.
 		const bool bPrimitiveIsShadowRelevant = ViewRelevance.bShadowRelevance;
 
-		const FSceneViewState::FProjectedShadowKey OpaqueKey(PrimitiveSceneInfo->PrimitiveComponentId, LightSceneInfo->Proxy->GetLightComponent(), INDEX_NONE, false);
+		const FSceneViewState::FProjectedShadowKey OpaqueKey(PrimitiveSceneInfo->PrimitiveComponentId, LightSceneInfo->Proxy->GetLightComponentId(), INDEX_NONE, false);
 
 		// Check if the shadow and preshadow are occluded.
 		const bool bOpaqueShadowIsOccluded = 
@@ -3320,7 +3320,7 @@ void FSceneRenderer::CreatePerObjectProjectedShadow(
 				((FSceneViewState*)View.State)->IsShadowOccluded(OpaqueKey, NumBufferedFrames)
 			);
 
-		const FSceneViewState::FProjectedShadowKey TranslucentKey(PrimitiveSceneInfo->PrimitiveComponentId, LightSceneInfo->Proxy->GetLightComponent(), INDEX_NONE, true);
+		const FSceneViewState::FProjectedShadowKey TranslucentKey(PrimitiveSceneInfo->PrimitiveComponentId, LightSceneInfo->Proxy->GetLightComponentId(), INDEX_NONE, true);
 
 		const bool bTranslucentShadowIsOccluded = 
 			!bCreateTranslucentObjectShadow ||
diff --git a/Engine/Source/Runtime/Renderer/Private/VirtualShadowMaps/VirtualShadowMapArray.cpp b/Engine/Source/Runtime/Renderer/Private/VirtualShadowMaps/VirtualShadowMapArray.cpp
--- a/Engine/Source/Runtime/Renderer/Private/VirtualShadowMaps/VirtualShadowMapArray.cpp
+++ b/Engine/Source/Runtime/Renderer/Private/VirtualShadowMaps/VirtualShadowMapArray.cpp
@@ -8,6 +8,7 @@
 #include "GPUMessaging.h"
 #include "HairStrands/HairStrandsData.h"
 #include "InstanceCulling/InstanceCullingMergedContext.h"
+#include "LightComponentId.h"
 #include "Nanite/Nanite.h"
 #include "RendererModule.h"
 #include "Rendering/NaniteResources.h"
@@ -455,8 +456,8 @@ FAutoConsoleCommand CmdDumpVSMLightNames(
 	FConsoleCommandDelegate::CreateStatic(DumpVSMLightNames)
 );
 
-UPTRINT GVirtualShadowMapLastSelectedVisualizeLightId = 0;
-UPTRINT GVirtualShadowMapVisualizeLightId = 0;
+FLightComponentId GVirtualShadowMapLastSelectedVisualizeLightId;
+FLightComponentId GVirtualShadowMapVisualizeLightId;
 bool GVirtualShadowMapVisualizeByLightId = false;
 
 FString GVirtualShadowMapVisualizeLightName;
@@ -1846,7 +1847,7 @@ void FVirtualShadowMapVisualizeLightSearch::CheckLight(const FLightSceneProxy* C
 		UE_LOG(LogRenderer, Display, TEXT("%s"), *CheckLightName);
 	}
 
-	const UPTRINT CheckProxyId = (UPTRINT)(CheckProxy->GetLightComponent());
+	const FLightComponentId CheckProxyId = CheckProxy->GetLightComponentId();
 
 	const int SelectAdjacentVisualizeLight = GetVirtualShadowMapVisualizationData().SelectAdjacentVisualizeLight;
 	const bool SelectNextLight = SelectAdjacentVisualizeLight > 0;
@@ -1862,8 +1863,8 @@ void FVirtualShadowMapVisualizeLightSearch::CheckLight(const FLightSceneProxy* C
 
 	if (GVirtualShadowMapVisualizeByLightId)
 	{
-		UPTRINT FoundAdjacentLightId = FoundProxy ? (UPTRINT)(FoundProxy->GetLightComponent()) : 0;
-		bool bIsFoundAdjacentLightInvalid = FoundAdjacentLightId == 0;
+		FLightComponentId FoundAdjacentLightId = FoundProxy ? FoundProxy->GetLightComponentId() : FLightComponentId();
+		bool bIsFoundAdjacentLightInvalid = !FoundAdjacentLightId.IsValid();
 
 		if (SelectPrevLight
 			// Light comes before current selection
@@ -1917,11 +1918,11 @@ void FVirtualShadowMapVisualizeLightSearch::ChooseLight()
 
 	if (FoundProxy)
 	{
-		GVirtualShadowMapVisualizeLightId = (UPTRINT)(FoundProxy->GetLightComponent());
+		GVirtualShadowMapVisualizeLightId = FoundProxy->GetLightComponentId();
 	}
 	else if (SelectAdjacentVisualizeLight) // selected past first or last light, select none
 	{
-		GVirtualShadowMapVisualizeLightId = 0;
+		GVirtualShadowMapVisualizeLightId = FLightComponentId();
 	}
 #endif //VSM_ENABLE_VISUALIZATION
 }
