Fix ensure in GetBlueNoiseGlobalParameters when noise vec 2 texture has not been loaded yet

#Synced-CL 48459577

diff --git a/Engine/Source/Runtime/Renderer/Private/BlueNoise.cpp b/Engine/Source/Runtime/Renderer/Private/BlueNoise.cpp
--- a/Engine/Source/Runtime/Renderer/Private/BlueNoise.cpp
+++ b/Engine/Source/Runtime/Renderer/Private/BlueNoise.cpp
@@ -49,7 +49,15 @@ FBlueNoiseParameters GetBlueNoiseParameters()
 	check(GEngine->BlueNoiseScalarTexture && GEngine->BlueNoiseVec2Texture);
 
 	Out.ScalarTexture = GEngine->BlueNoiseScalarTexture->GetResource()->TextureRHI;
-	Out.Vec2Texture = GEngine->BlueNoiseVec2Texture->GetResource()->TextureRHI;
+
+	if (GEngine->BlueNoiseVec2Texture != nullptr && GEngine->BlueNoiseVec2Texture->GetResource() != nullptr)
+	{
+		Out.Vec2Texture = GEngine->BlueNoiseVec2Texture->GetResource()->TextureRHI;
+	}
+	else
+	{
+		Out.Vec2Texture = GSystemTextures.BlackDummy->GetRHI();
+	}
 
 	FillUpBlueNoiseParametersFromTexture(Out);
 	return Out;
@@ -58,8 +66,7 @@ FBlueNoiseParameters GetBlueNoiseParameters()
 FBlueNoise GetBlueNoiseGlobalParameters()
 {
 	FBlueNoise Out;
-	if (GEngine->BlueNoiseScalarTexture != nullptr && GEngine->BlueNoiseScalarTexture->GetResource() != nullptr && 
-		GEngine->BlueNoiseVec2Texture != nullptr && GEngine->BlueNoiseVec2Texture->GetResource() != nullptr)
+	if (GEngine->BlueNoiseScalarTexture != nullptr && GEngine->BlueNoiseScalarTexture->GetResource() != nullptr)
 	{
 		Out.BlueNoise = GetBlueNoiseParameters();
 	}
