Added r.Mobile.ExtendedGBuffer option to enable the extended 160-bit GBuffer.
Disabled octahedral normal encoding when using the extended GBuffer and static lighting is disabled
Enabled decal normal blending with the extended GBuffer.
Added a shader permutation to support capsule local lights in mobile deferred shading.
#Synced-CL 48637686

diff --git a/Engine/Shaders/Private/DeferredLightingCommon.ush b/Engine/Shaders/Private/DeferredLightingCommon.ush
--- a/Engine/Shaders/Private/DeferredLightingCommon.ush
+++ b/Engine/Shaders/Private/DeferredLightingCommon.ush
@@ -258,9 +258,9 @@ float GetLocalLightAttenuation(
 	if (LightData.bInverseSquared)
 	{
 		LightMask = Square( saturate( 1 - Square( DistanceSqr * Square(LightData.InvRadius) ) ) );
-		// This extra attenuation has been added for supporting existing 'legacy' shading model on mobile. 
+		// This extra attenuation has been added for supporting existing 'legacy' shading model on mobile. Not needed for capsule lights
 		// This is not needed for Substrate which unifies all lighting paths
-		#if SHADING_PATH_MOBILE && (!SUBSTRATE_ENABLED || ALLOW_LOCAL_LIGHT_DISTANCE_ATTENUATION)
+		#if SHADING_PATH_MOBILE && !MOBILE_SHADING_PATH_SUPPORT_CAPSULE_LIGHT && (!SUBSTRATE_ENABLED || ALLOW_LOCAL_LIGHT_DISTANCE_ATTENUATION)
 		if (!LightData.bRectLight)
 		{ 
 			LightMask *= 1.0f / (DistanceSqr + 1.0f);
@@ -361,7 +361,7 @@ FLightAccumulator AccumulateDynamicLighting(
 		
 		FDirectLighting Lighting = (FDirectLighting)0;
 
-		half NoL = max(0, dot(GBuffer.WorldNormal, L));
+		half NoL = saturate(dot(N, L));
 	#if TRANSLUCENCY_NON_DIRECTIONAL
 		NoL = 1.0f;
 	#endif
diff --git a/Engine/Shaders/Private/DeferredShadingCommon.ush b/Engine/Shaders/Private/DeferredShadingCommon.ush
--- a/Engine/Shaders/Private/DeferredShadingCommon.ush
+++ b/Engine/Shaders/Private/DeferredShadingCommon.ush
@@ -567,11 +567,15 @@ void MobileEncodeGBuffer(
 	}
 	else
 	{
-		OutGBufferA.rg = UnitVectorToOctahedron( normalize(GBuffer.WorldNormal) ) * 0.5f + 0.5f;
-#if ALLOW_STATIC_LIGHTING
-		OutGBufferA.b = EncodeIndirectIrradiance(GBuffer.IndirectIrradiance);
+#if	MOBILE_EXTENDED_GBUFFER && !ALLOW_STATIC_LIGHTING
+		OutGBufferA.rgb = EncodeNormal(GBuffer.WorldNormal);
 #else
-		OutGBufferA.b = 1;
+		OutGBufferA.rg = UnitVectorToOctahedron(normalize(GBuffer.WorldNormal)) * 0.5 + 0.5;
+		#if ALLOW_STATIC_LIGHTING
+			OutGBufferA.b = EncodeIndirectIrradiance(GBuffer.IndirectIrradiance);
+		#else
+			OutGBufferA.b = 1;
+		#endif
 #endif
 		OutGBufferA.a = GBuffer.PerObjectGBufferData;		
 
@@ -665,11 +669,16 @@ void MobileEncodeGBuffer(
 FGBufferData MobileDecodeGBuffer(half4 InGBufferA, half4 InGBufferB, half4 InGBufferC, half4 InGBufferD)
 {
 	FGBufferData GBuffer = (FGBufferData)0;
-	GBuffer.WorldNormal = OctahedronToUnitVector(InGBufferA.xy * 2.0f - 1.0f);
-#if ALLOW_STATIC_LIGHTING
-	GBuffer.IndirectIrradiance = DecodeIndirectIrradiance(InGBufferA.z);
+
+#if	MOBILE_EXTENDED_GBUFFER && !ALLOW_STATIC_LIGHTING		
+	GBuffer.WorldNormal = DecodeNormal(InGBufferA.rgb);
 #else
-	GBuffer.IndirectIrradiance = 1;
+	GBuffer.WorldNormal = OctahedronToUnitVector(InGBufferA.rg * 2.0 - 1.0);
+	#if ALLOW_STATIC_LIGHTING
+		GBuffer.IndirectIrradiance = DecodeIndirectIrradiance(InGBufferA.b);
+	#else
+		GBuffer.IndirectIrradiance = 1;
+	#endif
 #endif
 	GBuffer.PerObjectGBufferData = InGBufferA.a;
 
diff --git a/Engine/Shaders/Private/LightData.ush b/Engine/Shaders/Private/LightData.ush
--- a/Engine/Shaders/Private/LightData.ush
+++ b/Engine/Shaders/Private/LightData.ush
@@ -588,6 +588,27 @@ FDeferredLightData ConvertToDeferredLight(
 	Out.ShadowMapChannelMask		= UnpackShadowMapChannelMask(LightSceneInfoExtraDataPacked);
 	OutPreviewShadowMapChannelMask	= UnpackShadowMapChannelMask(LightSceneInfoExtraDataPacked >> 4);
 	OutLightingChannelMask			= UnpackLightingChannelMask(LightSceneInfoExtraDataPacked);
+
+#if	SHADING_PATH_MOBILE
+	#if RECT_LIGHT_AS_SPOTLIGHT
+		// Render RectLight approximately as SpotLight
+		Out.bSpotLight = Out.bSpotLight || Out.bRectLight;
+		Out.bRectLight = false;
+	#endif
+	// Override values for performance on mobile
+	#if !MOBILE_SHADING_PATH_SUPPORT_CAPSULE_LIGHT
+		Out.SourceRadius = 0;
+		Out.SourceLength = 0;
+		Out.SoftSourceRadius = 0;
+		Out.SpecularScale = 1.0; // If enabled -> UnpackLightSpecularScale(In);
+		Out.DiffuseScale = 1.0; // If enabled -> UnpackLightDiffuseScale(In);
+		if (!Out.bRectLight)
+		{
+			Out.Tangent = 0;
+		}
+	#endif
+#endif
+
 	return Out;
 }
 
@@ -618,41 +639,6 @@ FDeferredLightData ConvertToDeferredLight(const FLocalLightData In)
 	return ConvertToDeferredLight(In, 1.0f, OutPreviewShadowMapChannelMask, OutLightingChannelMask);
 }
 
-FDeferredLightData ConvertToDeferredLight_Mobile(const FLocalLightData In)
-{
-	half4 OutPreviewShadowMapChannelMask = 0;
-	uint OutLightingChannelMask = LIGHTING_CHANNEL_MASK;
-	FDeferredLightData Out = ConvertToDeferredLight(In, 1.0f, OutPreviewShadowMapChannelMask, OutLightingChannelMask);	
-
-#if MOBILE_SHADING_PATH_SUPPORT_CAPSULE_LIGHT
-	if (!Out.bRectLight)
-	{
-		return ConvertToDeferredLight(In);
-	}
-#endif
-
-	// Override value for performance on mobile
-	Out.SoftSourceRadius= 0;
-	Out.SpecularScale = 1.0f; // If enabled -> UnpackLightSpecularScale(In);
-	Out.DiffuseScale  = 1.0f; // If enabled -> UnpackLightDiffuseScale(In);
-
-#if RECT_LIGHT_AS_SPOTLIGHT
-	// Render RectLight approximately as SpotLight
-	Out.bSpotLight	= Out.bSpotLight || Out.bRectLight;
-	Out.bRectLight	= false;
-#endif
-
-	BRANCH
-	if (!Out.bRectLight)
-	{
-		Out.SourceRadius = 0;
-		Out.SourceLength = 0;
-		Out.Tangent = 0;
-	}
-
-	return Out;
-}
-
 FLightShaderParameters ConvertToLightShaderParameters(FDeferredLightData In)
 {
 	FLightShaderParameters Out;
diff --git a/Engine/Shaders/Private/MobileLightingCommon.ush b/Engine/Shaders/Private/MobileLightingCommon.ush
--- a/Engine/Shaders/Private/MobileLightingCommon.ush
+++ b/Engine/Shaders/Private/MobileLightingCommon.ush
@@ -632,7 +632,7 @@ void AccumulateLightGridLocalLighting(const FCulledLightsGridHeader InLightGridH
 		}
 #endif
 
-		FDeferredLightData LightData = ConvertToDeferredLight_Mobile(LocalLight);
+		FDeferredLightData LightData = ConvertToDeferredLight(LocalLight);
 
 		if (LightType == LIGHT_TYPE_SPOT || LightType == LIGHT_TYPE_POINT)
 		{
@@ -689,7 +689,7 @@ void MergeLocalLights(const FCulledLightsGridHeader InLightGridHeader
 			continue;
 		}
 
-		FDeferredLightData LightData = ConvertToDeferredLight_Mobile(LocalLight);
+		FDeferredLightData LightData = ConvertToDeferredLight(LocalLight);
 		LightData.SpecularScale	= UnpackLightSpecularScale(LocalLight);
 		LightData.DiffuseScale	= UnpackLightDiffuseScale(LocalLight);
 
diff --git a/Engine/Shaders/Private/Substrate/SubstrateMobileForwardLighting.ush b/Engine/Shaders/Private/Substrate/SubstrateMobileForwardLighting.ush
--- a/Engine/Shaders/Private/Substrate/SubstrateMobileForwardLighting.ush
+++ b/Engine/Shaders/Private/Substrate/SubstrateMobileForwardLighting.ush
@@ -431,7 +431,7 @@ float3 SubstrateMobileForwardLighting(
 						half4 PreviewShadowMapChannelMask = 1;
 						half4 DynamicShadowFactors = 1;
 						uint  LocalLightChannelMask = LIGHTING_CHANNEL_MASK;
-						FDeferredLightData LightData = ConvertToDeferredLight_Mobile(LocalLight);
+						FDeferredLightData LightData = ConvertToDeferredLight(LocalLight);
 					
 						// Rect light optionally supported in forward.
 						// * This is not done explicitely in the legacy path (resulting in no lighting) but we do it for Substrate in order to avoid visual artefact.
diff --git a/Engine/Shaders/Private/VolumetricFog.usf b/Engine/Shaders/Private/VolumetricFog.usf
--- a/Engine/Shaders/Private/VolumetricFog.usf
+++ b/Engine/Shaders/Private/VolumetricFog.usf
@@ -1030,11 +1030,7 @@ void LightScatteringCS(
 			if (VolumetricScatteringIntensity > 0)
 			{
 				const bool bSoftFadeEnabled = VolumetricFog.LightSoftFading > 0;
-#if SHADING_PATH_MOBILE
-				const FDeferredLightData LightData = ConvertToDeferredLight_Mobile(LocalLight);
-#else
 				const FDeferredLightData LightData = ConvertToDeferredLight(LocalLight);
-#endif
 
 				float3 L = 0;
 				float3 ToLight = 0;
diff --git a/Engine/Source/Runtime/Core/Private/HAL/ConsoleManager.cpp b/Engine/Source/Runtime/Core/Private/HAL/ConsoleManager.cpp
--- a/Engine/Source/Runtime/Core/Private/HAL/ConsoleManager.cpp
+++ b/Engine/Source/Runtime/Core/Private/HAL/ConsoleManager.cpp
@@ -3804,6 +3804,13 @@ static TAutoConsoleVariable<int32> CVarMobileShadingPath(
 		 "1: Deferred shading (Mobile HDR is required for Deferred) (default)"),
 	ECVF_RenderThreadSafe | ECVF_ReadOnly);
 
+static TAutoConsoleVariable<int32> CVarMobileExtendedGBuffer(
+	TEXT("r.Mobile.ExtendedGBuffer"),
+	0,
+	TEXT("0: 128bit GBuffer (default)\n"
+		 "1: Enable the 160-bit GBuffer, which supports ClearCoat bottom normals and improved normal encoding with decal normal blending (requires static lighting to be disabled)."),
+	ECVF_RenderThreadSafe | ECVF_ReadOnly);
+
 static TAutoConsoleVariable<int32> CVarMobileAllowDeferredShadingOpenGL(
 	TEXT("r.Mobile.AllowDeferredShadingOpenGL"),
 	1,
diff --git a/Engine/Source/Runtime/Engine/Private/ShaderCompiler/ShaderCompiler.cpp b/Engine/Source/Runtime/Engine/Private/ShaderCompiler/ShaderCompiler.cpp
--- a/Engine/Source/Runtime/Engine/Private/ShaderCompiler/ShaderCompiler.cpp
+++ b/Engine/Source/Runtime/Engine/Private/ShaderCompiler/ShaderCompiler.cpp
@@ -3450,7 +3450,7 @@ void GlobalBeginCompileShader(
 		{
 			const bool bGLESDeferredShading = (Target.Platform == SP_OPENGL_ES3_1_ANDROID && bAllowFramebufferFetch);
 			SET_SHADER_DEFINE_AND_COMPILE_ARGUMENT(Input.Environment, USE_GLES_FBF_DEFERRED, bGLESDeferredShading);
-			SET_SHADER_DEFINE(Input.Environment, MOBILE_EXTENDED_GBUFFER, MobileUsesExtenedGBuffer((EShaderPlatform)Target.Platform) ? 1 : 0);
+			SET_SHADER_DEFINE(Input.Environment, MOBILE_EXTENDED_GBUFFER, MobileUsesExtendedGBuffer((EShaderPlatform)Target.Platform) ? 1 : 0);
 			SET_SHADER_DEFINE(Input.Environment, MOBILE_SUPPORTS_SM5_MATERIAL_NODES, MobileSupportsSM5MaterialNodes((EShaderPlatform)Target.Platform) ? 1 : 0);
 		}
 		else
@@ -3792,7 +3792,7 @@ void GlobalBeginCompileShader(
 
 	{
 		static const auto CVar = IConsoleManager::Get().FindTConsoleVariableDataInt(TEXT("r.ClearCoatNormal"));
-		bool bCCBN = CVar && CVar->GetValueOnAnyThread() != 0 && (!bIsMobilePlatform || MobileSupportsSM5MaterialNodes((EShaderPlatform)Target.Platform));
+		bool bCCBN = CVar && CVar->GetValueOnAnyThread() != 0 && (!bIsMobilePlatform || !IsMobileDeferredShadingEnabled((EShaderPlatform)Target.Platform) || MobileUsesExtendedGBuffer((EShaderPlatform)Target.Platform));
 		SET_SHADER_DEFINE(Input.Environment, CLEAR_COAT_BOTTOM_NORMAL, bCCBN ? 1 : 0);
 	}
 
diff --git a/Engine/Source/Runtime/Engine/Private/ShaderCompiler/ShaderGenerationUtil.cpp b/Engine/Source/Runtime/Engine/Private/ShaderCompiler/ShaderGenerationUtil.cpp
--- a/Engine/Source/Runtime/Engine/Private/ShaderCompiler/ShaderGenerationUtil.cpp
+++ b/Engine/Source/Runtime/Engine/Private/ShaderCompiler/ShaderGenerationUtil.cpp
@@ -330,7 +330,7 @@ static FShaderGlobalDefines FetchShaderGlobalDefines(EShaderPlatform TargetPlatf
 
 	{
 		static const auto CVar = IConsoleManager::Get().FindTConsoleVariableDataInt(TEXT("r.ClearCoatNormal"));
-		Ret.CLEAR_COAT_BOTTOM_NORMAL = CVar ? (CVar->GetValueOnAnyThread() != 0 && (!bIsMobilePlatform || MobileSupportsSM5MaterialNodes(TargetPlatform))) : 0;
+		Ret.CLEAR_COAT_BOTTOM_NORMAL = CVar ? (CVar->GetValueOnAnyThread() != 0 && (!bIsMobilePlatform || !IsMobileDeferredShadingEnabled(TargetPlatform) || MobileUsesExtendedGBuffer(TargetPlatform))) : 0;
 	}
 
 	{
diff --git a/Engine/Source/Runtime/RenderCore/Private/GBufferInfo.cpp b/Engine/Source/Runtime/RenderCore/Private/GBufferInfo.cpp
--- a/Engine/Source/Runtime/RenderCore/Private/GBufferInfo.cpp
+++ b/Engine/Source/Runtime/RenderCore/Private/GBufferInfo.cpp
@@ -567,7 +567,7 @@ FGBufferInfo RENDERCORE_API FetchMobileGBufferInfo(const FGBufferParams& Params)
 	else
 	{
 		Info.NumTargets = 4;
-		if (MobileUsesExtenedGBuffer(Params.ShaderPlatform))
+		if (MobileUsesExtendedGBuffer(Params.ShaderPlatform))
 		{
 			Info.NumTargets++;
 		}
diff --git a/Engine/Source/Runtime/RenderCore/Private/RenderUtils.cpp b/Engine/Source/Runtime/RenderCore/Private/RenderUtils.cpp
--- a/Engine/Source/Runtime/RenderCore/Private/RenderUtils.cpp
+++ b/Engine/Source/Runtime/RenderCore/Private/RenderUtils.cpp
@@ -629,21 +629,23 @@ RENDERCORE_API bool MobileUsesShadowMaskTexture(const FStaticShaderPlatform Plat
 }
 
 // Whether to support more than 4 color attachments for GBuffer 
-RENDERCORE_API bool MobileUsesExtenedGBuffer(FStaticShaderPlatform ShaderPlatform)
+RENDERCORE_API bool MobileUsesExtendedGBuffer(FStaticShaderPlatform ShaderPlatform)
 {
+	// For a single pass deferred shading (r.Mobile.AllowFramebufferFetch = 1)
 	// Android GLES: uses PLS for deferred shading and limited to 128 bits
 	// Vulkan requires:
 		// maxDescriptorSetInputAttachments > 4
 		// maxColorAttachments > 4
 	// iOS: A8+
-	return (ShaderPlatform != SP_OPENGL_ES3_1_ANDROID) && MobileSupportsSM5MaterialNodes(ShaderPlatform);
+	static FShaderPlatformCachedIniValue<int32> MobileExtendedGBufferIniValue(TEXT("r.Mobile.ExtendedGBuffer"));
+	return MobileExtendedGBufferIniValue.Get(ShaderPlatform) != 0;
 }
 
 // Required for shading models with a custom data
 RENDERCORE_API bool MobileUsesGBufferCustomData(FStaticShaderPlatform ShaderPlatform)
 {
 	// we can pack CustomData into static lighting related space
-	return MobileUsesExtenedGBuffer(ShaderPlatform) || !IsStaticLightingAllowed();
+	return MobileUsesExtendedGBuffer(ShaderPlatform) || !IsStaticLightingAllowed();
 }
 
 RENDERCORE_API bool MobileBasePassAlwaysUsesCSM(const FStaticShaderPlatform Platform)
diff --git a/Engine/Source/Runtime/RenderCore/Private/Shader.cpp b/Engine/Source/Runtime/RenderCore/Private/Shader.cpp
--- a/Engine/Source/Runtime/RenderCore/Private/Shader.cpp
+++ b/Engine/Source/Runtime/RenderCore/Private/Shader.cpp
@@ -1744,7 +1744,7 @@ void ShaderMapAppendKey(EShaderPlatform Platform, FShaderKeyGenerator& KeyGen)
 	{
 		static const auto CVar = IConsoleManager::Get().FindTConsoleVariableDataInt(TEXT("r.ClearCoatNormal"));
 		KeyGen.AppendSeparator();
-		KeyGen.Append((CVar && CVar->GetValueOnAnyThread() != 0 && (!IsMobilePlatform(Platform) || MobileSupportsSM5MaterialNodes(Platform))) ? TEXT("CCBN") : TEXT("NoCCBN"));
+		KeyGen.Append((CVar && CVar->GetValueOnAnyThread() != 0 && (!IsMobilePlatform(Platform) || !IsMobileDeferredShadingEnabled(Platform) || MobileUsesExtendedGBuffer(Platform))) ? TEXT("CCBN") : TEXT("NoCCBN"));
 	}
 
 	{
@@ -2028,7 +2028,7 @@ void ShaderMapAppendKey(EShaderPlatform Platform, FShaderKeyGenerator& KeyGen)
 			if (bIsMobileDeferredShading)
 			{
 				KeyGen.AppendSeparator();
-				KeyGen.Append(MobileUsesExtenedGBuffer(Platform) ? TEXT("MobDShEx") : TEXT("MobDSh"));
+				KeyGen.Append(MobileUsesExtendedGBuffer(Platform) ? TEXT("MobDShEx") : TEXT("MobDSh"));
 			}
 			else
 			{
diff --git a/Engine/Source/Runtime/RenderCore/Public/RenderUtils.h b/Engine/Source/Runtime/RenderCore/Public/RenderUtils.h
--- a/Engine/Source/Runtime/RenderCore/Public/RenderUtils.h
+++ b/Engine/Source/Runtime/RenderCore/Public/RenderUtils.h
@@ -320,7 +320,7 @@ RENDERCORE_API bool MobileForwardEnableClusteredReflections(const FStaticShaderP
 
 RENDERCORE_API bool MobileUsesShadowMaskTexture(const FStaticShaderPlatform Platform);
 
-RENDERCORE_API bool MobileUsesExtenedGBuffer(const FStaticShaderPlatform Platform);
+RENDERCORE_API bool MobileUsesExtendedGBuffer(const FStaticShaderPlatform Platform);
 
 RENDERCORE_API bool MobileUsesGBufferCustomData(const FStaticShaderPlatform Platform);
 
diff --git a/Engine/Source/Runtime/Renderer/Private/DecalRenderingCommon.cpp b/Engine/Source/Runtime/Renderer/Private/DecalRenderingCommon.cpp
--- a/Engine/Source/Runtime/Renderer/Private/DecalRenderingCommon.cpp
+++ b/Engine/Source/Runtime/Renderer/Private/DecalRenderingCommon.cpp
@@ -29,6 +29,7 @@ namespace DecalRendering
 		const bool bIsMobileDeferredPlatform = bIsMobilePlatform && IsMobileDeferredShadingEnabled(Platform);
 		const bool bIsDBufferPlatform = IsUsingDBuffers(Platform);
 		const bool bIsDBufferMaskPlatform = bIsDBufferPlatform && FDataDrivenShaderPlatformInfo::GetSupportsPerPixelDBufferMask(Platform);
+		const bool bIsMobileNormalBlendable = bIsDBufferPlatform || (bIsMobileDeferredPlatform && MobileUsesExtendedGBuffer(Platform) && !IsStaticLightingAllowed());
 
 		Desc.bWriteDBufferMask = bIsDBufferMaskPlatform;
 
@@ -47,7 +48,7 @@ namespace DecalRendering
 		{
 			Desc.bWriteRoughnessSpecularMetallic = false;
 		}
-		if (bIsMobilePlatform && !bIsDBufferPlatform)
+		if (bIsMobilePlatform && !bIsMobileNormalBlendable)
 		{
 			// Note that we disable writing normals for mobile AND mobile deferred.
 			// Mobile deferred has octahedral encoding for normals and so doesn't support decal normal blending.
diff --git a/Engine/Source/Runtime/Renderer/Private/MobileDeferredShadingPass.cpp b/Engine/Source/Runtime/Renderer/Private/MobileDeferredShadingPass.cpp
--- a/Engine/Source/Runtime/Renderer/Private/MobileDeferredShadingPass.cpp
+++ b/Engine/Source/Runtime/Renderer/Private/MobileDeferredShadingPass.cpp
@@ -63,6 +63,12 @@ static FAutoConsoleVariableRef CVarMobileIgnoreDeferredShadingSkyLightChannels(
 	ECVF_RenderThreadSafe
 );
 
+static TAutoConsoleVariable<int32> CVarMobileAllowCapsuleLights(
+	TEXT("r.Mobile.AllowCapsuleLights"),
+	1,
+	TEXT("Allow capsule lights in a mobile deferred shading (default)"),
+	ECVF_Scalability | ECVF_RenderThreadSafe);
+
 BEGIN_SHADER_PARAMETER_STRUCT(FMobileDeferredPassParameters, )
 	SHADER_PARAMETER_RDG_UNIFORM_BUFFER(FMobileSceneTextureUniformParameters, MobileSceneTextures)
 	RENDER_TARGET_BINDING_SLOTS()
@@ -212,8 +218,9 @@ public:
 	class FRadialLightTypeDim		: SHADER_PERMUTATION_RANGE_INT("RADIAL_LIGHT_TYPE", LIGHT_TYPE_POINT, LIGHT_TYPE_RECT);
 	class FIESProfileDim			: SHADER_PERMUTATION_BOOL("USE_IES_PROFILE");
 	class FSpotLightShadowDim		: SHADER_PERMUTATION_BOOL("SUPPORT_SPOTLIGHTS_SHADOW");
+	class FCapsuleLightDim			: SHADER_PERMUTATION_BOOL("MOBILE_SHADING_PATH_SUPPORT_CAPSULE_LIGHT");
 	class FSimpleLightDim			: SHADER_PERMUTATION_BOOL("SIMPLE_LIGHT");
-	using FPermutationDomain = TShaderPermutationDomain<FEnableShadingModelSupport, FRadialLightTypeDim, FIESProfileDim, FSpotLightShadowDim, FSimpleLightDim>;
+	using FPermutationDomain = TShaderPermutationDomain<FEnableShadingModelSupport, FRadialLightTypeDim, FIESProfileDim, FSpotLightShadowDim, FCapsuleLightDim, FSimpleLightDim>;
 
 	BEGIN_SHADER_PARAMETER_STRUCT(FParameters, )
 		SHADER_PARAMETER_STRUCT(FLightShaderParameters, Light)
@@ -264,6 +271,12 @@ public:
 			PermutationVector.Set<FRadialLightTypeDim>(LIGHT_TYPE_POINT);
 			PermutationVector.Set<FIESProfileDim>(false);
 			PermutationVector.Set<FSpotLightShadowDim>(false);
+			PermutationVector.Set<FCapsuleLightDim>(false);
+		}
+
+		if (PermutationVector.Get<FRadialLightTypeDim>() == LIGHT_TYPE_RECT)
+		{
+			PermutationVector.Set<FCapsuleLightDim>(false);
 		}
 
 		return PermutationVector;
@@ -280,9 +293,6 @@ public:
 
 		const bool bMobileForceDepthRead = MobileUsesFullDepthPrepass(Parameters.Platform);
 		OutEnvironment.SetDefine(TEXT("FORCE_DEPTH_TEXTURE_READS"), bMobileForceDepthRead ? 1u : 0u);
-
-		const bool bSupportCapsule = MobileSupportsSM5MaterialNodes(Parameters.Platform);
-		OutEnvironment.SetDefine(TEXT("MOBILE_SHADING_PATH_SUPPORT_CAPSULE_LIGHT"), bSupportCapsule ? 1u : 0u);
 	}
 
 	void SetParameters(FRHIBatchedShaderParameters& BatchedParameters, const FViewInfo& View, const FMaterialRenderProxy* Proxy, const FMaterial& Material)
@@ -922,6 +932,9 @@ static void RenderLocalLight(
 	const FMatrix WorldToLight = LightSceneInfo.Proxy->GetWorldToLight() * FScaleMatrix(FVector(InverseScale));
 	PassParameters.MobileDeferredCommonParameters.TranslatedWorldToLight = FMatrix44f(FTranslationMatrix(-View.ViewMatrices.GetPreViewTranslation()) * WorldToLight);
 	PassParameters.MobileDeferredCommonParameters.CameraRelativeLightPosition = GetCamRelativeLightPosition(View.ViewMatrices, LightSceneInfo);
+
+	const bool bCapsuleLight = CVarMobileAllowCapsuleLights.GetValueOnAnyThread() != 0 && !bIsRectLight &&
+								(PassParameters.Light.SourceRadius != 0	|| PassParameters.Light.SoftSourceRadius != 0 || PassParameters.Light.SourceLength != 0);
 	
 	// Do two passes, first masking DefautLit, second masking all other shading models
 	const bool bOnlyDefaultLitInView = IsOnlyDefaultLitShadingModel(View.ShadingModelMaskInView);
@@ -937,6 +950,7 @@ static void RenderLocalLight(
 		PermutationVector.Set<FMobileRadialLightFunctionPS::FRadialLightTypeDim>(LightType);
 		PermutationVector.Set<FMobileRadialLightFunctionPS::FIESProfileDim>(bUseIESTexture);
 		PermutationVector.Set<FMobileRadialLightFunctionPS::FSpotLightShadowDim>(bShouldCastShadow);
+		PermutationVector.Set<FMobileRadialLightFunctionPS::FCapsuleLightDim>(bCapsuleLight);
 		FCachedLightMaterial LightMaterial;
 		TShaderRef<FMobileRadialLightFunctionPS> PixelShader;
 		GetLightMaterial(DefaultLightMaterial, LightFunctionMaterialProxy, PermutationVector.ToDimensionValueId(), LightMaterial, PixelShader);
@@ -1287,7 +1301,7 @@ void FDeferredMobileLightMaterialPSOCollector::CollectPSOInitializersLocal(
 	}
 
 	auto AddPSOInitializer =
-		[&](int32 PassIndex, uint8 LightType, bool bUseIESTexture, bool bCameraInsideLightGeometry)
+		[&](int32 PassIndex, uint8 LightType, bool bUseIESTexture, bool bCameraInsideLightGeometry, bool bCapsuleLight)
 		{
 			const bool bEnableShadingModelSupport = (PassIndex > 0);
 			const bool bShouldCastShadow = false;
@@ -1299,6 +1313,7 @@ void FDeferredMobileLightMaterialPSOCollector::CollectPSOInitializersLocal(
 			PermutationVector.Set<FMobileRadialLightFunctionPS::FRadialLightTypeDim>(LightType);
 			PermutationVector.Set<FMobileRadialLightFunctionPS::FIESProfileDim>(bUseIESTexture);
 			PermutationVector.Set<FMobileRadialLightFunctionPS::FSpotLightShadowDim>(bShouldCastShadow);
+			PermutationVector.Set<FMobileRadialLightFunctionPS::FCapsuleLightDim>(bCapsuleLight);
 			FMaterialShaderTypes ShaderTypesToGet;
 			ShaderTypesToGet.AddShaderType<FMobileRadialLightFunctionPS>(PermutationVector.ToDimensionValueId());
 
@@ -1378,15 +1393,28 @@ void FDeferredMobileLightMaterialPSOCollector::CollectPSOInitializersLocal(
 			PSOInitializers.Add(MoveTemp(PSOPrecacheData));
 		};
 
+	bool bAllowCapsuleLights = (CVarMobileAllowCapsuleLights.GetValueOnAnyThread() != 0);
+
 	// int32 PassIndex, uint8 LightType, bool bUseIESTexture, bool bCameraInsideLightGeometry
-	AddPSOInitializer(0, LIGHT_TYPE_POINT, false, false);
-	AddPSOInitializer(1, LIGHT_TYPE_POINT, false, false);
-	AddPSOInitializer(0, LIGHT_TYPE_POINT, true, false);
-	AddPSOInitializer(1, LIGHT_TYPE_POINT, true, false);
-	AddPSOInitializer(0, LIGHT_TYPE_POINT, true, true);
-	AddPSOInitializer(1, LIGHT_TYPE_POINT, true, true);
-	AddPSOInitializer(0, LIGHT_TYPE_SPOT, false, false);
-	AddPSOInitializer(1, LIGHT_TYPE_SPOT, false, false);
+	AddPSOInitializer(0, LIGHT_TYPE_POINT, false, false, false);
+	AddPSOInitializer(1, LIGHT_TYPE_POINT, false, false, false);
+	AddPSOInitializer(0, LIGHT_TYPE_POINT, true, false, false);
+	AddPSOInitializer(1, LIGHT_TYPE_POINT, true, false, false);
+	AddPSOInitializer(0, LIGHT_TYPE_POINT, true, true, false);
+	AddPSOInitializer(1, LIGHT_TYPE_POINT, true, true, false);
+	AddPSOInitializer(0, LIGHT_TYPE_SPOT, false, false, false);
+	AddPSOInitializer(1, LIGHT_TYPE_SPOT, false, false, false);
+
+	if (bAllowCapsuleLights)
+	{
+		AddPSOInitializer(0, LIGHT_TYPE_POINT, false, false, true);
+		AddPSOInitializer(1, LIGHT_TYPE_POINT, false, false, true);
+		AddPSOInitializer(0, LIGHT_TYPE_POINT, false, true, true);
+		AddPSOInitializer(1, LIGHT_TYPE_POINT, false, true, true);
+		
+		AddPSOInitializer(0, LIGHT_TYPE_SPOT, false, false, true);
+		AddPSOInitializer(1, LIGHT_TYPE_SPOT, false, false, true);
+	}
 }
 
 IPSOCollector* CreateDeferredMobileLightMaterialPSOCollector(ERHIFeatureLevel::Type FeatureLevel)
diff --git a/Engine/Source/Runtime/Renderer/Private/MobileShadingRenderer.cpp b/Engine/Source/Runtime/Renderer/Private/MobileShadingRenderer.cpp
--- a/Engine/Source/Runtime/Renderer/Private/MobileShadingRenderer.cpp
+++ b/Engine/Source/Runtime/Renderer/Private/MobileShadingRenderer.cpp
@@ -2228,7 +2228,7 @@ FColorTargets FMobileSceneRenderer::GetColorTargets_Deferred(FSceneTextures& Sce
 		ColorTargets.Add(SceneTextures.GBufferB);
 		ColorTargets.Add(SceneTextures.GBufferC);
 		
-		if (MobileUsesExtenedGBuffer(ShaderPlatform))
+		if (MobileUsesExtendedGBuffer(ShaderPlatform))
 		{
 			ColorTargets.Add(SceneTextures.GBufferD);
 		}
