Fix for assert related to blue noise. Follow up to 48559795.

#tests PC editor
#rnx
#Synced-CL 48604927

diff --git a/Engine/Source/Runtime/Renderer/Private/BlueNoise.cpp b/Engine/Source/Runtime/Renderer/Private/BlueNoise.cpp
--- a/Engine/Source/Runtime/Renderer/Private/BlueNoise.cpp
+++ b/Engine/Source/Runtime/Renderer/Private/BlueNoise.cpp
@@ -46,8 +46,8 @@ FBlueNoiseParameters GetBlueNoiseParameters()
 {
 	FBlueNoiseParameters Out;
 	check(GEngine);
-	check(GEngine->BlueNoiseScalarTexture && GEngine->BlueNoiseVec2Texture);
 
+	check(GEngine->BlueNoiseScalarTexture && GEngine->BlueNoiseScalarTexture->GetResource());
 	Out.ScalarTexture = GEngine->BlueNoiseScalarTexture->GetResource()->TextureRHI;
 
 	if (GEngine->BlueNoiseVec2Texture != nullptr && GEngine->BlueNoiseVec2Texture->GetResource() != nullptr)
