[TSR] Add r.TSR.ShadingRejection.Mode for temporal stability control.

#Synced-CL 48521690

diff --git a/Engine/Shaders/Private/TemporalSuperResolution/TSRShadingAnalysis.ush b/Engine/Shaders/Private/TemporalSuperResolution/TSRShadingAnalysis.ush
--- a/Engine/Shaders/Private/TemporalSuperResolution/TSRShadingAnalysis.ush
+++ b/Engine/Shaders/Private/TemporalSuperResolution/TSRShadingAnalysis.ush
@@ -15,6 +15,9 @@
 #define CONFIG_THIN_GEOMETRY_DETECTION 0
 #endif
 
+#define REJECT_SHADING_RESPONSIVE	0
+#define REJECT_SHADING_STABILITY	1
+
 //------------------------------------------------------- CONSTANTS
 
 static const float kHistoryGuidePreceptionAdd = 1.0;
@@ -26,6 +29,7 @@ static const tsr_half FilteringWeight = rcp(tsr_half(1.0 + 4 * 0.5 + 4 * 0.25));
 float ExposureOffsetFactor;
 float FlickeringFramePeriod;
 uint bPassthroughAlpha;
+uint RejectShadingMode;
 
 //------------------------------------------------------- COMPOSE SEPARATE TRANSLUCENCY
 
@@ -298,12 +302,18 @@ void ComputeMoireError(
 		// Compute the clamp error
 		tsr_tensor_half ClampError;
 		{
-			tsr_tensor_half TotalVarInputDiffC0C2;
-			tsr_tensor_half TotalVarInputC2;
-			Deconcatenate(abs(TotalVariation3x3(Concatenate(abs(InputC0 - InputC2), InputC2))), /* out */ TotalVarInputDiffC0C2, /* out */ TotalVarInputC2);
 		
 			ClampError = BackbufferQuantizationErrorVector;
-			ClampError = max(ClampError, Blur3x3(min(TotalVarInputDiffC0C2, TotalVarInputC2)));
+			
+			BRANCH
+			if (RejectShadingMode == REJECT_SHADING_RESPONSIVE)
+			{
+				tsr_tensor_half TotalVarInputDiffC0C2;
+				tsr_tensor_half TotalVarInputC2;
+				Deconcatenate(abs(TotalVariation3x3(Concatenate(abs(InputC0 - InputC2), InputC2))), /* out */ TotalVarInputDiffC0C2, /* out */ TotalVarInputC2);
+				ClampError = max(ClampError, Blur3x3(min(TotalVarInputDiffC0C2, TotalVarInputC2)));
+			}	
+
 			ClampError = max(ClampError, InputC2BoxSize * tsr_half(FilteringWeight * 0.25));
 			ClampError = ClampError + BackbufferQuantizationErrorVector;
 		}
diff --git a/Engine/Source/Runtime/Renderer/Private/PostProcess/TemporalSuperResolution.cpp b/Engine/Source/Runtime/Renderer/Private/PostProcess/TemporalSuperResolution.cpp
--- a/Engine/Source/Runtime/Renderer/Private/PostProcess/TemporalSuperResolution.cpp
+++ b/Engine/Source/Runtime/Renderer/Private/PostProcess/TemporalSuperResolution.cpp
@@ -119,6 +119,13 @@ TAutoConsoleVariable<int32> CVarTSR16BitVALUOnNvidia(
 
 #endif // PLATFORM_DESKTOP
 
+TAutoConsoleVariable<int32> CVarTSRShadingRejectionMode(
+	TEXT("r.TSR.ShadingRejection.Mode"), 0,
+	TEXT(".\n")
+	TEXT("  0: More responsiveness and less ghosting but tends to have more blocky artifacts.(Default)\n")
+	TEXT("  1: Better stability and shading rejection control but tends to have ghosting.\n"),
+	ECVF_RenderThreadSafe);
+
 TAutoConsoleVariable<float> CVarTSRHistoryRejectionSampleCount(
 	TEXT("r.TSR.ShadingRejection.SampleCount"), 2.0f,
 	TEXT("Maximum number of sample in each output pixel of the history after total shading rejection.\n")
@@ -959,6 +966,7 @@ class FTSRRejectShadingCS : public FTSRConvolutionNetworkShader
 		SHADER_PARAMETER(int32, bEnableFlickeringHeuristic)
 		SHADER_PARAMETER(int32, bPassthroughAlpha)
 		SHADER_PARAMETER(float, ExposureOffsetFactor)
+		SHADER_PARAMETER(int32, RejectShadingMode)
 
 		SHADER_PARAMETER_RDG_TEXTURE(Texture2D, InputTexture)
 		SHADER_PARAMETER_RDG_TEXTURE_SRV(Texture2D, InputMoireLumaTexture)
@@ -2947,6 +2955,7 @@ FDefaultTemporalUpscaler::FOutputs AddTemporalSuperResolutionPasses(
 		PassParameters->bEnableResurrection = bCanResurrectHistory;
 		PassParameters->bEnableFlickeringHeuristic = FlickeringFramePeriod > 0.0f;
 		PassParameters->bPassthroughAlpha = IsPrimitiveAlphaHoldoutEnabled(View);
+		PassParameters->RejectShadingMode = CVarTSRShadingRejectionMode.GetValueOnRenderThread();
 
 		PassParameters->InputTexture = PassInputs.SceneColor.Texture;
 		if (PassInputs.FlickeringInputTexture.IsValid() && FlickeringFramePeriod > 0.0f)
