PR #14075: Add control over mobile upscaling method

#Synced-CL 48702595

diff --git a/Engine/Source/Runtime/Renderer/Private/PostProcess/PostProcessUpscale.cpp b/Engine/Source/Runtime/Renderer/Private/PostProcess/PostProcessUpscale.cpp
--- a/Engine/Source/Runtime/Renderer/Private/PostProcess/PostProcessUpscale.cpp
+++ b/Engine/Source/Runtime/Renderer/Private/PostProcess/PostProcessUpscale.cpp
@@ -30,6 +30,15 @@ TAutoConsoleVariable<int32> CVarUpscaleQuality(
 	TEXT(" 5: 36-tap Gaussian-filtered unsharp mask (very expensive, but good for extreme upsampling).\n"),
 	ECVF_Scalability | ECVF_RenderThreadSafe);
 
+TAutoConsoleVariable<int32> CVarMobileUpscaleQuality(
+	TEXT("r.Mobile.Upscale.Quality"),
+	1,
+	TEXT("Defines the quality in which ScreenPercentage and WindowedFullscreen scales the 3d rendering on mobile platforms.\n")
+	TEXT(" 0: Nearest filtering\n")
+	TEXT(" 1: Simple Bilinear (default)\n")
+	TEXT(" 8: Area\n"),
+	ECVF_Scalability | ECVF_RenderThreadSafe);
+
 TAutoConsoleVariable<int32> CVarUpscaleComputeEnabled(
 	TEXT("r.Upscale.ComputeEnabled"),
 	1,
@@ -239,6 +248,14 @@ EUpscaleMethod GetUpscaleMethod()
 	return static_cast<EUpscaleMethod>(FMath::Clamp(Value, StartIndex, static_cast<int32>(EUpscaleMethod::Gaussian)));
 }
 
+EUpscaleMethod GetMobileUpscaleMethod()
+{
+	const int32 StartIndex = 1; // 0 is EUpscaleMethod::None, methods start at 1
+	const int32 Value = CVarMobileUpscaleQuality.GetValueOnRenderThread() + StartIndex;
+
+	return Value <= static_cast<int32>(EUpscaleMethod::Bilinear) ? static_cast<EUpscaleMethod>(FMath::Max(Value, StartIndex)) : EUpscaleMethod::Area;
+}
+
 // static
 FScreenPassTexture ISpatialUpscaler::AddDefaultUpscalePass(
 	FRDGBuilder& GraphBuilder,
diff --git a/Engine/Source/Runtime/Renderer/Private/PostProcess/PostProcessUpscale.h b/Engine/Source/Runtime/Renderer/Private/PostProcess/PostProcessUpscale.h
--- a/Engine/Source/Runtime/Renderer/Private/PostProcess/PostProcessUpscale.h
+++ b/Engine/Source/Runtime/Renderer/Private/PostProcess/PostProcessUpscale.h
@@ -20,6 +20,7 @@ enum class EUpscaleMethod : uint8
 };
 
 EUpscaleMethod GetUpscaleMethod();
+EUpscaleMethod GetMobileUpscaleMethod();
 
 enum class EUpscaleStage
 {
diff --git a/Engine/Source/Runtime/Renderer/Private/PostProcess/PostProcessing.cpp b/Engine/Source/Runtime/Renderer/Private/PostProcess/PostProcessing.cpp
--- a/Engine/Source/Runtime/Renderer/Private/PostProcess/PostProcessing.cpp
+++ b/Engine/Source/Runtime/Renderer/Private/PostProcess/PostProcessing.cpp
@@ -3493,7 +3493,9 @@ void AddMobilePostProcessingPasses(FRDGBuilder& GraphBuilder, FScene* Scene, con
 		}
 		else
 		{
-			SceneColor = ISpatialUpscaler::AddDefaultUpscalePass(GraphBuilder, View, PassInputs, EUpscaleMethod::Area, View.LensDistortionLUT);
+			EUpscaleMethod Method = GetMobileUpscaleMethod();
+
+			SceneColor = ISpatialUpscaler::AddDefaultUpscalePass(GraphBuilder, View, PassInputs, Method, View.LensDistortionLUT);
 		}
 	}
 
