Remove InverseGamma.y adjustment from HDR_LinearWithToneCurve path

This matches previous 5.7 behavior.

#Synced-CL 48708293

diff --git a/Engine/Source/Runtime/Renderer/Private/PostProcess/PostProcessTonemap.cpp b/Engine/Source/Runtime/Renderer/Private/PostProcess/PostProcessTonemap.cpp
--- a/Engine/Source/Runtime/Renderer/Private/PostProcess/PostProcessTonemap.cpp
+++ b/Engine/Source/Runtime/Renderer/Private/PostProcess/PostProcessTonemap.cpp
@@ -264,7 +264,7 @@ FTonemapperOutputDeviceParameters GetTonemapperOutputDeviceParameters(const FSce
 
 	FVector3f InvDisplayGammaValue;
 	InvDisplayGammaValue.X = 1.0f / Family.RenderTarget->GetDisplayGamma();
-	InvDisplayGammaValue.Y = 2.2f / Family.RenderTarget->GetDisplayGamma();
+	InvDisplayGammaValue.Y = (OutputDeviceValue == EDisplayOutputFormat::HDR_LinearWithToneCurve) ? 1.f : (2.2f / Family.RenderTarget->GetDisplayGamma());
 	InvDisplayGammaValue.Z = 1.0f / FMath::Max(Gamma, 1.0f);
 
 	FTonemapperOutputDeviceParameters Parameters;
