Mobile: Support multi-channel local lights in a deferred shading
#Synced-CL 48703671

diff --git a/Engine/Source/Runtime/Renderer/Private/MobileDeferredShadingPass.cpp b/Engine/Source/Runtime/Renderer/Private/MobileDeferredShadingPass.cpp
--- a/Engine/Source/Runtime/Renderer/Private/MobileDeferredShadingPass.cpp
+++ b/Engine/Source/Runtime/Renderer/Private/MobileDeferredShadingPass.cpp
@@ -872,8 +872,7 @@ static void RenderLocalLight(
 	const FCachedLightMaterial& DefaultLightMaterial,
 	const TArray<FVisibleLightInfo, SceneRenderingAllocator>& VisibleLightInfos)
 {
-	uint8 LightingChannelMask = LightSceneInfo.Proxy->GetLightingChannelMask();
-	if (!LightSceneInfo.ShouldRenderLight(View) || LightingChannelMask == 0)
+	if (!LightSceneInfo.ShouldRenderLight(View) || LightSceneInfo.Proxy->GetLightingChannelMask() == 0)
 	{
 		return;
 	}
@@ -905,9 +904,6 @@ static void RenderLocalLight(
 	GraphicsPSOInit.PrimitiveType = PT_TriangleList;
 	const FSphere LightBounds = LightSceneInfo.Proxy->GetBoundingSphere();
 
-	const uint32 LightingChannel = GetLightingChannel(LightingChannelMask);
-	const uint8 LightingChannelStencilValue = GetLightingChannelStencilValue(LightingChannel);
-
 	FDeferredLightVS::FPermutationDomain PermutationVectorVS;
 	PermutationVectorVS.Set<FDeferredLightVS::FRadialLight>(true);
 	TShaderMapRef<FDeferredLightVS> VertexShader(View.ShaderMap, PermutationVectorVS);
@@ -942,38 +938,47 @@ static void RenderLocalLight(
 
 	for (int32 PassIndex = 0; PassIndex < NumPasses; PassIndex++)
 	{
-		const bool bEnableShadingModelSupport = (PassIndex > 0);
-		SetLocalLightRasterizerAndDepthState(GraphicsPSOInit, View, LightBounds, LightingChannel, bEnableShadingModelSupport);
-
-		FMobileRadialLightFunctionPS::FPermutationDomain PermutationVector;
-		PermutationVector.Set<FMobileRadialLightFunctionPS::FEnableShadingModelSupport>(bEnableShadingModelSupport);
-		PermutationVector.Set<FMobileRadialLightFunctionPS::FRadialLightTypeDim>(LightType);
-		PermutationVector.Set<FMobileRadialLightFunctionPS::FIESProfileDim>(bUseIESTexture);
-		PermutationVector.Set<FMobileRadialLightFunctionPS::FSpotLightShadowDim>(bShouldCastShadow);
-		PermutationVector.Set<FMobileRadialLightFunctionPS::FCapsuleLightDim>(bCapsuleLight);
-		FCachedLightMaterial LightMaterial;
-		TShaderRef<FMobileRadialLightFunctionPS> PixelShader;
-		GetLightMaterial(DefaultLightMaterial, LightFunctionMaterialProxy, PermutationVector.ToDimensionValueId(), LightMaterial, PixelShader);
+		// Clamp mask to 3 channels
+		uint8 LightingChannelMask = (LightSceneInfo.Proxy->GetLightingChannelMask() & 0x7);
+		while (LightingChannelMask != 0)
+		{
+			const uint32 LightingChannel = GetLightingChannel(LightingChannelMask);
+			const uint8 LightingChannelStencilValue = GetLightingChannelStencilValue(LightingChannel);
+			LightingChannelMask &= ~(1u << LightingChannel);
 
-		GraphicsPSOInit.BoundShaderState.VertexDeclarationRHI = GetVertexDeclarationFVector4();
-		GraphicsPSOInit.BoundShaderState.VertexShaderRHI = VertexShader.GetVertexShader();
-		GraphicsPSOInit.BoundShaderState.PixelShaderRHI = PixelShader.GetPixelShader();
+			const bool bEnableShadingModelSupport = (PassIndex > 0);
+			SetLocalLightRasterizerAndDepthState(GraphicsPSOInit, View, LightBounds, LightingChannel, bEnableShadingModelSupport);
 
-		const uint8 StencilRef = PassShadingModelStencilValue(bEnableShadingModelSupport) | STENCIL_LIGHTING_CHANNELS_MASK(LightingChannelStencilValue);
+			FMobileRadialLightFunctionPS::FPermutationDomain PermutationVector;
+			PermutationVector.Set<FMobileRadialLightFunctionPS::FEnableShadingModelSupport>(bEnableShadingModelSupport);
+			PermutationVector.Set<FMobileRadialLightFunctionPS::FRadialLightTypeDim>(LightType);
+			PermutationVector.Set<FMobileRadialLightFunctionPS::FIESProfileDim>(bUseIESTexture);
+			PermutationVector.Set<FMobileRadialLightFunctionPS::FSpotLightShadowDim>(bShouldCastShadow);
+			PermutationVector.Set<FMobileRadialLightFunctionPS::FCapsuleLightDim>(bCapsuleLight);
+			FCachedLightMaterial LightMaterial;
+			TShaderRef<FMobileRadialLightFunctionPS> PixelShader;
+			GetLightMaterial(DefaultLightMaterial, LightFunctionMaterialProxy, PermutationVector.ToDimensionValueId(), LightMaterial, PixelShader);
 
-		SetGraphicsPipelineState(RHICmdList, GraphicsPSOInit, StencilRef);
+			GraphicsPSOInit.BoundShaderState.VertexDeclarationRHI = GetVertexDeclarationFVector4();
+			GraphicsPSOInit.BoundShaderState.VertexShaderRHI = VertexShader.GetVertexShader();
+			GraphicsPSOInit.BoundShaderState.PixelShaderRHI = PixelShader.GetPixelShader();
 
-		SetShaderParameters(RHICmdList, VertexShader, VertexShader.GetVertexShader(), ParametersVS);
+			const uint8 StencilRef = PassShadingModelStencilValue(bEnableShadingModelSupport) | STENCIL_LIGHTING_CHANNELS_MASK(LightingChannelStencilValue);
 
-		SetShaderParametersMixedPS(RHICmdList, PixelShader, PassParameters, View, LightMaterial.MaterialProxy, *LightMaterial.Material);
+			SetGraphicsPipelineState(RHICmdList, GraphicsPSOInit, StencilRef);
 
-		if (LightType == LightType_Point || LightType == LightType_Rect)
-		{
-			StencilingGeometry::DrawSphere(RHICmdList);
-		}
-		else // LightType_Spot
-		{
-			StencilingGeometry::DrawCone(RHICmdList);
+			SetShaderParameters(RHICmdList, VertexShader, VertexShader.GetVertexShader(), ParametersVS);
+
+			SetShaderParametersMixedPS(RHICmdList, PixelShader, PassParameters, View, LightMaterial.MaterialProxy, *LightMaterial.Material);
+
+			if (LightType == LightType_Point || LightType == LightType_Rect)
+			{
+				StencilingGeometry::DrawSphere(RHICmdList);
+			}
+			else // LightType_Spot
+			{
+				StencilingGeometry::DrawCone(RHICmdList);
+			}
 		}
 	}
 }
