Mobile: Fixed issue with a reflection texture source mip during downsampling
Also use same sampling pattern for downsampling as SM5 does
#Synced-CL 48703671

diff --git a/Engine/Shaders/Private/ReflectionEnvironmentShaders.usf b/Engine/Shaders/Private/ReflectionEnvironmentShaders.usf
--- a/Engine/Shaders/Private/ReflectionEnvironmentShaders.usf
+++ b/Engine/Shaders/Private/ReflectionEnvironmentShaders.usf
@@ -244,6 +244,21 @@ void CopyCubemapToIlluminanceMeterCubemapPS(
 
 uint MipIndex;
 uint NumMips;
+uint SourceMipIndex;
+
+float4 SampleCubemap(float3 Coordinates, uint InMipIndex)
+{
+	return TextureCubeSampleLevel(SourceCubemapTexture, SourceCubemapSampler, Coordinates, InMipIndex);
+}
+
+float4 SampleCubemap(float3 Coordinates)
+{
+#ifdef USE_COMPUTE
+	return SampleCubemap(Coordinates, 0);
+#else
+	return SampleCubemap(Coordinates, SourceMipIndex);
+#endif
+}
 
 #ifdef USE_COMPUTE
 int FaceThreadGroupSize;
@@ -296,7 +311,7 @@ void DownsamplePS(
 		float2( 0,  1),
 	};
 
-	OutColor = SourceCubemapTexture.SampleLevel(SourceCubemapSampler, CubeCoordinates, 0 );
+	OutColor = SampleCubemap(CubeCoordinates);
 
 	UNROLL
 	for( uint i = 0; i < 8; i++ )
@@ -306,7 +321,7 @@ void DownsamplePS(
 		float3 SampleDir = CubeCoordinates;
 		SampleDir += TangentX * ( Offsets[i].x * SampleOffset );
 		SampleDir += TangentY * ( Offsets[i].y * SampleOffset );
-		OutColor += SourceCubemapTexture.SampleLevel(SourceCubemapSampler, SampleDir, 0 ) * Weight;
+		OutColor += SampleCubemap(SampleDir) * Weight;
 	}
 
 	OutColor *= rcp( 1.0 + 1.0 + 2.0 );
@@ -483,29 +498,6 @@ void ComputeCubeMaxLuminancePS(out float4 OutColor : SV_Target0)
 
 #endif
 
-float4 SampleCubemap(float3 Coordinates, uint InMipIndex)
-{
-	return TextureCubeSampleLevel(SourceCubemapTexture, SourceCubemapSampler, Coordinates, InMipIndex);
-}
-
-uint SourceMipIndex;
-float4 SampleCubemap(float3 Coordinates)
-{
-	return SampleCubemap(Coordinates, SourceMipIndex);
-}
-
-void DownsamplePS_Mobile(
-	float4 SvPosition : SV_POSITION,
-	out float4 OutColor : SV_Target0
-)
-{
-	float2 UV = SvPosition.xy * SvPositionToUVScale;
-	float2 ScaledUVs = UV * 2 - 1;
-	float3 CubeCoordinates = GetCubemapVector(ScaledUVs, CubeFace);
-
-	OutColor = SampleCubemap(CubeCoordinates);
-}
-
 #if SHADING_PATH_MOBILE
 // There's a precision problem in OpenGL because the default precision in the Pixel Shader is set to half, even for constant values, use Hammersley16 instead.
 #define HammersleyDistribution Hammersley16
diff --git a/Engine/Source/Runtime/Renderer/Private/MobileReflectionEnvironmentCapture.cpp b/Engine/Source/Runtime/Renderer/Private/MobileReflectionEnvironmentCapture.cpp
--- a/Engine/Source/Runtime/Renderer/Private/MobileReflectionEnvironmentCapture.cpp
+++ b/Engine/Source/Runtime/Renderer/Private/MobileReflectionEnvironmentCapture.cpp
@@ -19,67 +19,8 @@ extern FRDGTexture* FilterCubeMap(FRDGBuilder& GraphBuilder, FGlobalShaderMap* S
 
 extern void PremultiplyCubeMipAlpha(FRDGBuilder& GraphBuilder, FGlobalShaderMap* ShaderMap, FRDGTexture* CubemapTexture, int32 MipIndex);
 
-class FMobileCubeDownsamplePS : public FGlobalShader
-{
-	DECLARE_GLOBAL_SHADER(FMobileCubeDownsamplePS);
-	SHADER_USE_PARAMETER_STRUCT(FMobileCubeDownsamplePS, FGlobalShader);
-
-	BEGIN_SHADER_PARAMETER_STRUCT(FParameters, )
-		SHADER_PARAMETER_RDG_TEXTURE_SRV(TextureCube, SourceCubemapTexture)
-		SHADER_PARAMETER_SAMPLER(SamplerState, SourceCubemapSampler)
-		SHADER_PARAMETER(FVector2f, SvPositionToUVScale)
-		SHADER_PARAMETER(int32, CubeFace)
-		SHADER_PARAMETER(int32, SourceMipIndex)
-		RENDER_TARGET_BINDING_SLOTS()
-	END_SHADER_PARAMETER_STRUCT()
-};
-
-IMPLEMENT_GLOBAL_SHADER(FMobileCubeDownsamplePS, "/Engine/Private/ReflectionEnvironmentShaders.usf", "DownsamplePS_Mobile", SF_Pixel);
-
 namespace MobileReflectionEnvironmentCapture
 {
-	void CreateCubeMips(FRDGBuilder& GraphBuilder, FGlobalShaderMap* ShaderMap, FRDGTexture* CubemapTexture)
-	{
-		RDG_EVENT_SCOPE(GraphBuilder, "CreateCubeMips");
-
-		TShaderMapRef<FMobileCubeDownsamplePS> PixelShader(ShaderMap);
-
-		const int32 NumMips = CubemapTexture->Desc.NumMips;
-
-		// Downsample all the mips, each one reads from the mip above it
-		for (int32 MipIndex = 1; MipIndex < NumMips; MipIndex++)
-		{
-			const int32 SourceMipIndex = FMath::Max(MipIndex - 1, 0);
-			const int32 MipSize = 1 << (NumMips - MipIndex - 1);
-			const FIntRect ViewRect(0, 0, MipSize, MipSize);
-
-			FRDGTextureSRVDesc SourceSRVDesc = FRDGTextureSRVDesc::CreateForMipLevel(CubemapTexture, MipIndex - 1);
-			if (GRHISupportsTextureViews == false)
-			{
-				SourceSRVDesc = FRDGTextureSRVDesc::Create(CubemapTexture);
-			}
-
-			for (int32 CubeFace = 0; CubeFace < CubeFace_MAX; CubeFace++)
-			{
-				auto* PassParameters = GraphBuilder.AllocParameters<FMobileCubeDownsamplePS::FParameters>();
-				PassParameters->SourceCubemapTexture = GraphBuilder.CreateSRV(SourceSRVDesc);
-				PassParameters->SourceCubemapSampler = TStaticSamplerState<SF_Bilinear, AM_Clamp, AM_Clamp, AM_Clamp>::GetRHI();
-				PassParameters->CubeFace = CubeFace;
-				PassParameters->SourceMipIndex = SourceMipIndex;
-				PassParameters->SvPositionToUVScale = FVector2f(1.0f / MipSize, 1.0f / MipSize);
-				PassParameters->RenderTargets[0] = FRenderTargetBinding(CubemapTexture, ERenderTargetLoadAction::ENoAction, MipIndex, CubeFace);
-
-				FPixelShaderUtils::AddFullscreenPass(
-					GraphBuilder,
-					ShaderMap,
-					RDG_EVENT_NAME("CreateCubeMips (Mip: %d, Face: %d)", MipIndex, CubeFace),
-					PixelShader,
-					PassParameters,
-					ViewRect);
-			}
-		}
-	}
-
 	void ComputeAverageBrightness(FRDGBuilder& GraphBuilder, FGlobalShaderMap* ShaderMap, FRDGTexture* CubemapTexture, float* OutAverageBrightness)
 	{
 		CreateCubeMips(GraphBuilder, ShaderMap, CubemapTexture);
diff --git a/Engine/Source/Runtime/Renderer/Private/MobileReflectionEnvironmentCapture.h b/Engine/Source/Runtime/Renderer/Private/MobileReflectionEnvironmentCapture.h
--- a/Engine/Source/Runtime/Renderer/Private/MobileReflectionEnvironmentCapture.h
+++ b/Engine/Source/Runtime/Renderer/Private/MobileReflectionEnvironmentCapture.h
@@ -18,5 +18,4 @@ namespace MobileReflectionEnvironmentCapture
 	void ComputeAverageBrightness(FRDGBuilder& GraphBuilder, FGlobalShaderMap* ShaderMap, FRDGTexture* CubemapTexture, float* OutAverageBrightness);
 	FRDGTexture* FilterReflectionEnvironment(FRDGBuilder& GraphBuilder, FGlobalShaderMap* ShaderMap, FRDGTexture* CubemapTexture, FSHVectorRGB3* OutIrradianceEnvironmentMap);
 	void CopyToSkyTexture(FRDGBuilder& GraphBuilder, FScene* Scene, FRDGTexture* CubemapTexture, FTexture* ProcessedTexture);
-	void CreateCubeMips(FRDGBuilder& GraphBuilder, FGlobalShaderMap* ShaderMap, FRDGTexture* CubemapTexture);
 }
diff --git a/Engine/Source/Runtime/Renderer/Private/ReflectionEnvironmentCapture.cpp b/Engine/Source/Runtime/Renderer/Private/ReflectionEnvironmentCapture.cpp
--- a/Engine/Source/Runtime/Renderer/Private/ReflectionEnvironmentCapture.cpp
+++ b/Engine/Source/Runtime/Renderer/Private/ReflectionEnvironmentCapture.cpp
@@ -120,6 +120,7 @@ BEGIN_SHADER_PARAMETER_STRUCT(FCubeShaderParameters, )
 	SHADER_PARAMETER_SAMPLER(SamplerState, SourceCubemapSampler)
 	SHADER_PARAMETER(FVector2f, SvPositionToUVScale)
 	SHADER_PARAMETER(int32, CubeFace)
+	SHADER_PARAMETER(int32, SourceMipIndex)
 	SHADER_PARAMETER(int32, MipIndex)
 	SHADER_PARAMETER(int32, NumMips)
 	RENDER_TARGET_BINDING_SLOTS()
@@ -323,15 +324,23 @@ void CreateCubeMips(FRDGBuilder& GraphBuilder, FGlobalShaderMap* ShaderMap, FRDG
 	// Downsample all the mips, each one reads from the mip above it
 	for (int32 MipIndex = 1; MipIndex < NumMips; MipIndex++)
 	{
+		const int32 SourceMipIndex = (MipIndex - 1);
 		const int32 MipSize = 1 << (NumMips - MipIndex - 1);
 		const FIntRect ViewRect(0, 0, MipSize, MipSize);
 
+		FRDGTextureSRVDesc SourceSRVDesc = FRDGTextureSRVDesc::CreateForMipLevel(CubemapTexture, SourceMipIndex);
+		if (GRHISupportsTextureViews == false)
+		{
+			SourceSRVDesc = FRDGTextureSRVDesc::Create(CubemapTexture);
+		}
+
 		for (int32 CubeFace = 0; CubeFace < CubeFace_MAX; CubeFace++)
 		{
 			auto* PassParameters = GraphBuilder.AllocParameters<FCubeDownsamplePS::FParameters>();
-			PassParameters->SourceCubemapTexture = GraphBuilder.CreateSRV(FRDGTextureSRVDesc::CreateForMipLevel(CubemapTexture, MipIndex - 1));
+			PassParameters->SourceCubemapTexture = GraphBuilder.CreateSRV(SourceSRVDesc);
 			PassParameters->SourceCubemapSampler = TStaticSamplerState<SF_Bilinear, AM_Clamp, AM_Clamp, AM_Clamp>::GetRHI();
 			PassParameters->CubeFace = CubeFace;
+			PassParameters->SourceMipIndex = (GRHISupportsTextureViews ? 0 : SourceMipIndex);
 			PassParameters->MipIndex = MipIndex;
 			PassParameters->NumMips = NumMips;
 			PassParameters->SvPositionToUVScale = FVector2f(1.0f / MipSize, 1.0f / MipSize);
diff --git a/Engine/Source/Runtime/Renderer/Private/ReflectionEnvironmentCapture.h b/Engine/Source/Runtime/Renderer/Private/ReflectionEnvironmentCapture.h
--- a/Engine/Source/Runtime/Renderer/Private/ReflectionEnvironmentCapture.h
+++ b/Engine/Source/Runtime/Renderer/Private/ReflectionEnvironmentCapture.h
@@ -23,3 +23,4 @@ inline uint32 GetNumMips(uint32 MipSize)
 }
 
 void ConvolveCubeMap(FRDGBuilder& GraphBuilder, FGlobalShaderMap* ShaderMap, uint32 CubeMipStart, uint32 CubeMipEnd, uint32 FaceStart, uint32 FaceCount, FRDGTexture* RDGSrcRenderTarget, FRDGTexture* RDGDstRenderTarget);
+void CreateCubeMips(FRDGBuilder& GraphBuilder, FGlobalShaderMap* ShaderMap, FRDGTexture* CubemapTexture);
diff --git a/Engine/Source/Runtime/Renderer/Private/ReflectionEnvironmentRealTimeCapture.cpp b/Engine/Source/Runtime/Renderer/Private/ReflectionEnvironmentRealTimeCapture.cpp
--- a/Engine/Source/Runtime/Renderer/Private/ReflectionEnvironmentRealTimeCapture.cpp
+++ b/Engine/Source/Runtime/Renderer/Private/ReflectionEnvironmentRealTimeCapture.cpp
@@ -97,7 +97,7 @@ class FDownsampleCubeFaceCS : public FGlobalShader
 	END_SHADER_PARAMETER_STRUCT()
 
 public:
-	static bool ShouldCompilePermutation(const FGlobalShaderPermutationParameters& Parameters) { return GetMaxSupportedFeatureLevel(Parameters.Platform) >= ERHIFeatureLevel::ES3_1; }
+	static bool ShouldCompilePermutation(const FGlobalShaderPermutationParameters& Parameters) { return !IsMobilePlatform(Parameters.Platform); }
 	static void ModifyCompilationEnvironment(const FGlobalShaderPermutationParameters& Parameters, FShaderCompilerEnvironment& OutEnvironment)
 	{
 		FGlobalShader::ModifyCompilationEnvironment(Parameters, OutEnvironment);
@@ -1036,7 +1036,7 @@ void FScene::AllocateAndCaptureFrameSkyEnvMap(
 		if (bIsMobilePlatform)
 		{
 			FRDGTextureRef SkyCubeTexture = GraphBuilder.RegisterExternalTexture(SkyRenderTarget, TEXT("SkyRenderTarget"));
-			MobileReflectionEnvironmentCapture::CreateCubeMips(GraphBuilder, GetGlobalShaderMap(FeatureLevel), SkyCubeTexture);
+			CreateCubeMips(GraphBuilder, GetGlobalShaderMap(FeatureLevel), SkyCubeTexture);
 		}
 		else
 		{
