Extend RVT corruption workaround.
Setting "r.VT.RVT.PageCorruptionFix 2" will apply dummy SetPipelineState() call to all page renders.
The old setting of "r.VT.RVT.PageCorruptionFix 1" continues to only apply to the first page.
This is the best current workaround for artifacts seen on some hardware when editing material layer parameters for material instances that contribute to RVT.
#Synced-CL 48777116

diff --git a/Engine/Source/Runtime/Renderer/Private/VT/RuntimeVirtualTextureRender.cpp b/Engine/Source/Runtime/Renderer/Private/VT/RuntimeVirtualTextureRender.cpp
--- a/Engine/Source/Runtime/Renderer/Private/VT/RuntimeVirtualTextureRender.cpp
+++ b/Engine/Source/Runtime/Renderer/Private/VT/RuntimeVirtualTextureRender.cpp
@@ -48,10 +48,11 @@ namespace RuntimeVirtualTexture
 		TEXT("Compress texture data direct to the physical texture on platforms that support it."),
 		ECVF_RenderThreadSafe);
 
-	static TAutoConsoleVariable<bool> CVarVTApplyPageCorruptionFix(
+	static TAutoConsoleVariable<int32> CVarVTApplyPageCorruptionFix(
 		TEXT("r.VT.RVT.PageCorruptionFix"),
-		false,
-		TEXT("Apply change that has been found to fix some rare page corruption on PC."),
+		0,
+		TEXT("Apply change that has been found to fix some rare page corruption on PC.\n")
+		TEXT("0: off (default), 1: apply fix on first rendered page per frame, 2: apply fix on all rendered pages."),
 		ECVF_RenderThreadSafe);
 
     int32 RenderCaptureNextRVTPagesDraws = 0;
@@ -2085,28 +2086,32 @@ namespace RuntimeVirtualTexture
 	/** Mesh render pass prologue to set the viewport. Also applies a page corruption workaround when that is enabled. */
 	void MeshPassPrologue(FRHICommandList& RHICmdList, FIntRect const& InViewRect, int32 InPageIndex, EShaderPlatform InShaderPlatform)
 	{
-		if (InPageIndex == 0 && CVarVTApplyPageCorruptionFix.GetValueOnRenderThread() && IsPCPlatform(InShaderPlatform) && IsD3DPlatform(InShaderPlatform))
+		if (IsPCPlatform(InShaderPlatform) && IsD3DPlatform(InShaderPlatform))
 		{
-			// Workaround fix for an issue where runtime virtual texture page corruption causes square artifacts.
-			// Repro of the bug is rare. But it's been found that inserting a single call to ID3D12GraphicsCommandList::SetPipelineState()
-			// before rendering any RVT pages resolves the issue.
-			FGlobalShaderMap* GlobalShaderMap = GetGlobalShaderMap(GMaxRHIFeatureLevel);
-			TShaderMapRef<FShader_VirtualTextureCopy_VS> VertexShader(GlobalShaderMap);
-			TShaderMapRef<FShader_VirtualTextureCopy_PS<ERuntimeVirtualTextureMaterialType::BaseColor>> PixelShader(GlobalShaderMap);
-
-			FGraphicsPipelineStateInitializer GraphicsPSOInit;
-			GraphicsPSOInit.RenderTargetsEnabled = 1;
-			GraphicsPSOInit.RenderTargetFormats[0] = PF_B8G8R8A8;
-			GraphicsPSOInit.RenderTargetFlags[0] = TexCreate_None;
-			GraphicsPSOInit.NumSamples = 1;
-			GraphicsPSOInit.DepthStencilState = TStaticDepthStencilState<false, CF_Always>::GetRHI();
-			GraphicsPSOInit.BlendState = TStaticBlendState<CW_RGBA, BO_Add, BF_One, BF_Zero, BO_Add, BF_One, BF_Zero>::GetRHI();
-			GraphicsPSOInit.RasterizerState = TStaticRasterizerState<>::GetRHI();
-			GraphicsPSOInit.PrimitiveType = PT_TriangleList;
-			GraphicsPSOInit.BoundShaderState.VertexDeclarationRHI = GetVertexDeclarationFVector4();
-			GraphicsPSOInit.BoundShaderState.VertexShaderRHI = VertexShader.GetVertexShader();
-			GraphicsPSOInit.BoundShaderState.PixelShaderRHI = PixelShader.GetPixelShader();
-			SetGraphicsPipelineState(RHICmdList, GraphicsPSOInit, 0, EApplyRendertargetOption::DoNothing);
+			const int32 FixLevel = CVarVTApplyPageCorruptionFix.GetValueOnRenderThread();
+			if (FixLevel > 1 || (FixLevel == 1 && InPageIndex == 0))
+			{
+				// Workaround fix for an issue where runtime virtual texture page corruption causes square artifacts.
+				// Repro of the bug is rare. But it's been found that inserting a single call to ID3D12GraphicsCommandList::SetPipelineState()
+				// before rendering any RVT pages resolves the issue.
+				FGlobalShaderMap* GlobalShaderMap = GetGlobalShaderMap(GMaxRHIFeatureLevel);
+				TShaderMapRef<FShader_VirtualTextureCopy_VS> VertexShader(GlobalShaderMap);
+				TShaderMapRef<FShader_VirtualTextureCopy_PS<ERuntimeVirtualTextureMaterialType::BaseColor>> PixelShader(GlobalShaderMap);
+
+				FGraphicsPipelineStateInitializer GraphicsPSOInit;
+				GraphicsPSOInit.RenderTargetsEnabled = 1;
+				GraphicsPSOInit.RenderTargetFormats[0] = PF_B8G8R8A8;
+				GraphicsPSOInit.RenderTargetFlags[0] = TexCreate_None;
+				GraphicsPSOInit.NumSamples = 1;
+				GraphicsPSOInit.DepthStencilState = TStaticDepthStencilState<false, CF_Always>::GetRHI();
+				GraphicsPSOInit.BlendState = TStaticBlendState<CW_RGBA, BO_Add, BF_One, BF_Zero, BO_Add, BF_One, BF_Zero>::GetRHI();
+				GraphicsPSOInit.RasterizerState = TStaticRasterizerState<>::GetRHI();
+				GraphicsPSOInit.PrimitiveType = PT_TriangleList;
+				GraphicsPSOInit.BoundShaderState.VertexDeclarationRHI = GetVertexDeclarationFVector4();
+				GraphicsPSOInit.BoundShaderState.VertexShaderRHI = VertexShader.GetVertexShader();
+				GraphicsPSOInit.BoundShaderState.PixelShaderRHI = PixelShader.GetPixelShader();
+				SetGraphicsPipelineState(RHICmdList, GraphicsPSOInit, 0, EApplyRendertargetOption::DoNothing);
+			}
 		}
 
 		RHICmdList.SetViewport(static_cast<float>(InViewRect.Min.X), static_cast<float>(InViewRect.Min.Y), 0.0f, static_cast<float>(InViewRect.Max.X), static_cast<float>(InViewRect.Max.Y), 1.0f);
