Set RVT corruption workaround on by default and only apply for nvidia.
Reverted "r.VT.RVT.PageCorruptionFix" to only have 0/1 options again for simplicity.
#Synced-CL 48792737

diff --git a/Engine/Source/Runtime/Renderer/Private/VT/RuntimeVirtualTextureRender.cpp b/Engine/Source/Runtime/Renderer/Private/VT/RuntimeVirtualTextureRender.cpp
--- a/Engine/Source/Runtime/Renderer/Private/VT/RuntimeVirtualTextureRender.cpp
+++ b/Engine/Source/Runtime/Renderer/Private/VT/RuntimeVirtualTextureRender.cpp
@@ -48,11 +48,10 @@ namespace RuntimeVirtualTexture
 		TEXT("Compress texture data direct to the physical texture on platforms that support it."),
 		ECVF_RenderThreadSafe);
 
-	static TAutoConsoleVariable<int32> CVarVTApplyPageCorruptionFix(
+	static TAutoConsoleVariable<bool> CVarVTApplyPageCorruptionFix(
 		TEXT("r.VT.RVT.PageCorruptionFix"),
-		0,
-		TEXT("Apply change that has been found to fix some rare page corruption on PC.\n")
-		TEXT("0: off (default), 1: apply fix on first rendered page per frame, 2: apply fix on all rendered pages."),
+		true,
+		TEXT("Apply change that has been found to fix some rare page corruption on PC.\n"),
 		ECVF_RenderThreadSafe);
 
     int32 RenderCaptureNextRVTPagesDraws = 0;
@@ -2084,35 +2083,34 @@ namespace RuntimeVirtualTexture
 	}
 
 	/** Mesh render pass prologue to set the viewport. Also applies a page corruption workaround when that is enabled. */
-	void MeshPassPrologue(FRHICommandList& RHICmdList, FIntRect const& InViewRect, int32 InPageIndex, EShaderPlatform InShaderPlatform)
+	void MeshPassPrologue(FRHICommandList& RHICmdList, FIntRect const& InViewRect, EShaderPlatform InShaderPlatform)
 	{
-		if (IsPCPlatform(InShaderPlatform) && IsD3DPlatform(InShaderPlatform))
+#if PLATFORM_WINDOWS && !PLATFORM_CPU_ARM_FAMILY
+		static bool bIsNVidia = IsRHIDeviceNVIDIA();
+		if (bIsNVidia && CVarVTApplyPageCorruptionFix.GetValueOnRenderThread())
 		{
-			const int32 FixLevel = CVarVTApplyPageCorruptionFix.GetValueOnRenderThread();
-			if (FixLevel > 1 || (FixLevel == 1 && InPageIndex == 0))
-			{
-				// Workaround fix for an issue where runtime virtual texture page corruption causes square artifacts.
-				// Repro of the bug is rare. But it's been found that inserting a single call to ID3D12GraphicsCommandList::SetPipelineState()
-				// before rendering any RVT pages resolves the issue.
-				FGlobalShaderMap* GlobalShaderMap = GetGlobalShaderMap(GMaxRHIFeatureLevel);
-				TShaderMapRef<FShader_VirtualTextureCopy_VS> VertexShader(GlobalShaderMap);
-				TShaderMapRef<FShader_VirtualTextureCopy_PS<ERuntimeVirtualTextureMaterialType::BaseColor>> PixelShader(GlobalShaderMap);
-
-				FGraphicsPipelineStateInitializer GraphicsPSOInit;
-				GraphicsPSOInit.RenderTargetsEnabled = 1;
-				GraphicsPSOInit.RenderTargetFormats[0] = PF_B8G8R8A8;
-				GraphicsPSOInit.RenderTargetFlags[0] = TexCreate_None;
-				GraphicsPSOInit.NumSamples = 1;
-				GraphicsPSOInit.DepthStencilState = TStaticDepthStencilState<false, CF_Always>::GetRHI();
-				GraphicsPSOInit.BlendState = TStaticBlendState<CW_RGBA, BO_Add, BF_One, BF_Zero, BO_Add, BF_One, BF_Zero>::GetRHI();
-				GraphicsPSOInit.RasterizerState = TStaticRasterizerState<>::GetRHI();
-				GraphicsPSOInit.PrimitiveType = PT_TriangleList;
-				GraphicsPSOInit.BoundShaderState.VertexDeclarationRHI = GetVertexDeclarationFVector4();
-				GraphicsPSOInit.BoundShaderState.VertexShaderRHI = VertexShader.GetVertexShader();
-				GraphicsPSOInit.BoundShaderState.PixelShaderRHI = PixelShader.GetPixelShader();
-				SetGraphicsPipelineState(RHICmdList, GraphicsPSOInit, 0, EApplyRendertargetOption::DoNothing);
-			}
+			// Workaround fix for an issue where runtime virtual texture page corruption causes square artifacts from corrupt pages.
+			// Repro of the bug in cooked builds is rare. But there is a reliable editor repro when editing layered materials parameters: UE-323864. 
+			// It's been found that inserting a single call to ID3D12GraphicsCommandList::SetPipelineState() before rendering any RVT pages resolves the issue.
+			FGlobalShaderMap* GlobalShaderMap = GetGlobalShaderMap(GMaxRHIFeatureLevel);
+			TShaderMapRef<FShader_VirtualTextureCopy_VS> VertexShader(GlobalShaderMap);
+			TShaderMapRef<FShader_VirtualTextureCopy_PS<ERuntimeVirtualTextureMaterialType::BaseColor>> PixelShader(GlobalShaderMap);
+
+			FGraphicsPipelineStateInitializer GraphicsPSOInit;
+			GraphicsPSOInit.RenderTargetsEnabled = 1;
+			GraphicsPSOInit.RenderTargetFormats[0] = PF_B8G8R8A8;
+			GraphicsPSOInit.RenderTargetFlags[0] = TexCreate_None;
+			GraphicsPSOInit.NumSamples = 1;
+			GraphicsPSOInit.DepthStencilState = TStaticDepthStencilState<false, CF_Always>::GetRHI();
+			GraphicsPSOInit.BlendState = TStaticBlendState<CW_RGBA, BO_Add, BF_One, BF_Zero, BO_Add, BF_One, BF_Zero>::GetRHI();
+			GraphicsPSOInit.RasterizerState = TStaticRasterizerState<>::GetRHI();
+			GraphicsPSOInit.PrimitiveType = PT_TriangleList;
+			GraphicsPSOInit.BoundShaderState.VertexDeclarationRHI = GetVertexDeclarationFVector4();
+			GraphicsPSOInit.BoundShaderState.VertexShaderRHI = VertexShader.GetVertexShader();
+			GraphicsPSOInit.BoundShaderState.PixelShaderRHI = PixelShader.GetPixelShader();
+			SetGraphicsPipelineState(RHICmdList, GraphicsPSOInit, 0, EApplyRendertargetOption::DoNothing);
 		}
+#endif
 
 		RHICmdList.SetViewport(static_cast<float>(InViewRect.Min.X), static_cast<float>(InViewRect.Min.Y), 0.0f, static_cast<float>(InViewRect.Max.X), static_cast<float>(InViewRect.Max.Y), 1.0f);
 	}
@@ -2254,9 +2252,9 @@ namespace RuntimeVirtualTexture
 	    	{
 	    		GatherMeshesToDraw(DynamicMeshPassContext, GraphBuilder, Scene, View, MaterialType, RuntimeVirtualTextureId, vLevel, MaxLevel, bAllowCachedMeshDrawCommands);
 	    	}, 
-			[ViewRect = View->ViewRect, PageIndex, ShaderPlatform = Scene->GetShaderPlatform()](FRHICommandList& RHICmdList)
+			[ViewRect = View->ViewRect, ShaderPlatform = Scene->GetShaderPlatform()](FRHICommandList& RHICmdList)
 			{
-				MeshPassPrologue(RHICmdList, ViewRect, PageIndex, ShaderPlatform);
+				MeshPassPrologue(RHICmdList, ViewRect, ShaderPlatform);
 			});
 		}
 	}
