Added OverrideVirtualTextureThrottle option to scene capture component. That allows to remove all sorts of throttling on the RVT system for the duration of the capture. This has a performance cost but can ensure that textures are not blurry if the scene capture renders a region that requires a lot of RVT pages to be locked or rendered to.
This does *not* perform any kind of warmup, though, so only the existing RVT page requests will be processed during the capture when this option is on. Another capture has to be performed prior to this in order to generate the appropriate requests.

#Synced-CL 48774492

diff --git a/Engine/Source/Runtime/Engine/Classes/Components/SceneCaptureComponent.h b/Engine/Source/Runtime/Engine/Classes/Components/SceneCaptureComponent.h
--- a/Engine/Source/Runtime/Engine/Classes/Components/SceneCaptureComponent.h
+++ b/Engine/Source/Runtime/Engine/Classes/Components/SceneCaptureComponent.h
@@ -88,6 +88,13 @@ class USceneCaptureComponent : public USceneComponent
 	UPROPERTY(EditAnywhere, BlueprintReadWrite, Category = SceneCapture)
 	uint8 bCaptureOnMovement : 1;
 
+	/** 
+	 * Cancels throttling on the virtual texture system during the scene capture. Use this to make sure that the virtual textures don't look blurry on the scene capture. 
+	 * Please note that this does *not* pre-warm the RVT pages (use a prior scene capture in order to do this and issue the appropriate RVT requests) and has a performance cost. 
+	 */
+	UPROPERTY(EditAnywhere, BlueprintReadWrite, AdvancedDisplay, Category = SceneCapture)
+	uint8 bOverrideVirtualTextureThrottle : 1;
+
 	/** Capture a GPU frame for this scene capture, next time it renders (capture program must be connected). */
 	UPROPERTY(EditAnywhere, BlueprintReadWrite, Transient, DuplicateTransient, SkipSerialization, NonTransactional, AdvancedDisplay, Category = SceneCapture)
 	uint8 bCaptureGpuNextRender : 1;
diff --git a/Engine/Source/Runtime/Engine/Private/Components/SceneCaptureComponent.cpp b/Engine/Source/Runtime/Engine/Private/Components/SceneCaptureComponent.cpp
--- a/Engine/Source/Runtime/Engine/Private/Components/SceneCaptureComponent.cpp
+++ b/Engine/Source/Runtime/Engine/Private/Components/SceneCaptureComponent.cpp
@@ -171,6 +171,7 @@ USceneCaptureComponent::USceneCaptureComponent(const FObjectInitializer& ObjectI
 	CaptureSource = SCS_SceneColorHDR;
 	bCaptureEveryFrame = true;
 	bCaptureOnMovement = true;
+	bOverrideVirtualTextureThrottle = false;
 	bCaptureGpuNextRender = false;
 	bDumpGpuNextRender = false;
 	bSuppressGpuCaptureOrDump = false;
diff --git a/Engine/Source/Runtime/Renderer/Private/PlanarReflectionRendering.cpp b/Engine/Source/Runtime/Renderer/Private/PlanarReflectionRendering.cpp
--- a/Engine/Source/Runtime/Renderer/Private/PlanarReflectionRendering.cpp
+++ b/Engine/Source/Runtime/Renderer/Private/PlanarReflectionRendering.cpp
@@ -21,6 +21,7 @@
 #include "GlobalShader.h"
 #include "SceneRenderTargetParameters.h"
 #include "SceneRendering.h"
+#include "SceneRenderingUtils.h"
 #include "DeferredShadingRenderer.h"
 #include "ScenePrivate.h"
 #include "PostProcess/SceneFilterRendering.h"
@@ -35,6 +36,7 @@
 #include "ClearQuad.h"
 #include "SceneTextureParameters.h"
 #include "SceneViewExtension.h"
+#include "SceneCaptureRendering.h"
 #include "Substrate/Substrate.h"
 #include "SceneRenderBuilder.h"
 
@@ -422,13 +424,15 @@ static bool UpdatePlanarReflectionContents_RenderThread(
 
 extern void BuildProjectionMatrix(FIntPoint RenderTargetSize, float FOV, float InNearClippingPlane, FMatrix& ProjectionMatrix);
 
+namespace UE::Renderer::SceneCapture::Private
+{
+
 extern TArray<FSceneView*> SetupViewFamilyForSceneCapture(
 	FSceneViewFamily& ViewFamily,
 	USceneCaptureComponent* SceneCaptureComponent,
 	const TArrayView<const FSceneCaptureViewInfo> Views,
 	float MaxViewDistance,
-	bool bCaptureSceneColor,
-	bool bIsPlanarReflection,
+	ESetupViewFamilyFlags Flags,
 	FPostProcessSettings* PostProcessSettings,
 	const FFinalPostProcessSettings* InheritedMainViewPostProcessSettings,
 	float PostProcessBlendWeight,
@@ -439,8 +443,11 @@ extern void SetupSceneViewExtensionsForSceneCapture(
 	FSceneViewFamily& ViewFamily,
 	TConstArrayView<FSceneView*> Views);
 
+} // namespace UE::Renderer::SceneCapture::Private
+
 void FScene::UpdatePlanarReflectionContents(UPlanarReflectionComponent* CaptureComponent, FSceneRenderer& ParentSceneRenderer, ISceneRenderBuilder& SceneRenderBuilder)
 {
+	using namespace UE::Renderer::SceneCapture::Private;
 	check(CaptureComponent);
 
 	FIntPoint DesiredBufferSize = FSceneRenderer::GetDesiredInternalBufferSize(ParentSceneRenderer.ViewFamily);
@@ -544,12 +551,17 @@ void FScene::UpdatePlanarReflectionContents(UPlanarReflectionComponent* CaptureC
 	ViewExtensionContext.bStereoEnabled = true;
 	ViewFamily.ViewExtensions = GEngine->ViewExtensions->GatherActiveExtensions(ViewExtensionContext);
 
+	ESetupViewFamilyFlags SetupViewFamilyFlags = ESetupViewFamilyFlags::CaptureSceneColor | ESetupViewFamilyFlags::IsPlanarReflection;
+	if (CaptureComponent->bOverrideVirtualTextureThrottle)
+	{
+		SetupViewFamilyFlags |= ESetupViewFamilyFlags::OverrideVirtualTextureThrottle;
+	}
+
 	TArray<FSceneView*> Views = SetupViewFamilyForSceneCapture(
 		ViewFamily,
 		CaptureComponent,
 		SceneCaptureViewInfo, CaptureComponent->MaxViewDistanceOverride,
-		/* bCaptureSceneColor = */ true,
-		/* bIsPlanarReflection = */ true,
+		SetupViewFamilyFlags,
 		&PostProcessSettings,
 		/* InheritedMainViewPostProcessSettings = */ nullptr,
 		1.0f,
diff --git a/Engine/Source/Runtime/Renderer/Private/SceneCaptureRendering.cpp b/Engine/Source/Runtime/Renderer/Private/SceneCaptureRendering.cpp
--- a/Engine/Source/Runtime/Renderer/Private/SceneCaptureRendering.cpp
+++ b/Engine/Source/Runtime/Renderer/Private/SceneCaptureRendering.cpp
@@ -31,6 +31,7 @@
 #include "GlobalShader.h"
 #include "SceneRenderTargetParameters.h"
 #include "SceneRendering.h"
+#include "SceneRenderingUtils.h"
 #include "DeferredShadingRenderer.h"
 #include "ScenePrivate.h"
 #include "PostProcess/SceneFilterRendering.h"
@@ -663,13 +664,15 @@ void GetShowOnlyAndHiddenComponents(USceneCaptureComponent* SceneCaptureComponen
 	}
 }
 
+namespace UE::Renderer::SceneCapture::Private
+{
+
 TArray<FSceneView*> SetupViewFamilyForSceneCapture(
 	FSceneViewFamily& ViewFamily,
 	USceneCaptureComponent* SceneCaptureComponent,
 	const TArrayView<const FSceneCaptureViewInfo> Views,
 	float MaxViewDistance,
-	bool bCaptureSceneColor,
-	bool bIsPlanarReflection,
+	ESetupViewFamilyFlags Flags,
 	FPostProcessSettings* PostProcessSettings,
 	const FFinalPostProcessSettings* InheritedMainViewPostProcessSettings,
 	float PostProcessBlendWeight,
@@ -686,6 +689,8 @@ TArray<FSceneView*> SetupViewFamilyForSceneCapture(
 	ViewFamily.FrameNumber = ViewFamily.Scene->GetFrameNumber();
 	ViewFamily.FrameCounter = GFrameCounter;
 
+	ViewFamily.bOverrideVirtualTextureThrottle = EnumHasAnyFlags(Flags, ESetupViewFamilyFlags::OverrideVirtualTextureThrottle);
+
 	TArray<FSceneView*> ViewPtrArray;
 	ViewPtrArray.Reserve(Views.Num());
 
@@ -707,7 +712,7 @@ TArray<FSceneView*> SetupViewFamilyForSceneCapture(
 		ViewInitOptions.StereoViewIndex = SceneCaptureViewInfo.StereoViewIndex;
 		ViewInitOptions.ProjectionMatrix = SceneCaptureViewInfo.ProjectionMatrix;
 		ViewInitOptions.bIsSceneCapture = true;
-		ViewInitOptions.bIsPlanarReflection = bIsPlanarReflection;
+		ViewInitOptions.bIsPlanarReflection = EnumHasAnyFlags(Flags, ESetupViewFamilyFlags::IsPlanarReflection);
 		ViewInitOptions.FOV = SceneCaptureViewInfo.FOV;
 		ViewInitOptions.DesiredFOV = SceneCaptureViewInfo.FOV;
 
@@ -716,7 +721,7 @@ TArray<FSceneView*> SetupViewFamilyForSceneCapture(
 			ViewInitOptions.WorldToMetersScale = ViewFamily.Scene->GetWorld()->GetWorldSettings()->WorldToMeters;
 		}
 
-		if (bCaptureSceneColor)
+		if (EnumHasAnyFlags(Flags, ESetupViewFamilyFlags::CaptureSceneColor))
 		{
 			ViewFamily.EngineShowFlags.PostProcessing = 0;
 			ViewInitOptions.OverlayColor = FLinearColor::Black;
@@ -832,9 +837,7 @@ static FSceneRenderer* CreateSceneRendererForSceneCapture(
 	const FMatrix& ProjectionMatrix,
 	float MaxViewDistance,
 	float InFOV,
-	bool bCaptureSceneColor,
-	bool bCameraCut2D,
-	bool bCopyMainViewTemporalSettings2D,
+	ESceneRendererCreationFlags InFlags,
 	FPostProcessSettings* PostProcessSettings,
 	float PostProcessBlendWeight,
 	const AActor* ViewActor,
@@ -874,24 +877,39 @@ static FSceneRenderer* CreateSceneRendererForSceneCapture(
 			InheritedMainViewPostProcessSettings = &MainViewFamily->Views[0]->FinalPostProcessSettings;
 		}
 	}
+	else 
+	{
+		// These flags are reserved for 2D captures :
+		check(!EnumHasAnyFlags(InFlags, ESceneRendererCreationFlags::CameraCut2D | ESceneRendererCreationFlags::CopyMainViewTemporalSettings2D));
+	}
 
 	FSceneViewFamilyContext ViewFamily(FSceneViewFamily::ConstructionValues(
 		RenderTarget,
 		Scene,
 		SceneCaptureComponent->ShowFlags)
-		.SetResolveScene(!bCaptureSceneColor)
+		.SetResolveScene(!EnumHasAnyFlags(InFlags, ESceneRendererCreationFlags::CaptureSceneColor))
 		.SetRealtimeUpdate(SceneCaptureComponent->bCaptureEveryFrame || SceneCaptureComponent->bAlwaysPersistRenderingState));
 
 	FSceneViewExtensionContext ViewExtensionContext(Scene);
 	ViewFamily.ViewExtensions = GEngine->ViewExtensions->GatherActiveExtensions(ViewExtensionContext);
 
+	ESetupViewFamilyFlags SetupViewFamilyFlags = ESetupViewFamilyFlags::None;
+	if (EnumHasAnyFlags(InFlags, ESceneRendererCreationFlags::CaptureSceneColor))
+	{
+		SetupViewFamilyFlags |= ESetupViewFamilyFlags::CaptureSceneColor;
+	}
+
+	if (EnumHasAnyFlags(InFlags, ESceneRendererCreationFlags::OverrideVirtualTextureThrottle))
+	{
+		SetupViewFamilyFlags |= ESetupViewFamilyFlags::OverrideVirtualTextureThrottle;
+	}
+
 	TArray<FSceneView*> Views = SetupViewFamilyForSceneCapture(
 		ViewFamily,
 		SceneCaptureComponent,
 		MakeArrayView(&SceneCaptureViewInfo, 1),
 		MaxViewDistance, 
-		bCaptureSceneColor,
-		/* bIsPlanarReflection = */ false,
+		SetupViewFamilyFlags,
 		PostProcessSettings,
 		InheritedMainViewPostProcessSettings,
 		PostProcessBlendWeight,
@@ -945,9 +963,9 @@ static FSceneRenderer* CreateSceneRendererForSceneCapture(
 			Views[0]->bAllowTemporalJitter = false;
 		}
 
-		Views[0]->bCameraCut = bCameraCut2D;
+		Views[0]->bCameraCut = EnumHasAnyFlags(InFlags, ESceneRendererCreationFlags::CameraCut2D);
 
-		if (bCopyMainViewTemporalSettings2D)
+		if (EnumHasAnyFlags(InFlags, ESceneRendererCreationFlags::CopyMainViewTemporalSettings2D))
 		{
 			const FSceneViewFamily* MainViewFamily = SceneCaptureComponent2D->MainViewFamily;
 			const FSceneView& SourceView = *MainViewFamily->Views[0];
@@ -1002,6 +1020,8 @@ static FSceneRenderer* CreateSceneRendererForSceneCapture(
 	return SceneRenderBuilder.CreateSceneRenderer(&ViewFamily);
 }
 
+} // namespace UE::Renderer::SceneCapture::Private
+
 FSceneCaptureCustomRenderPassUserData FSceneCaptureCustomRenderPassUserData::GDefaultData;
 
 class FSceneCapturePass final : public FCustomRenderPassBase
@@ -1077,6 +1097,7 @@ ESceneRenderGroupFlags GetSceneCaptureGroupFlags(USceneCaptureComponent* Capture
 
 void FScene::UpdateSceneCaptureContents(USceneCaptureComponent2D* CaptureComponent, ISceneRenderBuilder& SceneRenderBuilder)
 {
+	using namespace UE::Renderer::SceneCapture::Private;
 	check(CaptureComponent);
 
 	const ESceneRenderGroupFlags SceneRenderGroupFlags = GetSceneCaptureGroupFlags(CaptureComponent);
@@ -1345,18 +1366,41 @@ void FScene::UpdateSceneCaptureContents(USceneCaptureComponent2D* CaptureCompone
 			return;
 		}
 
-		
+		ESceneRendererCreationFlags SceneRendererCreationFlags = ESceneRendererCreationFlags::None;
+		if (bUseSceneColorTexture)
+		{
+			SceneRendererCreationFlags |= ESceneRendererCreationFlags::CaptureSceneColor;
+		}
+
 		// Copy temporal AA related settings for main view camera scene capture, to match jitter.  Don't match if the resolution divisor is set,
 		// if it's set to ignore screen percentage, or if it's final color, which will run its own AA.  For custom render passes (handled above),
 		// computed jitter results are copied from the main view later in FSceneRenderer::PrepareViewStateForVisibility, but this doesn't work
 		// for regular scene captures, because they run in a separate scene renderer before the main view, where the main view's results haven't
 		// been computed yet.
-		const bool bCopyMainViewTemporalSettings2D = (CaptureComponent->ShouldRenderWithMainViewCamera() && CaptureComponent->MainViewFamily &&
+		if ((CaptureComponent->ShouldRenderWithMainViewCamera() && CaptureComponent->MainViewFamily &&
 			CaptureComponent->MainViewResolutionDivisor.X <= 1 && CaptureComponent->MainViewResolutionDivisor.Y <= 1 && !CaptureComponent->ShouldIgnoreScreenPercentage() &&
 			CaptureComponent->CaptureSource != ESceneCaptureSource::SCS_FinalColorLDR &&
 			CaptureComponent->CaptureSource != ESceneCaptureSource::SCS_FinalColorHDR &&
-			CaptureComponent->CaptureSource != ESceneCaptureSource::SCS_FinalToneCurveHDR);
-		const bool bCameraCut2D = bCopyMainViewTemporalSettings2D ? CaptureComponent->MainViewFamily->Views[0]->bCameraCut : CaptureComponent->bCameraCutThisFrame;
+			CaptureComponent->CaptureSource != ESceneCaptureSource::SCS_FinalToneCurveHDR))
+		{
+			SceneRendererCreationFlags |= ESceneRendererCreationFlags::CopyMainViewTemporalSettings2D;
+			if (CaptureComponent->MainViewFamily->Views[0]->bCameraCut)
+			{
+				SceneRendererCreationFlags |= ESceneRendererCreationFlags::CameraCut2D;
+			}
+		}
+		else
+		{
+			if (CaptureComponent->bCameraCutThisFrame)
+			{
+				SceneRendererCreationFlags |= ESceneRendererCreationFlags::CameraCut2D;
+			}
+		}
+
+		if (CaptureComponent->bOverrideVirtualTextureThrottle)
+		{
+			SceneRendererCreationFlags |= ESceneRendererCreationFlags::OverrideVirtualTextureThrottle;
+		}
 
 		FSceneRenderer* SceneRenderer = CreateSceneRendererForSceneCapture(
 			SceneRenderBuilder,
@@ -1369,16 +1413,14 @@ void FScene::UpdateSceneCaptureContents(USceneCaptureComponent2D* CaptureCompone
 			ProjectionMatrix, 
 			CaptureComponent->MaxViewDistanceOverride, 
 			CaptureComponent->FOVAngle,
-			bUseSceneColorTexture,
-			bCameraCut2D,
-			bCopyMainViewTemporalSettings2D,
+			SceneRendererCreationFlags,
 			&CaptureComponent->PostProcessSettings, 
 			CaptureComponent->PostProcessBlendWeight,
 			CaptureComponent->GetViewOwner());
 
 		check(SceneRenderer != nullptr);
 
-		SceneRenderer->Views[0].bSceneCaptureMainViewJitter = bCopyMainViewTemporalSettings2D;
+		SceneRenderer->Views[0].bSceneCaptureMainViewJitter = EnumHasAnyFlags(SceneRendererCreationFlags, ESceneRendererCreationFlags::CopyMainViewTemporalSettings2D);
 		SceneRenderer->Views[0].bFogOnlyOnRenderedOpaque = CaptureComponent->bConsiderUnrenderedOpaquePixelAsFullyTranslucent;
 
 		SceneRenderer->ViewFamily.SceneCaptureCompositeMode = CaptureComponent->CompositeMode;
@@ -1558,6 +1600,7 @@ static const int32 GCubeFaceViewportOffsets[6][2] =
 
 void FScene::UpdateSceneCaptureContents(USceneCaptureComponentCube* CaptureComponent, ISceneRenderBuilder& SceneRenderBuilder)
 {
+	using namespace UE::Renderer::SceneCapture::Private;
 	ESceneRenderGroupFlags SceneRenderGroupFlags = GetSceneCaptureGroupFlags(CaptureComponent);
 
 	struct FLocal
@@ -1659,7 +1702,17 @@ void FScene::UpdateSceneCaptureContents(USceneCaptureComponentCube* CaptureCompo
 
 		const FVector Location = CaptureComponent->GetComponentToWorld().GetTranslation();
 
-		bool bCaptureSceneColor = CaptureNeedsSceneColor(CaptureComponent->CaptureSource);
+		ESceneRendererCreationFlags SceneRendererCreationFlags = ESceneRendererCreationFlags::None;
+
+		if (CaptureNeedsSceneColor(CaptureComponent->CaptureSource))
+		{
+			SceneRendererCreationFlags |= ESceneRendererCreationFlags::CaptureSceneColor;
+		}
+
+		if (CaptureComponent->bOverrideVirtualTextureThrottle)
+		{
+			SceneRendererCreationFlags |= ESceneRendererCreationFlags::OverrideVirtualTextureThrottle;
+		}
 
 		SCENE_RENDER_GROUP_SCOPE(SceneRenderBuilder, MoveTemp(EventName), SceneRenderGroupFlags);
 
@@ -1673,8 +1726,6 @@ void FScene::UpdateSceneCaptureContents(USceneCaptureComponentCube* CaptureCompo
 				FMatrix ProjectionMatrix;
 				ComputeProjectionMatrix(TargetFace, ViewRotationMatrix, ProjectionMatrix);
 
-				constexpr bool bCameraCut2D = false;
-				constexpr bool bCopyMainViewTemporalSettings2D = false;
 				FSceneRenderer* SceneRenderer = CreateSceneRendererForSceneCapture(
 					SceneRenderBuilder,
 					this,
@@ -1686,9 +1737,7 @@ void FScene::UpdateSceneCaptureContents(USceneCaptureComponentCube* CaptureCompo
 					ProjectionMatrix,
 					CaptureComponent->MaxViewDistanceOverride,
 					FOVInDegrees,
-					bCaptureSceneColor,
-					bCameraCut2D,
-					bCopyMainViewTemporalSettings2D,
+					SceneRendererCreationFlags,
 					&CaptureComponent->PostProcessSettings,
 					CaptureComponent->PostProcessBlendWeight,
 					CaptureComponent->GetViewOwner(),
@@ -1817,19 +1866,24 @@ void FScene::UpdateSceneCaptureContents(USceneCaptureComponentCube* CaptureCompo
 				CubeFaceTarget.Get(),
 				this,
 				CaptureComponent->ShowFlags)
-				.SetResolveScene(!bCaptureSceneColor)
+				.SetResolveScene(!EnumHasAnyFlags(SceneRendererCreationFlags, ESceneRendererCreationFlags::CaptureSceneColor))
 				.SetRealtimeUpdate(CaptureComponent->bCaptureEveryFrame || CaptureComponent->bAlwaysPersistRenderingState));
 
 			FSceneViewExtensionContext ViewExtensionContext(this);
 			ViewFamily.ViewExtensions = GEngine->ViewExtensions->GatherActiveExtensions(ViewExtensionContext);
 
+			ESetupViewFamilyFlags SetupViewFamilyFlags = ESetupViewFamilyFlags::None;
+			if (EnumHasAnyFlags(SceneRendererCreationFlags, ESceneRendererCreationFlags::CaptureSceneColor))
+			{
+				SetupViewFamilyFlags |= ESetupViewFamilyFlags::CaptureSceneColor;
+			}
+
 			TArray<FSceneView*> Views = SetupViewFamilyForSceneCapture(
 				ViewFamily,
 				CaptureComponent,
 				SceneCaptureViewInfos,
 				CaptureComponent->MaxViewDistanceOverride,
-				bCaptureSceneColor,
-				/* bIsPlanarReflection = */ false,
+				SetupViewFamilyFlags,
 				&CaptureComponent->PostProcessSettings,
 				nullptr,
 				CaptureComponent->PostProcessBlendWeight,
diff --git a/Engine/Source/Runtime/Renderer/Private/SceneRenderingUtils.h b/Engine/Source/Runtime/Renderer/Private/SceneRenderingUtils.h
--- a/Engine/Source/Runtime/Renderer/Private/SceneRenderingUtils.h
+++ b/Engine/Source/Runtime/Renderer/Private/SceneRenderingUtils.h
@@ -55,3 +55,25 @@ inline void AddBatchedElementsPass(const FBatchedElements& BatchedElements, FRDG
 		}
 	);
 }
+
+namespace UE::Renderer::SceneCapture::Private
+{
+enum class ESetupViewFamilyFlags
+{
+	None = 0,
+	CaptureSceneColor = (1 << 0),
+	IsPlanarReflection = (1 << 1),
+	OverrideVirtualTextureThrottle = (1 << 2),
+};
+ENUM_CLASS_FLAGS(ESetupViewFamilyFlags);
+
+enum class ESceneRendererCreationFlags
+{
+	None = 0,
+	CaptureSceneColor = (1 << 0),
+	CameraCut2D = (1 << 1),	// Only valid for 2D captures
+	CopyMainViewTemporalSettings2D = (1 << 2), // Only valid for 2D captures
+	OverrideVirtualTextureThrottle = (1 << 3),
+};
+ENUM_CLASS_FLAGS(ESceneRendererCreationFlags);
+} // namespace UE::Renderer::SceneCapture::Private
\ No newline at end of file
