[Lit Translucency] Improvements to r.TranslucencyLightingVolume.MarkVoxels when using spatial blur.
- Dilate marked voxels to avoid energy loss due to blurring using unlit voxels.
- Skip unmarked voxels in spatial blur.
- Refactored GatherTranslucencyVolumeMarkedVoxels(...) to batch cascade dispatches by pass enabling more overlap on GPU.

- RenderDeferredLighting on test level using 128^3 TLV: 5.6ms ->5ms (-0.6ms).

#Synced-CL 48802085

diff --git a/Engine/Shaders/Private/TranslucentLightingShaders.usf b/Engine/Shaders/Private/TranslucentLightingShaders.usf
--- a/Engine/Shaders/Private/TranslucentLightingShaders.usf
+++ b/Engine/Shaders/Private/TranslucentLightingShaders.usf
@@ -95,6 +95,48 @@ float3 GetTranslatedWorldPositionFromVolumeCoord(uint3 VolumeCoord)
 	return View.TranslucencyLightingVolumeMin[VolumeCascadeIndex].xyz + (VolumeCoord + .5f) * View.TranslucencyLightingVolumeInvSize[VolumeCascadeIndex].w;
 }
 
+#ifdef DilateMarkedVoxelsCS
+
+RWTexture3D<uint> RWVolumeMarkTexture;
+Texture3D<uint> VolumeMarkTexture;
+uint3 VolumeSize;
+
+void MarkVoxel(int3 VoxelCoord)
+{
+	if (all(VoxelCoord >= 0) && all(VoxelCoord < (int3)VolumeSize))
+	{
+		RWVolumeMarkTexture[VoxelCoord] = 255;
+	}
+}
+
+[numthreads(THREADGROUP_SIZE_X, THREADGROUP_SIZE_Y, THREADGROUP_SIZE_Z)]
+void DilateMarkedVoxelsCS(uint3 DispatchThreadId : SV_DispatchThreadID)
+{
+	uint3 VolumeCoord = DispatchThreadId.xyz;
+
+	if (any(VolumeCoord >= VolumeSize))
+	{
+		return;
+	}
+	
+	if (VolumeMarkTexture[VolumeCoord] == 0)
+	{
+		return;
+	}
+
+	MarkVoxel(VolumeCoord + int3(0, 0, 0));
+	MarkVoxel(VolumeCoord + int3(1, 1, 1));
+	MarkVoxel(VolumeCoord + int3(-1, 1, 1));
+	MarkVoxel(VolumeCoord + int3(1, -1, 1));
+	MarkVoxel(VolumeCoord + int3(-1, -1, 1));
+	MarkVoxel(VolumeCoord + int3(1, 1, -1));
+	MarkVoxel(VolumeCoord + int3(-1, 1, -1));
+	MarkVoxel(VolumeCoord + int3(1, -1, -1));
+	MarkVoxel(VolumeCoord + int3(-1, -1, -1));
+}
+
+#endif // DilateMarkedVoxelsCS
+
 #ifdef GatherMarkedVoxelsCS
 
 Texture3D<uint> VolumeMarkTexture;
@@ -190,7 +232,7 @@ void FilterTranslucentVolumeCS(uint3 DispatchThreadId : SV_DispatchThreadID)
 		return;
 	}
 
-	#if !USE_SPATIAL_FILTER && USE_MARK_TEXTURE
+	#if USE_MARK_TEXTURE
 	{
 		// Skip empty voxels
 		if (TranslucencyLightingVolumeMarkTexture[VolumeCoord] == 0)
diff --git a/Engine/Source/Runtime/Renderer/Private/DeferredShadingRenderer.cpp b/Engine/Source/Runtime/Renderer/Private/DeferredShadingRenderer.cpp
--- a/Engine/Source/Runtime/Renderer/Private/DeferredShadingRenderer.cpp
+++ b/Engine/Source/Runtime/Renderer/Private/DeferredShadingRenderer.cpp
@@ -3280,10 +3280,10 @@ void FDeferredShadingSceneRenderer::Render(FRDGBuilder& GraphBuilder, const FSce
 			}
 #endif
 
-			GatherTranslucencyVolumeMarkedVoxels(GraphBuilder);
-
 			const FSortedLightSetSceneInfo& SortedLightSet = *GatherAndSortLightsTask.GetResult();
 
+			GatherTranslucencyVolumeMarkedVoxels(GraphBuilder, SortedLightSet);
+
 			RenderLights(GraphBuilder, SceneTextures, LightingChannelsTexture, SortedLightSet);
 
 			if (MegaLightsContext.IsValid())
diff --git a/Engine/Source/Runtime/Renderer/Private/DeferredShadingRenderer.h b/Engine/Source/Runtime/Renderer/Private/DeferredShadingRenderer.h
--- a/Engine/Source/Runtime/Renderer/Private/DeferredShadingRenderer.h
+++ b/Engine/Source/Runtime/Renderer/Private/DeferredShadingRenderer.h
@@ -860,7 +860,7 @@ private:
 		FTranslucencyLightingVolumeTextures& Textures,
 		const FSortedLightSetSceneInfo& SortedLightSet);
 
-	void GatherTranslucencyVolumeMarkedVoxels(FRDGBuilder& GraphBuilder);
+	void GatherTranslucencyVolumeMarkedVoxels(FRDGBuilder& GraphBuilder, const FSortedLightSetSceneInfo& SortedLightSet);
 
 	/** Render stationary light overlap as complexity to scene color. */
 	void RenderStationaryLightOverlap(
diff --git a/Engine/Source/Runtime/Renderer/Private/TranslucentLighting.cpp b/Engine/Source/Runtime/Renderer/Private/TranslucentLighting.cpp
--- a/Engine/Source/Runtime/Renderer/Private/TranslucentLighting.cpp
+++ b/Engine/Source/Runtime/Renderer/Private/TranslucentLighting.cpp
@@ -797,6 +797,45 @@ void FProjectedShadowInfo::RenderTranslucencyDepths(FRDGBuilder& GraphBuilder, F
 	});
 }
 
+class FDilateMarkedVoxelsCS : public FGlobalShader
+{
+public:
+	DECLARE_GLOBAL_SHADER(FDilateMarkedVoxelsCS);
+	SHADER_USE_PARAMETER_STRUCT(FDilateMarkedVoxelsCS, FGlobalShader);
+
+	BEGIN_SHADER_PARAMETER_STRUCT(FParameters, )
+		SHADER_PARAMETER_RDG_TEXTURE_UAV(RWTexture3D, RWVolumeMarkTexture)
+		SHADER_PARAMETER_RDG_TEXTURE(Texture3D, VolumeMarkTexture)
+
+		SHADER_PARAMETER(FIntVector, VolumeSize)
+	END_SHADER_PARAMETER_STRUCT()
+
+	using FPermutationDomain = TShaderPermutationDomain<>;
+
+	static bool ShouldCompilePermutation(const FGlobalShaderPermutationParameters& Parameters)
+	{
+		return IsFeatureLevelSupported(Parameters.Platform, ERHIFeatureLevel::SM5);
+	}
+
+	static FIntVector GetGroupSize()
+	{
+		return FIntVector(4, 4, 4);
+	}
+
+	static void ModifyCompilationEnvironment(const FGlobalShaderPermutationParameters& Parameters, FShaderCompilerEnvironment& OutEnvironment)
+	{
+		FGlobalShader::ModifyCompilationEnvironment(Parameters, OutEnvironment);
+
+		const FIntVector GroupSize = GetGroupSize();
+
+		OutEnvironment.SetDefine(TEXT("THREADGROUP_SIZE_X"), GroupSize.X);
+		OutEnvironment.SetDefine(TEXT("THREADGROUP_SIZE_Y"), GroupSize.Y);
+		OutEnvironment.SetDefine(TEXT("THREADGROUP_SIZE_Z"), GroupSize.Z);
+	}
+};
+
+IMPLEMENT_GLOBAL_SHADER(FDilateMarkedVoxelsCS, "/Engine/Private/TranslucentLightingShaders.usf", "DilateMarkedVoxelsCS", SF_Compute);
+
 class FGatherMarkedVoxelsCS : public FGlobalShader
 {
 public:
@@ -923,7 +962,6 @@ public:
 		{
 			PermutationVector.Set<FUseTemporalFilter>(false);
 			PermutationVector.Set<FInjectMegaLights>(false);
-			PermutationVector.Set<FUseMarkTexture>(false);
 		}
 
 		return PermutationVector;
@@ -2250,16 +2288,20 @@ static void DrawDebugTranslucencyLightingVolume(
 	}
 }
 
+static bool IsMegaLightsTranslucencyVolumeEnabled(const FSceneViewFamily& ViewFamily, const FSortedLightSetSceneInfo& SortedLightSet)
+{
+	return MegaLights::IsEnabled(ViewFamily)
+		&& MegaLights::UseTranslucencyVolume()
+		&& SortedLightSet.MegaLightsLightStart < SortedLightSet.SortedLights.Num();
+}
+
 void FilterTranslucencyLightingVolume(
 	FRDGBuilder& GraphBuilder,
 	TArrayView<FViewInfo> Views,
 	FTranslucencyLightingVolumeTextures& Textures,
 	const FSortedLightSetSceneInfo& SortedLightSet)
 {
-	const bool bMegaLightsTranslucencyVolume = Views.Num() > 0
-		&& MegaLights::IsEnabled(*Views[0].Family)
-		&& MegaLights::UseTranslucencyVolume()
-		&& SortedLightSet.MegaLightsLightStart < SortedLightSet.SortedLights.Num();
+	const bool bMegaLightsTranslucencyVolume = Views.Num() > 0 && IsMegaLightsTranslucencyVolumeEnabled(*Views[0].Family, SortedLightSet);
 
 	const bool bSpatialFilter = bMegaLightsTranslucencyVolume
 		? MegaLights::IsTranslucencyVolumeSpatialFilterEnabled()
@@ -2524,6 +2566,9 @@ void FilterTranslucencyLightingVolume(
 
 				const FIntVector VolumeSize = FIntVector(TranslucencyLightingVolumeDim);
 
+				FRDGTextureRef TranslucencyLightingVolumeMarkTexture = View.TranslucencyVolumeMarkData[VolumeCascadeIndex].MarkTexture;
+				const bool bUseVolumeMarkTexture = IsTranslucencyLightingVolumeUsingVoxelMarking() && TranslucencyLightingVolumeMarkTexture != nullptr;
+
 				auto* PassParameters = GraphBuilder.AllocParameters<FFilterTranslucentVolumeCS::FParameters>();
 				PassParameters->View = View.ViewUniformBuffer;
 				PassParameters->VolumeSize = VolumeSize;
@@ -2533,13 +2578,14 @@ void FilterTranslucencyLightingVolume(
 				PassParameters->TranslucencyLightingVolumeDirectional = InputVolumeDirectionalTexture;
 				PassParameters->TranslucencyLightingVolumeAmbientSampler = SamplerStateRHI;
 				PassParameters->TranslucencyLightingVolumeDirectionalSampler = SamplerStateRHI;
+				PassParameters->TranslucencyLightingVolumeMarkTexture = TranslucencyLightingVolumeMarkTexture;
 				PassParameters->RWTranslucencyLightingVolumeAmbient = GraphBuilder.CreateUAV(OutputVolumeAmbientTexture);
 				PassParameters->RWTranslucencyLightingVolumeDirectional = GraphBuilder.CreateUAV(OutputVolumeDirectionalTexture);
 
 				FFilterTranslucentVolumeCS::FPermutationDomain PermutationVector;
 				PermutationVector.Set<FFilterTranslucentVolumeCS::FUseTemporalFilter>(false);
 				PermutationVector.Set<FFilterTranslucentVolumeCS::FUseSpatialFilter>(true);
-				PermutationVector.Set<FFilterTranslucentVolumeCS::FUseMarkTexture>(false);
+				PermutationVector.Set<FFilterTranslucentVolumeCS::FUseMarkTexture>(bUseVolumeMarkTexture);
 				PermutationVector.Set<FFilterTranslucentVolumeCS::FInjectMegaLights>(false);
 				PermutationVector = FFilterTranslucentVolumeCS::RemapPermutation(PermutationVector);
 
@@ -2607,7 +2653,7 @@ static void SplitSimpleLightsByView(TConstArrayView<FViewInfo> Views, const FSor
 	}
 }
 
-void FDeferredShadingSceneRenderer::GatherTranslucencyVolumeMarkedVoxels(FRDGBuilder& GraphBuilder)
+void FDeferredShadingSceneRenderer::GatherTranslucencyVolumeMarkedVoxels(FRDGBuilder& GraphBuilder, const FSortedLightSetSceneInfo& SortedLightSet)
 {
 	if (!IsTranslucencyLightingVolumeUsingVoxelMarking() || !ViewFamily.EngineShowFlags.DirectLighting || !RHISupportsVolumeTextureRendering(ViewFamily.GetShaderPlatform()))
 	{
@@ -2615,13 +2661,29 @@ void FDeferredShadingSceneRenderer::GatherTranslucencyVolumeMarkedVoxels(FRDGBui
 	}
 
 	RDG_EVENT_SCOPE(GraphBuilder, "TranslucencyLightingVolumeGatherMarkedVoxels");
+	RDG_GPU_STAT_SCOPE(GraphBuilder, TranslucencyLightingVolumeGatherMarkedVoxels);
 
 	const FIntVector VolumeSize = FIntVector(GetTranslucencyLightingVolumeDim());
 
+	const bool bMegaLightsTranslucencyVolume = Views.Num() > 0 && IsMegaLightsTranslucencyVolumeEnabled(*Views[0].Family, SortedLightSet);
+
+	const bool bSpatialFilter = bMegaLightsTranslucencyVolume
+		? MegaLights::IsTranslucencyVolumeSpatialFilterEnabled()
+		: GTranslucencyLightingVolumeBlur != 0;
+
 	for (int32 ViewIndex = 0; ViewIndex < Views.Num(); ViewIndex++)
 	{
 		FViewInfo& View = Views[ViewIndex];
 
+		struct FGatherMarkDataResources
+		{
+			FRDGTextureRef DilatedMarkTexture = nullptr;
+			FRDGTextureUAVRef DilatedMarkTextureUAV = nullptr;
+		};
+
+		TStaticArray<FGatherMarkDataResources, TVC_MAX> GatherMarkDataResources;
+
+		// Allocate and clear UAVs
 		for (uint32 VolumeCascadeIndex = 0; VolumeCascadeIndex < TVC_MAX; VolumeCascadeIndex++)
 		{
 			FViewInfo::FTranslucencyVolumeMarkData& VolumeMarkData = View.TranslucencyVolumeMarkData[VolumeCascadeIndex];
@@ -2633,13 +2695,72 @@ void FDeferredShadingSceneRenderer::GatherTranslucencyVolumeMarkedVoxels(FRDGBui
 				continue;
 			}
 
+			if (bSpatialFilter)
+			{
+				GatherMarkDataResources[VolumeCascadeIndex].DilatedMarkTexture = GraphBuilder.CreateTexture(View.TranslucencyVolumeMarkData[VolumeCascadeIndex].MarkTexture->Desc, TEXT("TranslucencyLightVolume.DilatedMarkTexture"));
+
+				GatherMarkDataResources[VolumeCascadeIndex].DilatedMarkTextureUAV = GraphBuilder.CreateUAV(GatherMarkDataResources[VolumeCascadeIndex].DilatedMarkTexture);
+
+				AddClearUAVPass(GraphBuilder, GatherMarkDataResources[VolumeCascadeIndex].DilatedMarkTextureUAV, 0u);
+			}
+
 			VolumeMarkData.VoxelAllocator = GraphBuilder.CreateBuffer(FRDGBufferDesc::CreateStructuredDesc(sizeof(uint32), 1), TEXT("TranslucencyLightingVolume.VoxelAllocator"));
 			VolumeMarkData.VoxelData = GraphBuilder.CreateBuffer(FRDGBufferDesc::CreateStructuredDesc(sizeof(uint32), VolumeSize.X * VolumeSize.Y * VolumeSize.Z), TEXT("TranslucencyLightingVolume.VoxelData"));
 			AddClearUAVPass(GraphBuilder, GraphBuilder.CreateUAV(VolumeMarkData.VoxelAllocator), 0);
 
-			// TODO: when using spatial blur need to dilate marked voxels
+			VolumeMarkData.VoxelIndirectArgs = GraphBuilder.CreateBuffer(FRDGBufferDesc::CreateIndirectDesc<FRHIDispatchIndirectParameters>(1), TEXT("TranslucencyLightingVolume.VoxelIndirectArgs"));
+		}
+
+		// Dilate marked voxels when using spatial blur
+		if (bSpatialFilter)
+		{
+			for (uint32 VolumeCascadeIndex = 0; VolumeCascadeIndex < TVC_MAX; VolumeCascadeIndex++)
+			{
+				FViewInfo::FTranslucencyVolumeMarkData& VolumeMarkData = View.TranslucencyVolumeMarkData[VolumeCascadeIndex];
+
+				const bool bUseVolumeMarkTexture = VolumeMarkData.MarkTexture != nullptr;
+
+				if (VolumeMarkData.MarkTexture == nullptr)
+				{
+					continue;
+				}
+
+				{
+					auto* PassParameters = GraphBuilder.AllocParameters<FDilateMarkedVoxelsCS::FParameters>();
+					PassParameters->RWVolumeMarkTexture = GatherMarkDataResources[VolumeCascadeIndex].DilatedMarkTextureUAV;
+					PassParameters->VolumeMarkTexture = View.TranslucencyVolumeMarkData[VolumeCascadeIndex].MarkTexture;
+					PassParameters->VolumeSize = VolumeSize;
+
+					FDilateMarkedVoxelsCS::FPermutationDomain PermutationVector;
+
+					auto ComputeShader = View.ShaderMap->GetShader<FDilateMarkedVoxelsCS>(PermutationVector);
+
+					FIntVector NumGroups = FComputeShaderUtils::GetGroupCount(VolumeSize, FDilateMarkedVoxelsCS::GetGroupSize());
+
+					FComputeShaderUtils::AddPass(
+						GraphBuilder,
+						RDG_EVENT_NAME("DilateMarkedVoxels(VolumeCascade=%d)"),
+						ComputeShader,
+						PassParameters,
+						NumGroups);
+
+					View.TranslucencyVolumeMarkData[VolumeCascadeIndex].MarkTexture = GatherMarkDataResources[VolumeCascadeIndex].DilatedMarkTexture;
+				}
+			}
+		}
+
+		// Gather marked voxels
+		for (uint32 VolumeCascadeIndex = 0; VolumeCascadeIndex < TVC_MAX; VolumeCascadeIndex++)
+		{
+			FViewInfo::FTranslucencyVolumeMarkData& VolumeMarkData = View.TranslucencyVolumeMarkData[VolumeCascadeIndex];
+
+			const bool bUseVolumeMarkTexture = VolumeMarkData.MarkTexture != nullptr;
+
+			if (VolumeMarkData.MarkTexture == nullptr)
+			{
+				continue;
+			}
 
-			// gather marked voxels
 			{
 				auto* PassParameters = GraphBuilder.AllocParameters<FGatherMarkedVoxelsCS::FParameters>();
 				PassParameters->RWVoxelAllocator = GraphBuilder.CreateUAV(VolumeMarkData.VoxelAllocator);
@@ -2660,10 +2781,20 @@ void FDeferredShadingSceneRenderer::GatherTranslucencyVolumeMarkedVoxels(FRDGBui
 					PassParameters,
 					NumGroups);
 			}
+		}
 
-			VolumeMarkData.VoxelIndirectArgs = GraphBuilder.CreateBuffer(FRDGBufferDesc::CreateIndirectDesc<FRHIDispatchIndirectParameters>(1), TEXT("TranslucencyLightingVolume.VoxelIndirectArgs"));
+		// Setup indirect args
+		for (uint32 VolumeCascadeIndex = 0; VolumeCascadeIndex < TVC_MAX; VolumeCascadeIndex++)
+		{
+			FViewInfo::FTranslucencyVolumeMarkData& VolumeMarkData = View.TranslucencyVolumeMarkData[VolumeCascadeIndex];
 
-			// setup indirect args
+			const bool bUseVolumeMarkTexture = VolumeMarkData.MarkTexture != nullptr;
+
+			if (VolumeMarkData.MarkTexture == nullptr)
+			{
+				continue;
+			}
+			
 			{
 				auto* PassParameters = GraphBuilder.AllocParameters<FInitIndirectArgsCS::FParameters>();
 				PassParameters->RWIndirectArgs = GraphBuilder.CreateUAV(VolumeMarkData.VoxelIndirectArgs);
