Keep HZB texture resolution at the Upper Bound  views size (View.UpperBoundViewRect) in order to avoid the RT reallocation when dynamic resolution is enabled.

#Synced-CL 48819522

diff --git a/Engine/Source/Runtime/Renderer/Private/CustomDepthRendering.cpp b/Engine/Source/Runtime/Renderer/Private/CustomDepthRendering.cpp
--- a/Engine/Source/Runtime/Renderer/Private/CustomDepthRendering.cpp
+++ b/Engine/Source/Runtime/Renderer/Private/CustomDepthRendering.cpp
@@ -349,7 +349,9 @@ bool FSceneRenderer::RenderCustomDepthPass(
 			RDG_EVENT_SCOPE_CONDITIONAL(GraphBuilder, Views.Num() > 1 && bDrawSceneViewsInOneNanitePass, "View%d (together with %d more)", ViewIndex, Views.Num() - 1);
 			
 			FViewInfo& View = Views[ViewIndex];
-			FIntRect ViewRect = bDrawSceneViewsInOneNanitePass ? FIntRect(0, 0, FamilySize.X, FamilySize.Y) : View.ViewRect;
+
+			const FIntRect HZBViewRect = bDrawSceneViewsInOneNanitePass ? FIntRect(0, 0, FamilySize.X, FamilySize.Y) : GetHZBViewRect(View);
+
 			auto NaniteRenderer = Nanite::IRenderer::Create(
 				GraphBuilder,
 				*Scene,
@@ -358,7 +360,7 @@ bool FSceneRenderer::RenderCustomDepthPass(
 				SharedContext,
 				RasterContext,
 				CullingConfig,
-				ViewRect,
+				HZBViewRect,
 				/* PrevHZB = */ nullptr
 			);
 
diff --git a/Engine/Source/Runtime/Renderer/Private/DeferredShadingRenderer.cpp b/Engine/Source/Runtime/Renderer/Private/DeferredShadingRenderer.cpp
--- a/Engine/Source/Runtime/Renderer/Private/DeferredShadingRenderer.cpp
+++ b/Engine/Source/Runtime/Renderer/Private/DeferredShadingRenderer.cpp
@@ -501,7 +501,7 @@ bool FDeferredShadingSceneRenderer::RenderHzb(FRDGBuilder& GraphBuilder, FRDGTex
 				GraphBuilder,
 				SceneDepthTexture,
 				/* VisBufferTexture = */ nullptr,
-				View.ViewRect,
+				GetHZBViewRect(View),
 				View.GetFeatureLevel(),
 				View.GetShaderPlatform(),
 				TEXT("HZBClosest"),
@@ -1510,14 +1510,14 @@ void FDeferredShadingSceneRenderer::RenderNanite(FRDGBuilder& GraphBuilder, cons
 			RDG_EVENT_SCOPE_CONDITIONAL(GraphBuilder, InViews.Num() > 1 && !bDrawSceneViewsInOneNanitePass, "View%u", ViewIndex);
 			RDG_EVENT_SCOPE_CONDITIONAL(GraphBuilder, InViews.Num() > 1 && bDrawSceneViewsInOneNanitePass, "View%u (together with %d more)", ViewIndex, InViews.Num() - 1);
 
-			FIntRect ViewRect = bDrawSceneViewsInOneNanitePass ? FIntRect(0, 0, FamilySize.X, FamilySize.Y) : View.ViewRect;
+			FIntRect ViewRect = bDrawSceneViewsInOneNanitePass ? FIntRect(0, 0, FamilySize.X, FamilySize.Y) : GetHZBViewRect(View);
 			CullingConfig.SetViewFlags(View);
 
 			float LODScaleFactor = 1.0f;
 			if (View.PrimaryScreenPercentageMethod == EPrimaryScreenPercentageMethod::TemporalUpscale &&
 				CVarNaniteViewMeshLODBiasEnable.GetValueOnRenderThread() != 0)
 			{
-				float TemporalUpscaleFactor = float(View.GetSecondaryViewRectSize().X) / float(ViewRect.Width());
+				float TemporalUpscaleFactor = ViewRect.Width() ? float(View.GetSecondaryViewRectSize().X) / float(ViewRect.Width()) : 1.0f;
 
 				LODScaleFactor = TemporalUpscaleFactor * FMath::Exp2(-CVarNaniteViewMeshLODBiasOffset.GetValueOnRenderThread());
 				LODScaleFactor = FMath::Min(LODScaleFactor, FMath::Exp2(-CVarNaniteViewMeshLODBiasMin.GetValueOnRenderThread()));
diff --git a/Engine/Source/Runtime/Renderer/Private/HairStrands/HairStrandsVisibility.cpp b/Engine/Source/Runtime/Renderer/Private/HairStrands/HairStrandsVisibility.cpp
--- a/Engine/Source/Runtime/Renderer/Private/HairStrands/HairStrandsVisibility.cpp
+++ b/Engine/Source/Runtime/Renderer/Private/HairStrands/HairStrandsVisibility.cpp
@@ -3022,7 +3022,7 @@ static void AddHairOnlyHZBPass(
 		GraphBuilder,
 		HairDepthTexture,
 		/* VisBufferTexture = */ nullptr,
-		View.ViewRect,
+		FSceneRendererBase::GetHZBViewRect(View),
 		View.GetFeatureLevel(),
 		View.GetShaderPlatform(),
 		TEXT("HZBHairClosest"),
diff --git a/Engine/Source/Runtime/Renderer/Private/MobileShadingRenderer.cpp b/Engine/Source/Runtime/Renderer/Private/MobileShadingRenderer.cpp
--- a/Engine/Source/Runtime/Renderer/Private/MobileShadingRenderer.cpp
+++ b/Engine/Source/Runtime/Renderer/Private/MobileShadingRenderer.cpp
@@ -3013,7 +3013,7 @@ void FMobileSceneRenderer::RenderHZB(FRDGBuilder& GraphBuilder, FRDGTextureRef S
 					GraphBuilder,
 					SceneDepthTexture,
 					/* VisBufferTexture = */ nullptr,
-					View.ViewRect,
+					GetHZBViewRect(View),
 					View.GetFeatureLevel(),
 					View.GetShaderPlatform(),
 					TEXT("MobileHZBFurthest"),
diff --git a/Engine/Source/Runtime/Renderer/Private/PostProcess/TemporalSuperResolution.cpp b/Engine/Source/Runtime/Renderer/Private/PostProcess/TemporalSuperResolution.cpp
--- a/Engine/Source/Runtime/Renderer/Private/PostProcess/TemporalSuperResolution.cpp
+++ b/Engine/Source/Runtime/Renderer/Private/PostProcess/TemporalSuperResolution.cpp
@@ -1970,7 +1970,7 @@ FDefaultTemporalUpscaler::FOutputs AddTemporalSuperResolutionPasses(
 			ResolutionFractionUpperBound = PrimaryResolutionFractionUpperBound * View.Family->SecondaryViewFraction * View.SceneViewInitOptions.OverscanResolutionFraction;
 		}
 
-		FIntPoint MaxRenderingViewSize = FSceneRenderer::ApplyResolutionFraction(*View.Family, View.UnconstrainedViewRect.Size(), ResolutionFractionUpperBound);
+		FIntPoint MaxRenderingViewSize = FSceneRenderer::ApplyResolutionFraction(View.UnconstrainedViewRect.Size(), ResolutionFractionUpperBound);
 
 		FIntPoint QuantizedMaxGuideSize;
 		QuantizeSceneBufferSize(MaxRenderingViewSize, QuantizedMaxGuideSize);
diff --git a/Engine/Source/Runtime/Renderer/Private/SceneRendering.cpp b/Engine/Source/Runtime/Renderer/Private/SceneRendering.cpp
--- a/Engine/Source/Runtime/Renderer/Private/SceneRendering.cpp
+++ b/Engine/Source/Runtime/Renderer/Private/SceneRendering.cpp
@@ -3130,7 +3130,7 @@ FSceneRenderer::FSceneRenderer(const FSceneViewFamily* InViewFamily, FHitProxyCo
 }
 
 // static
-FIntPoint FSceneRenderer::ApplyResolutionFraction(const FSceneViewFamily& ViewFamily, const FIntPoint& UnscaledViewSize, float ResolutionFraction)
+FIntPoint FSceneRenderer::ApplyResolutionFraction(const FIntPoint& UnscaledViewSize, float ResolutionFraction)
 {
 	FIntPoint ViewSize;
 
@@ -3215,7 +3215,7 @@ FIntPoint FSceneRenderer::GetDesiredInternalBufferSize(const FSceneViewFamily& V
 		// Note: This ensures that custom passes (rendered with the main renderer) ignore screen percentage, like regular scene captures.
 		const float AdjustedResolutionFractionUpperBounds = View->CustomRenderPass ? 1.0f : (View->SceneViewInitOptions.OverridePrimaryResolutionFraction > 0.0 ? (View->SceneViewInitOptions.OverridePrimaryResolutionFraction * ViewFamily.SecondaryViewFraction)  : ResolutionFractionUpperBound);
 		
-		FIntPoint ViewSize = ApplyResolutionFraction(ViewFamily, View->UnconstrainedViewRect.Size(), AdjustedResolutionFractionUpperBounds);
+		FIntPoint ViewSize = ApplyResolutionFraction(View->UnconstrainedViewRect.Size(), AdjustedResolutionFractionUpperBounds);
 		FIntPoint ViewRectMin = QuantizeViewRectMin(FIntPoint(
 			FMath::CeilToInt(View->UnconstrainedViewRect.Min.X * AdjustedResolutionFractionUpperBounds),
 			FMath::CeilToInt(View->UnconstrainedViewRect.Min.Y * AdjustedResolutionFractionUpperBounds)));
@@ -3315,7 +3315,7 @@ void FSceneRenderer::PrepareViewRectsForRendering()
 		// The base pass have to respect FSceneView::UnscaledViewRect.
 		for (FViewInfo& View : Views)
 		{
-			View.ViewRect = View.UnscaledViewRect;
+			View.ViewRect = View.UpperBoundViewRect = View.UnscaledViewRect;
 		}
 
 		ComputeFamilySize();
@@ -3323,7 +3323,8 @@ void FSceneRenderer::PrepareViewRectsForRendering()
 		return;
 	}
 
-	float PrimaryResolutionFraction = DynamicResolutionFractions[GDynamicPrimaryResolutionFraction];
+	const float PrimaryResolutionFraction = DynamicResolutionFractions[GDynamicPrimaryResolutionFraction];
+	const float PrimaryMaxResolutionFraction = DynamicResolutionUpperBounds[GDynamicPrimaryResolutionFraction];
 	{
 		// Ensure screen percentage show flag is respected. Prefer to check() rather rendering at a differen screen percentage
 		// to make sure the renderer does not lie how a frame as been rendering to a dynamic resolution heuristic.
@@ -3357,28 +3358,34 @@ void FSceneRenderer::PrepareViewRectsForRendering()
 	}
 	
 	// Compute final resolution fraction.
-	float ResolutionFraction = PrimaryResolutionFraction * ViewFamily.SecondaryViewFraction * LensDistortionResolutionFraction;
+	const float ResolutionFraction = PrimaryResolutionFraction * ViewFamily.SecondaryViewFraction * LensDistortionResolutionFraction;
+	const float ResolutionMaxFraction = PrimaryMaxResolutionFraction * ViewFamily.SecondaryViewFraction * LensDistortionResolutionFraction;
 
 	// Checks that view rects are correctly initialized.
 	for (int32 i = 0; i < Views.Num(); i++)
 	{
 		FViewInfo& View = Views[i];
 
-		float ViewResolutionFraction = View.SceneViewInitOptions.OverridePrimaryResolutionFraction > 0.0 ? (View.SceneViewInitOptions.OverridePrimaryResolutionFraction * ViewFamily.SecondaryViewFraction) : ResolutionFraction;
+		auto InitViewRect = [&View, &MaxOverscanResolutionFraction, ViewFamilySecondaryViewFraction = ViewFamily.SecondaryViewFraction](float CurrentResFraction)
+			{
+				const float ViewResolutionFraction = View.SceneViewInitOptions.OverridePrimaryResolutionFraction > 0.0 ? (View.SceneViewInitOptions.OverridePrimaryResolutionFraction * ViewFamilySecondaryViewFraction) : CurrentResFraction;
 
-		FIntPoint ViewSize = ApplyResolutionFraction(ViewFamily, View.UnscaledViewRect.Size(), ViewResolutionFraction * View.SceneViewInitOptions.OverscanResolutionFraction);
-		FIntPoint ViewRectMin = QuantizeViewRectMin(FIntPoint(
-			FMath::CeilToInt(View.UnscaledViewRect.Min.X * ViewResolutionFraction * MaxOverscanResolutionFraction),
-			FMath::CeilToInt(View.UnscaledViewRect.Min.Y * ViewResolutionFraction * MaxOverscanResolutionFraction)));
+				FIntPoint ViewSize = ApplyResolutionFraction(View.UnscaledViewRect.Size(), ViewResolutionFraction * View.SceneViewInitOptions.OverscanResolutionFraction);
+				FIntPoint ViewRectMin = QuantizeViewRectMin(FIntPoint(
+					FMath::CeilToInt(View.UnscaledViewRect.Min.X * ViewResolutionFraction * MaxOverscanResolutionFraction),
+					FMath::CeilToInt(View.UnscaledViewRect.Min.Y * ViewResolutionFraction * MaxOverscanResolutionFraction)));
 
-		// Use the bottom-left view rect if requested, instead of top-left
-		if (CVarViewRectUseScreenBottom.GetValueOnRenderThread())
-		{
-			ViewRectMin.Y = FMath::CeilToInt( View.UnscaledViewRect.Max.Y * ViewFamily.SecondaryViewFraction ) - ViewSize.Y;
-		}
+				// Use the bottom-left view rect if requested, instead of top-left
+				if (CVarViewRectUseScreenBottom.GetValueOnRenderThread())
+				{
+					ViewRectMin.Y = FMath::CeilToInt(View.UnscaledViewRect.Max.Y * ViewFamilySecondaryViewFraction) - ViewSize.Y;
+				}
 
-		View.ViewRect.Min = ViewRectMin;
-		View.ViewRect.Max = ViewRectMin + ViewSize;
+				return FIntRect(ViewRectMin, ViewRectMin + ViewSize);
+			};
+
+		View.ViewRect = InitViewRect(ResolutionFraction);
+		View.UpperBoundViewRect = InitViewRect(ResolutionMaxFraction);
 
 		#if !UE_BUILD_SHIPPING
 		// For testing purpose, override the screen percentage method.
@@ -3405,6 +3412,7 @@ void FSceneRenderer::PrepareViewRectsForRendering()
 		}
 
 		check(View.ViewRect.Area() != 0);
+		check(View.UpperBoundViewRect.Area() != 0);
 #if DO_CHECK || USING_CODE_ANALYSIS
 		check(View.VerifyMembersChecks());
 #endif
@@ -3412,6 +3420,8 @@ void FSceneRenderer::PrepareViewRectsForRendering()
 		if (!ViewFamily.EngineShowFlags.HitProxies)
 		{
 			FIntPoint SecondaryViewRect = View.GetSecondaryViewRectSize();
+
+			const float ViewResolutionFraction = View.SceneViewInitOptions.OverridePrimaryResolutionFraction > 0.0 ? (View.SceneViewInitOptions.OverridePrimaryResolutionFraction * ViewFamily.SecondaryViewFraction) : ResolutionFraction;
 			GPixelRenderCounters.AddViewStatistics(ViewResolutionFraction, View.ViewRect.Size(), SecondaryViewRect.X * SecondaryViewRect.Y);
 		}
 	}
diff --git a/Engine/Source/Runtime/Renderer/Private/SceneRendering.h b/Engine/Source/Runtime/Renderer/Private/SceneRendering.h
--- a/Engine/Source/Runtime/Renderer/Private/SceneRendering.h
+++ b/Engine/Source/Runtime/Renderer/Private/SceneRendering.h
@@ -1143,6 +1143,9 @@ public:
 
 	/* Final position of the view in the final render target (in pixels), potentially scaled by ScreenPercentage */
 	FIntRect ViewRect;
+	
+	/* Final position of the view in the final render target (in pixels), potentially scaled by the Upper Bound of ScreenPercentage */
+	FIntRect UpperBoundViewRect;
 
 	/** 
 	 * The view's state, or NULL if no state exists.
@@ -2074,6 +2077,14 @@ public:
 	 */
 	virtual FDeferredShadingSceneRenderer* GetDeferredShadingSceneRenderer() { return nullptr; }
 
+	/**
+	* Utility function to retrieve from a View the appropriate ViewRect used to create an HZB
+	*/
+	static FIntRect GetHZBViewRect ( const FViewInfo& View )
+	{
+		return View.UpperBoundViewRect;
+	}
+
 private:
 	FSceneUniformBuffer SceneUniforms;
 	FSceneExtensionsRenderers SceneExtensionsRenderers;
@@ -2257,8 +2268,7 @@ public:
 	static bool ShouldCompositeEditorPrimitives(const FViewInfo& View);
 
 	/** Apply the ResolutionFraction on ViewSize, taking into account renderer's requirements. */
-	static FIntPoint ApplyResolutionFraction(
-		const FSceneViewFamily& ViewFamily, const FIntPoint& UnscaledViewSize, float ResolutionFraction);
+	static FIntPoint ApplyResolutionFraction(const FIntPoint& UnscaledViewSize, float ResolutionFraction);
 
 	/** Quantize the ViewRect.Min according to various renderer's downscale requirements. */
 	static FIntPoint QuantizeViewRectMin(const FIntPoint& ViewRectMin);
