Support dynamic deallocation of VSM data based on whether and VSM shadowed lights are active.
Enabled by default when MegaLights is enabled. Can be forced on (or off) with r.Shadow.Virtual.DynamicDeallocationMode.

#Synced-CL 48634010

diff --git a/Engine/Shaders/Private/VirtualShadowMaps/VirtualShadowMapPhysicalPageManagement.usf b/Engine/Shaders/Private/VirtualShadowMaps/VirtualShadowMapPhysicalPageManagement.usf
--- a/Engine/Shaders/Private/VirtualShadowMaps/VirtualShadowMapPhysicalPageManagement.usf
+++ b/Engine/Shaders/Private/VirtualShadowMaps/VirtualShadowMapPhysicalPageManagement.usf
@@ -165,6 +165,7 @@ void UpdatePhysicalPageAddresses(uint3 Index : SV_DispatchThreadID)
 	int PhysicalPageIndex = Index.x;
 	checkStructuredBufferAccessSlow(OutPhysicalPageMetaData, PhysicalPageIndex);
 	
+#if HAS_CACHE_DATA
 	FPhysicalPageMetaData PrevMetaData = OutPhysicalPageMetaData[PhysicalPageIndex];
 
 	FVSMPageOffset PrevGlobalPageOffset = FVSMPageOffset::Unpack(0u);
@@ -201,7 +202,6 @@ void UpdatePhysicalPageAddresses(uint3 Index : SV_DispatchThreadID)
 
 	if (bKeepPage)
 	{
-#if HAS_CACHE_DATA
 		// Propogate any invalidation flags from the previous page requests to the physical page
 		const uint PrevFlags = PrevPageRequestFlags[PrevGlobalPageOffset.GetResourceAddress()];
 		const uint InvalidationFlags = PrevFlags & VSM_EXTENDED_FLAG_ANY_INVALIDATED;
@@ -210,13 +210,15 @@ void UpdatePhysicalPageAddresses(uint3 Index : SV_DispatchThreadID)
 			// Add them to any previous flags
 			OutPhysicalPageMetaData[PhysicalPageIndex].Flags = PrevMetaData.Flags | InvalidationFlags;
 		}
-#endif
 	}
 	else
 	{
 		// Only need to zero out flags for it to be considered invalid
 		OutPhysicalPageMetaData[PhysicalPageIndex].Flags = 0;
 	}
+#else
+	OutPhysicalPageMetaData[PhysicalPageIndex].Flags = 0;
+#endif
 }
 
 int bDynamicPageInvalidation;
diff --git a/Engine/Source/Runtime/Renderer/Private/Shadows/ShadowSceneRenderer.cpp b/Engine/Source/Runtime/Renderer/Private/Shadows/ShadowSceneRenderer.cpp
--- a/Engine/Source/Runtime/Renderer/Private/Shadows/ShadowSceneRenderer.cpp
+++ b/Engine/Source/Runtime/Renderer/Private/Shadows/ShadowSceneRenderer.cpp
@@ -792,7 +792,7 @@ void FShadowSceneRenderer::DispatchVirtualShadowMapViewAndCullingSetup(FRDGBuild
 	SCOPED_NAMED_EVENT(FShadowSceneRenderer_DispatchVirtualShadowMapViewAndCullingSetup, FColor::Orange);
 
 	// Unconditionally update GPU physical pages (on all GPUs) with new VSM IDs/addresses
-	VirtualShadowMapArray.UpdatePhysicalPageAddresses(GraphBuilder);
+	VirtualShadowMapArray.UpdatePhysicalPageAddresses(GraphBuilder, SceneRenderer.ViewFamily);
 
 	if (!VirtualShadowMapShadows.IsEmpty() && UseNanite(SceneRenderer.ShaderPlatform))
 	{
diff --git a/Engine/Source/Runtime/Renderer/Private/VirtualShadowMaps/VirtualShadowMapArray.cpp b/Engine/Source/Runtime/Renderer/Private/VirtualShadowMaps/VirtualShadowMapArray.cpp
--- a/Engine/Source/Runtime/Renderer/Private/VirtualShadowMaps/VirtualShadowMapArray.cpp
+++ b/Engine/Source/Runtime/Renderer/Private/VirtualShadowMaps/VirtualShadowMapArray.cpp
@@ -869,6 +869,10 @@ void FVirtualShadowMapArray::Initialize(
 	UniformParameters.PerViewData = MakeEmptyVirtualShadowMapPerViewParameters(GraphBuilder);
 	UniformParameters.CachePrimitiveAsDynamic = GraphBuilder.CreateSRV(GSystemTextures.GetDefaultStructuredBuffer(GraphBuilder, sizeof(uint32)));
 
+	// These get updated/allocated lazily after we know that we have VSMs present in the scene
+	UniformParameters.PhysicalPagePool = GSystemTextures.GetZeroUIntArrayAtomicCompatDummy(GraphBuilder);
+	HZBPhysicalArrayRDG = nullptr;
+
 	if (bShouldRetainCacheInfo)
 	{
 		// Always reserve slots for the single-page SMs.
@@ -937,14 +941,6 @@ void FVirtualShadowMapArray::Initialize(
 			}
 		}
 
-		// If enabled, ensure we have a properly-sized physical page pool
-		// We can do this here since the pool is independent of the number of shadow maps
-		const int PoolArraySize = bEnableNaniteVisualization ? 3 : (ShouldCacheStaticSeparately() ? 2 : 1);
-		CacheManager->SetPhysicalPoolSize(GraphBuilder, GetPhysicalPoolSize(), PoolArraySize, GetMaxPhysicalPages());
-		PhysicalPagePoolRDG = GraphBuilder.RegisterExternalTexture(CacheManager->GetPhysicalPagePool());
-		PhysicalPageMetaDataRDG = GraphBuilder.RegisterExternalBuffer(CacheManager->GetPhysicalPageMetaData());
-		UniformParameters.PhysicalPagePool = PhysicalPagePoolRDG;
-
 		UniformParameters.CachePrimitiveAsDynamic = GraphBuilder.CreateSRV(CacheManager->UploadCachePrimitiveAsDynamic(GraphBuilder));
 	}
 	else
@@ -952,24 +948,8 @@ void FVirtualShadowMapArray::Initialize(
 		if (CacheManager)
 		{
 			CacheManager->FreePhysicalPool(GraphBuilder);
-		}
-		UniformParameters.PhysicalPagePool = GSystemTextures.GetZeroUIntArrayAtomicCompatDummy(GraphBuilder);
-	}
-
-	if (bShouldRetainCacheInfo && bUseHzbOcclusion)
-	{
-		const int HZBPoolArraySize = HasSeparateDynamicHZB() ? 2 : 1;
-		HZBPhysicalArray = CacheManager->SetHZBPhysicalPoolSize(GraphBuilder, GetHZBPhysicalPoolSize(), HZBPoolArraySize, PF_R32_FLOAT);
-		HZBPhysicalArrayRDG = GraphBuilder.RegisterExternalTexture(HZBPhysicalArray);
-	}
-	else
-	{
-		if (CacheManager)
-		{
 			CacheManager->FreeHZBPhysicalPool(GraphBuilder);
 		}
-		HZBPhysicalArray = nullptr;
-		HZBPhysicalArrayRDG = nullptr;
 	}
 
 	UpdateCachedUniformBuffers(GraphBuilder);
@@ -1062,7 +1042,8 @@ FIntPoint FVirtualShadowMapArray::GetHZBPhysicalPoolSize() const
 uint32 FVirtualShadowMapArray::GetTotalAllocatedPhysicalPages() const
 {
 	check(bInitialized);
-	return ShouldCacheStaticSeparately() ? (2U * UniformParameters.MaxPhysicalPages) : UniformParameters.MaxPhysicalPages;
+	// 2 array elements (static/dynamic)
+	return 2U * UniformParameters.MaxPhysicalPages;
 }
 
 TRDGUniformBufferRef<FVirtualShadowMapUniformParameters> FVirtualShadowMapArray::GetUncachedUniformBuffer(FRDGBuilder& GraphBuilder) const
@@ -1742,12 +1723,9 @@ void FVirtualShadowMapArray::PostRender(FRDGBuilder& GraphBuilder)
 		);
 	}
 
-	// If separate static/dynamic caching is enabled, we may need to merge some pages after rendering
-	if (ShouldCacheStaticSeparately()
 #if !UE_BUILD_SHIPPING
-		&& CVarDebugSkipMergePhysical.GetValueOnRenderThread() == 0
+	if (CVarDebugSkipMergePhysical.GetValueOnRenderThread() == 0)
 #endif
-		)
 	{
 		RDG_EVENT_SCOPE(GraphBuilder, "FVirtualShadowMapArray::MergeStaticPhysicalPages");
 
@@ -2024,19 +2002,65 @@ void FVirtualShadowMapArray::AppendPhysicalPageList(FRDGBuilder& GraphBuilder, b
 	);
 }
 
-void FVirtualShadowMapArray::UpdatePhysicalPageAddresses(FRDGBuilder& GraphBuilder)
+static TAutoConsoleVariable<int32> CVarVSMDynamicDeallocationMode(
+	TEXT("r.Shadow.Virtual.DynamicDeallocationMode"),
+	0,
+	TEXT("Controls the behavior of VSM resources being dynamically allocated/deallocated based on whether VSM shadows are present:\n")
+	TEXT("0 (default): Enables dynamic deallocation when MegaLights are also enabled, otherwise disabled\n")
+	TEXT("1: Enables dynamic deallocation always\n")
+	TEXT("2: Disables dynamic deallocation always\n")
+	TEXT("Dynamic deallocation can save memory if there are significant times when VSMs are not in use, but can add allocator churn."),
+	ECVF_RenderThreadSafe
+);
+
+static bool AllowVirtualShadowMapDynamicDeallocation(const FSceneViewFamily& ViewFamily)
 {
-	if (!IsEnabled())
-	{
+	int32 Mode = CVarVSMDynamicDeallocationMode.GetValueOnRenderThread();
+	return
+		(Mode == 1 ||
+		(Mode == 0 && MegaLights::IsEnabled(ViewFamily)));
+}
+
+void FVirtualShadowMapArray::UpdatePhysicalPageAddresses(FRDGBuilder& GraphBuilder, const FSceneViewFamily& ViewFamily)
+{
+	if (!IsEnabled() ||
+		(GetNumShadowMaps() == 0 && NextData.IsEmpty()))
+	{
+		// If we get here with no shadow maps and no cached data even for unreferenced entries it means we legitimately have no
+		// useful shadow data. Optionally we can deallocate the buffers here to save VRAM. There's a potential for allocator churn
+		// so we only do this in situations that we expect there to be a good chance that we aren't using VSMs at all for some
+		// reasonable chunk of time (i.e. MegaLights in use, etc)
+		if (AllowVirtualShadowMapDynamicDeallocation(ViewFamily))
+		{
+			CacheManager->FreePhysicalPool(GraphBuilder);
+			CacheManager->FreeHZBPhysicalPool(GraphBuilder);
+		}
 		return;
 	}
 
-	// NOTE: This past MUST run on all GPUs, as we still need to propogate changes to the VSM IDs even if
+	// At this point, ensure we have a properly-sized physical page pool and HZB
+	{
+		// 2 elements normally (static/dynamic cache), one extra if we want to store nanite visualization outputs in light space
+		const int PoolArraySize = bEnableNaniteVisualization ? 3 : 2;
+		CacheManager->SetPhysicalPoolSize(GraphBuilder, GetPhysicalPoolSize(), PoolArraySize, GetMaxPhysicalPages());
+		PhysicalPagePoolRDG = GraphBuilder.RegisterExternalTexture(CacheManager->GetPhysicalPagePool());
+		PhysicalPageMetaDataRDG = GraphBuilder.RegisterExternalBuffer(CacheManager->GetPhysicalPageMetaData());
+		UniformParameters.PhysicalPagePool = PhysicalPagePoolRDG;
+
+		if (bUseHzbOcclusion)
+		{
+			const int HZBPoolArraySize = HasSeparateDynamicHZB() ? 2 : 1;
+			CacheManager->SetHZBPhysicalPoolSize(GraphBuilder, GetHZBPhysicalPoolSize(), HZBPoolArraySize, PF_R32_FLOAT);
+			HZBPhysicalArrayRDG = GraphBuilder.RegisterExternalTexture(CacheManager->GetHZBPhysicalPool());
+		}
+	}
+
+	// NOTE: This pass MUST run on all GPUs, as we still need to propogate changes to the VSM IDs even if
 	// a given GPU may not do any rendering during this phase.
 	RDG_GPU_MASK_SCOPE(GraphBuilder, FRHIGPUMask::All());
 
 	FUpdatePhysicalPageAddresses::FParameters* PassParameters = GraphBuilder.AllocParameters<FUpdatePhysicalPageAddresses::FParameters>();
-	PassParameters->VirtualShadowMap		= GetUniformBuffer(0);
+	PassParameters->VirtualShadowMap		= GetUncachedUniformBuffer(GraphBuilder);
 	PassParameters->OutPhysicalPageMetaData = GraphBuilder.CreateUAV(PhysicalPageMetaDataRDG);
 
 	// Upload our prev -> next shadow data mapping (FNextVirtualShadowMapData) to the GPU
@@ -2046,12 +2070,13 @@ void FVirtualShadowMapArray::UpdatePhysicalPageAddresses(FRDGBuilder& GraphBuild
 
 	FUpdatePhysicalPageAddresses::FPermutationDomain PermutationVector;
 
-	TRefCountPtr<IPooledRenderTarget> PrevPageRequestFlags = CacheManager->GetPrevBuffers().PageRequestFlags;
-	if (PrevPageRequestFlags != nullptr)
+	const bool bCacheDataValid = CacheManager->IsCacheDataAvailable();
+	if (CacheManager->IsCacheDataAvailable())
 	{
+		TRefCountPtr<IPooledRenderTarget> PrevPageRequestFlags = CacheManager->GetPrevBuffers().PageRequestFlags;
 		PassParameters->PrevPageRequestFlags = GraphBuilder.RegisterExternalTexture(PrevPageRequestFlags, TEXT("Shadow.Virtual.PrevPageRequestFlags"));
 	}
-	PermutationVector.Set<FUpdatePhysicalPageAddresses::FHasCacheDataDim>(PrevPageRequestFlags != nullptr);
+	PermutationVector.Set<FUpdatePhysicalPageAddresses::FHasCacheDataDim>(bCacheDataValid);
 
 	SetStatsArgsAndPermutation<FUpdatePhysicalPageAddresses>(ShouldGenerateStats(), StatsBufferUAV, PassParameters, PermutationVector);
 
@@ -2107,7 +2132,6 @@ void FVirtualShadowMapArray::BeginMarkPages(
 
 	if (GetNumShadowMaps() == 0)
 	{
-		// Nothing to do
 		return;
 	}
 
@@ -3893,7 +3917,7 @@ void FVirtualShadowMapArray::RenderVirtualShadowMapsNanite(
 		}
 
 		// Prev HZB requires previous page tables and similar
-		bool bPrevHZBValid = HZBPhysicalArray != nullptr && CacheManager->GetPrevBuffers().PageTable != nullptr;
+		bool bPrevHZBValid = HZBPhysicalArrayRDG != nullptr && CacheManager->GetPrevBuffers().PageTable != nullptr;
 
 		Nanite::FConfiguration CullingConfig = { 0 };
 		CullingConfig.bIsShadowPass = true;
@@ -3912,7 +3936,7 @@ void FVirtualShadowMapArray::RenderVirtualShadowMapsNanite(
 			RasterContext,
 			CullingConfig,
 			VirtualShadowViewRect,
-			bPrevHZBValid ? GraphBuilder.RegisterExternalTexture(HZBPhysicalArray) : nullptr,
+			bPrevHZBValid ? HZBPhysicalArrayRDG : nullptr,
 			this
 		);
 
diff --git a/Engine/Source/Runtime/Renderer/Private/VirtualShadowMaps/VirtualShadowMapArray.h b/Engine/Source/Runtime/Renderer/Private/VirtualShadowMaps/VirtualShadowMapArray.h
--- a/Engine/Source/Runtime/Renderer/Private/VirtualShadowMaps/VirtualShadowMapArray.h
+++ b/Engine/Source/Runtime/Renderer/Private/VirtualShadowMaps/VirtualShadowMapArray.h
@@ -363,7 +363,7 @@ public:
 
 	static void SetShaderDefines(FShaderCompilerEnvironment& OutEnvironment);
 
-	void UpdatePhysicalPageAddresses(FRDGBuilder& GraphBuilder);
+	void UpdatePhysicalPageAddresses(FRDGBuilder& GraphBuilder, const FSceneViewFamily& ViewFamily);
 
 	void BeginMarkPages(
 		FRDGBuilder& GraphBuilder,
@@ -383,11 +383,6 @@ public:
 		return PhysicalPagePoolRDG != nullptr && PageTableRDG != nullptr;
 	}
 
-	bool ShouldCacheStaticSeparately() const
-	{
-		return UniformParameters.StaticCachedArrayIndex > 0;
-	}
-
 	bool HasSeparateDynamicHZB() const
 	{
 		return UniformParameters.StaticHZBArrayIndex > 0;
@@ -490,7 +485,6 @@ public:
 	// NOTE: The underlying textures are owned by FVirtualShadowMapCacheManager.
 	// We just import and maintain a copy of the RDG reference for this frame here.
 	FRDGTextureRef PhysicalPagePoolRDG = nullptr;
-	TRefCountPtr<IPooledRenderTarget> HZBPhysicalArray = nullptr;
 	FRDGTextureRef HZBPhysicalArrayRDG = nullptr;
 	FRDGBufferRef PhysicalPageMetaDataRDG = nullptr;
 
diff --git a/Engine/Source/Runtime/Renderer/Private/VirtualShadowMaps/VirtualShadowMapCacheManager.cpp b/Engine/Source/Runtime/Renderer/Private/VirtualShadowMaps/VirtualShadowMapCacheManager.cpp
--- a/Engine/Source/Runtime/Renderer/Private/VirtualShadowMaps/VirtualShadowMapCacheManager.cpp
+++ b/Engine/Source/Runtime/Renderer/Private/VirtualShadowMaps/VirtualShadowMapCacheManager.cpp
@@ -1158,7 +1158,7 @@ void FVirtualShadowMapArrayCacheManager::FreePhysicalPool(FRDGBuilder& GraphBuil
 	}
 }
 
-TRefCountPtr<IPooledRenderTarget> FVirtualShadowMapArrayCacheManager::SetHZBPhysicalPoolSize(FRDGBuilder& GraphBuilder, FIntPoint RequestedHZBSize, int32 RequestedArraySize, const EPixelFormat Format)
+void FVirtualShadowMapArrayCacheManager::SetHZBPhysicalPoolSize(FRDGBuilder& GraphBuilder, FIntPoint RequestedHZBSize, int32 RequestedArraySize, const EPixelFormat Format)
 {
 	if (!HZBPhysicalPagePoolArray 
 		|| HZBPhysicalPagePoolArray->GetDesc().Extent != RequestedHZBSize 
@@ -1181,8 +1181,6 @@ TRefCountPtr<IPooledRenderTarget> FVirtualShadowMapArrayCacheManager::SetHZBPhys
 
 		Invalidate(GraphBuilder);
 	}
-
-	return HZBPhysicalPagePoolArray;
 }
 
 void FVirtualShadowMapArrayCacheManager::FreeHZBPhysicalPool(FRDGBuilder& GraphBuilder)
@@ -1196,20 +1194,8 @@ void FVirtualShadowMapArrayCacheManager::FreeHZBPhysicalPool(FRDGBuilder& GraphB
 
 void FVirtualShadowMapArrayCacheManager::Invalidate(FRDGBuilder& GraphBuilder)
 {
-	// Clear the cache
-	CacheEntries.Reset();
-
 	PrevBuffers = FVirtualShadowMapArrayFrameData();
-
 	//UE_LOG(LogRenderer, Display, TEXT("Virtual shadow map cache invalidated."));
-
-	// Clear the physical page metadata (on all GPUs)
-	if (PhysicalPageMetaData)
-	{
-		RDG_GPU_MASK_SCOPE(GraphBuilder, FRHIGPUMask::All());
-		FRDGBufferRef PhysicalPageMetaDataRDG = GraphBuilder.RegisterExternalBuffer(PhysicalPageMetaData);
-		AddClearUAVPass(GraphBuilder, GraphBuilder.CreateUAV(PhysicalPageMetaDataRDG), 0);
-	}
 }
 
 bool FVirtualShadowMapArrayCacheManager::IsCacheEnabled()
diff --git a/Engine/Source/Runtime/Renderer/Private/VirtualShadowMaps/VirtualShadowMapCacheManager.h b/Engine/Source/Runtime/Renderer/Private/VirtualShadowMaps/VirtualShadowMapCacheManager.h
--- a/Engine/Source/Runtime/Renderer/Private/VirtualShadowMaps/VirtualShadowMapCacheManager.h
+++ b/Engine/Source/Runtime/Renderer/Private/VirtualShadowMaps/VirtualShadowMapCacheManager.h
@@ -296,7 +296,8 @@ public:
 	TRefCountPtr<FRDGPooledBuffer> GetPhysicalPageMetaData() const { return PhysicalPageMetaData; }
 
 	// Called by VirtualShadowMapArray to potentially resize the HZB physical pool
-	TRefCountPtr<IPooledRenderTarget> SetHZBPhysicalPoolSize(FRDGBuilder& GraphBuilder, FIntPoint RequestedSize, int32 RequestedArraySize, const EPixelFormat Format);
+	void SetHZBPhysicalPoolSize(FRDGBuilder& GraphBuilder, FIntPoint RequestedSize, int32 RequestedArraySize, const EPixelFormat Format);
+	TRefCountPtr<IPooledRenderTarget> GetHZBPhysicalPool() const { return HZBPhysicalPagePoolArray; }
 	void FreeHZBPhysicalPool(FRDGBuilder& GraphBuilder);
 
 	// Invalidate the cache for all shadows, causing any pages to be rerendered
