Fixed race condition by tossing reflection capture on GT instead of RT when IsEncodedHDRCubemapTextureRequired is enabled (disabled by default)

[RN] rnx
#Synced-CL 48507348

diff --git a/Engine/Source/Runtime/Engine/Classes/Components/ReflectionCaptureComponent.h b/Engine/Source/Runtime/Engine/Classes/Components/ReflectionCaptureComponent.h
--- a/Engine/Source/Runtime/Engine/Classes/Components/ReflectionCaptureComponent.h
+++ b/Engine/Source/Runtime/Engine/Classes/Components/ReflectionCaptureComponent.h
@@ -9,6 +9,7 @@
 #include "Templates/RefCounting.h"
 #include "Components/SceneComponent.h"
 #include "RenderCommandFence.h"
+#include "RHIShaderPlatform.h"
 #include "ReflectionCaptureComponent.generated.h"
 
 class FReflectionCaptureProxy;
@@ -94,6 +95,8 @@ class UReflectionCaptureComponent : public USceneComponent
 
 	ENGINE_API class FReflectionCaptureMapBuildData* GetMapBuildData() const;
 
+	ENGINE_API bool IsEncodedHDRCubemapTextureRequired(EShaderPlatform ShaderPlatform);
+	
 	virtual void PropagateLightingScenarioChange() override;
 
 	ENGINE_API static int32 GetReflectionCaptureSize();
diff --git a/Engine/Source/Runtime/Engine/Private/Components/ReflectionCaptureComponent.cpp b/Engine/Source/Runtime/Engine/Private/Components/ReflectionCaptureComponent.cpp
--- a/Engine/Source/Runtime/Engine/Private/Components/ReflectionCaptureComponent.cpp
+++ b/Engine/Source/Runtime/Engine/Private/Components/ReflectionCaptureComponent.cpp
@@ -118,7 +118,7 @@ FReflectionCaptureMapBuildData* UReflectionCaptureComponent::GetMapBuildData() c
 	return NULL;
 }
 
-bool IsEncodedHDRCubemapTextureRequired(EShaderPlatform ShaderPlatform)
+bool UReflectionCaptureComponent::IsEncodedHDRCubemapTextureRequired(EShaderPlatform ShaderPlatform)
 {
 	const bool bEncodedHDRCubemapTextureRequired = (GIsEditor || IsMobilePlatform(ShaderPlatform))
 		// mobile forward renderer or translucensy in mobile deferred need encoded reflection texture when clustered reflections disabled
@@ -126,7 +126,6 @@ bool IsEncodedHDRCubemapTextureRequired(EShaderPlatform ShaderPlatform)
 	return bEncodedHDRCubemapTextureRequired;
 }
 
-
 void UReflectionCaptureComponent::PropagateLightingScenarioChange()
 {
 	const FSceneInterface* Scene = GetWorld()->Scene;
diff --git a/Engine/Source/Runtime/Renderer/Private/ReflectionEnvironmentCapture.cpp b/Engine/Source/Runtime/Renderer/Private/ReflectionEnvironmentCapture.cpp
--- a/Engine/Source/Runtime/Renderer/Private/ReflectionEnvironmentCapture.cpp
+++ b/Engine/Source/Runtime/Renderer/Private/ReflectionEnvironmentCapture.cpp
@@ -77,6 +77,12 @@ static FAutoConsoleVariableRef CVarSkylightCaptureLODDistanceScale(
 	TEXT("Negative values will be clamped to 1"),
 	ECVF_Scalability);
 
+static bool GTossReflectionCaptureDataOnGT = false;
+static FAutoConsoleVariableRef CVarTossReflectionCaptureDataOnGT(
+	TEXT("r.TossReflectionCaptureDataOnGT"),
+	GTossReflectionCaptureDataOnGT,
+	TEXT("Wether or not toss reflection capture data on GT if needed\n"));
+
 /** 
  * Mip map used by a Roughness of 0, counting down from the lowest resolution mip (MipCount - 1).  
  * This has been tweaked along with ReflectionCaptureRoughnessMipScale to make good use of the resolution in each mip, especially the highest resolution mips.
@@ -1464,8 +1470,11 @@ void FScene::CaptureOrUploadReflectionCapture(UReflectionCaptureComponent* Captu
 
 			FScene* Scene = this;
 
+			bool bTossCaptureDataOnGT = CaptureComponent->IsEncodedHDRCubemapTextureRequired(GetShaderPlatform()) && GTossReflectionCaptureDataOnGT;
+			TWeakObjectPtr WeakCapture = CaptureComponent;
+
 			ENQUEUE_RENDER_COMMAND(UploadCaptureCommand)
-				([Scene, CaptureData, CaptureComponent](FRHICommandListImmediate& RHICmdList)
+				([Scene, CaptureData, CaptureComponent, WeakCapture, bTossCaptureDataOnGT](FRHICommandListImmediate& RHICmdList)
 			{
 				// After the final upload we cannot upload again because we tossed the source MapBuildData,
 				// After uploading it into the scene's texture array, to guaratee there's only one copy in memory.
@@ -1474,7 +1483,25 @@ void FScene::CaptureOrUploadReflectionCapture(UReflectionCaptureComponent* Captu
 				{
 					UploadReflectionCapture_RenderingThread(RHICmdList, Scene, CaptureData, CaptureComponent);
 
-					CaptureData->OnDataUploadedToGPUFinal();
+					if (bTossCaptureDataOnGT)
+					{
+						AsyncTask(ENamedThreads::GameThread, [CaptureData, WeakCapture]()
+								  {
+									  if (WeakCapture.IsValid())
+									  {
+										  FReflectionCaptureData* CaptureDataLocal = WeakCapture.Get()->GetMapBuildData();
+										  check(CaptureDataLocal == CaptureData);
+										  if (CaptureDataLocal != nullptr)
+										  {
+										      CaptureDataLocal->OnDataUploadedToGPUFinal();
+										  }
+									  }
+								  });
+					}
+					else
+					{
+						CaptureData->OnDataUploadedToGPUFinal();
+					}
 				}
 				else
 				{
