Added new cvar to allow for the material or niagara use of the GDF to be sufficient to prepare the GDF (rather than being reliant on the DistanceFieldAO feature).  Defaults to false to preserve existing behavior.
#tests Created new project and added a niagara system and switched it's collision behavior
#Synced-CL Unknown

diff --git a/Engine/Source/Runtime/Renderer/Private/DistanceFieldAmbientOcclusion.cpp b/Engine/Source/Runtime/Renderer/Private/DistanceFieldAmbientOcclusion.cpp
--- a/Engine/Source/Runtime/Renderer/Private/DistanceFieldAmbientOcclusion.cpp
+++ b/Engine/Source/Runtime/Renderer/Private/DistanceFieldAmbientOcclusion.cpp
@@ -139,6 +139,16 @@ FAutoConsoleVariableRef CVarAOGlobalDistanceFieldDetailedCheck(
 	ECVF_RenderThreadSafe
 );
 
+// if the detailed necessity check (bShouldPrepareForMaterialsOrNiagara) succeeds, this cvar can be used
+// to control whether the GDF is prepared, regardless of whether DFAO is being used
+int32 GGlobalDistanceFieldDetailedNecessitySufficient = 0;
+FAutoConsoleVariableRef CVarGlobalDistanceFieldDetailedSufficient(
+	TEXT("r.GlobalDistanceField.DetailedNecessitySufficient"),
+	GGlobalDistanceFieldDetailedNecessitySufficient,
+	TEXT("Whether detailed necessity checks in FSceneRenderer::ShouldPrepareGlobalDistanceField() is sufficient to enabled GDF."),
+	ECVF_RenderThreadSafe
+);
+
 bool UseDistanceFieldAO()
 {
 	return GDistanceFieldAO && GDistanceFieldAOQuality >= 1;
@@ -772,6 +782,11 @@ bool ShouldPrepareGlobalDistanceField(const FSceneRenderer* Renderer)
 		return bShouldPrepareForMaterialsOrNiagara;
 	}
 
+	if (GGlobalDistanceFieldDetailedNecessitySufficient && bShouldPrepareForMaterialsOrNiagara)
+	{
+		return true;
+	}
+
 	if (SupportsDistanceFieldAO(Renderer->FeatureLevel, Renderer->ShaderPlatform) && (bShouldPrepareForMaterialsOrNiagara || ShouldPrepareForDistanceFieldAO(Scene, ViewFamily, Renderer->AnyViewHasGIMethodSupportingDFAO())))
 	{
 		return true;
