[HWRT] Remove deprecated path using SceneInfo->StaticMeshes to cache ray tracing mesh commands.

#Synced-CL 48851180

diff --git a/Engine/Source/Runtime/Renderer/Private/PrimitiveSceneInfo.cpp b/Engine/Source/Runtime/Renderer/Private/PrimitiveSceneInfo.cpp
--- a/Engine/Source/Runtime/Renderer/Private/PrimitiveSceneInfo.cpp
+++ b/Engine/Source/Runtime/Renderer/Private/PrimitiveSceneInfo.cpp
@@ -1168,7 +1168,7 @@ void CacheRayTracingMeshBatch(
 template<class T>
 int32 CacheRayTracingMeshCommands(
 	FPrimitiveSceneInfo* SceneInfo,
-	const FRayTracingInstance& RayTracingInstance,
+	TConstArrayView<FMeshBatch> MeshBatches,
 	ERayTracingPrimitiveFlags Flags,
 	T& Commands,
 	FCachedRayTracingMeshCommandContext<T>& CommandContext,
@@ -1187,25 +1187,9 @@ int32 CacheRayTracingMeshCommands(
 	{
 		// Cache ray tracing mesh commands in FPrimitiveSceneInfo
 
-		const int32 LODCount = EnumHasAnyFlags(Flags, ERayTracingPrimitiveFlags::CacheInstances) ? 1 : SceneInfo->GetStaticRayTracingGeometryNum();
+		ensureMsgf(MeshBatches.Num() > 0, TEXT("Caching ray tracing mesh commands but instance doesn't have any."));
 
-		if (RayTracingInstance.Materials.Num() > 0)
-		{
-			for (const FMeshBatch& Mesh : RayTracingInstance.Materials)
-			{
-				check(Mesh.LODIndex < LODCount);
-			}
-		}
-		else if(!SceneInfo->StaticMeshes.IsEmpty())
-		{
-			// Keeping old behavior for one UE release. 
-			ensureMsgf(false, TEXT("GetCachedRayTracingInstance(...) must return materials when using static instances code path."));
-
-			for (const FStaticMeshBatch& Mesh : SceneInfo->StaticMeshes)
-			{
-				check(Mesh.LODIndex < LODCount);
-			}
-		}
+		const int32 LODCount = EnumHasAnyFlags(Flags, ERayTracingPrimitiveFlags::CacheInstances) ? 1 : SceneInfo->GetStaticRayTracingGeometryNum();
 
 		check(SceneInfo->GetRayTracingLODDataNum() == 0);
 				
@@ -1219,20 +1203,10 @@ int32 CacheRayTracingMeshCommands(
 		DeferredData.SceneInfo = SceneInfo;
 		DeferredData.CommandIndicesPerLOD.SetNum(LODCount);
 		
-		if (RayTracingInstance.Materials.Num() > 0)
+		for (const FMeshBatch& Mesh : MeshBatches)
 		{
-			for (const FMeshBatch& Mesh : RayTracingInstance.Materials)
-			{
-				CacheRayTracingMeshBatch(Mesh, SceneInfo, RayTracingLODData, Commands, CommandContext, RayTracingMeshProcessor, DeferredData);
-			}
-		}
-		else
-		{
-			// Keeping old behavior for one UE release. 
-			for (const FStaticMeshBatch& Mesh : SceneInfo->StaticMeshes)
-			{
-				CacheRayTracingMeshBatch(Mesh, SceneInfo, RayTracingLODData, Commands, CommandContext, RayTracingMeshProcessor, DeferredData);
-			}
+			check(Mesh.LODIndex >= 0 && Mesh.LODIndex < LODCount);
+			CacheRayTracingMeshBatch(Mesh, SceneInfo, RayTracingLODData, Commands, CommandContext, RayTracingMeshProcessor, DeferredData);
 		}
 
 		// Store in the Scene info
@@ -1284,7 +1258,7 @@ void FPrimitiveSceneInfo::CacheRayTracingPrimitives(FScene* Scene, const TArrayV
 							SceneRayTracingGroupId = Scene->PrimitiveRayTracingGroups.FindId(RayTracingGroupId);
 						}
 
-						check(Scene->PrimitiveRayTracingGroupIds[SceneInfo->GetIndex()] == SceneRayTracingGroupId);
+						check(Scene->PrimitiveRayTracingGroupIds[PrimitiveIndex] == SceneRayTracingGroupId);
 
 						CheckPrimitiveRayTracingData(PrimitiveRayTracingData, SceneProxy);
 					}
@@ -1296,7 +1270,7 @@ void FPrimitiveSceneInfo::CacheRayTracingPrimitives(FScene* Scene, const TArrayV
 					// Write flags
 					Flags = SceneProxy->GetCachedRayTracingInstance(RayTracingInstance);
 
-					const int32 DeferredDataIndex = CacheRayTracingMeshCommands(SceneInfo, RayTracingInstance, Flags, Context.Commands, Context.CommandContext, Context.RayTracingMeshProcessor, Context.DeferredDatas);
+					const int32 DeferredDataIndex = CacheRayTracingMeshCommands(SceneInfo, RayTracingInstance.Materials, Flags, Context.Commands, Context.CommandContext, Context.RayTracingMeshProcessor, Context.DeferredDatas);
 
 					if (EnumHasAnyFlags(Flags, ERayTracingPrimitiveFlags::CacheInstances))
 					{
