Add missing View.UpperBoundViewRect initialization in case of r.Test.ViewRectOffset 1 and CustomRenderPasses

#Synced-CL 48819522

diff --git a/Engine/Source/Runtime/Renderer/Private/SceneRendering.cpp b/Engine/Source/Runtime/Renderer/Private/SceneRendering.cpp
--- a/Engine/Source/Runtime/Renderer/Private/SceneRendering.cpp
+++ b/Engine/Source/Runtime/Renderer/Private/SceneRendering.cpp
@@ -3473,7 +3473,10 @@ void FSceneRenderer::PrepareViewRectsForRendering()
 				case 5: NewViewRectMin = FIntPoint(2 * Offset.X, 2 * Offset.Y); break;
 				}
 
-				View.ViewRect += QuantizeViewRectMin(NewViewRectMin) - View.ViewRect.Min;
+				const FIntPoint NewOffset = QuantizeViewRectMin(NewViewRectMin) - View.ViewRect.Min;
+
+				View.ViewRect += NewOffset;
+				View.UpperBoundViewRect += NewOffset;
 
 #if DO_CHECK || USING_CODE_ANALYSIS
 				check(View.VerifyMembersChecks());
@@ -3509,11 +3512,11 @@ void FSceneRenderer::PrepareViewRectsForRendering()
 						View.TemporalSourceView = &Views[0];
 					}
 				}
-				View.UnconstrainedViewRect = View.ViewRect;
+				View.UpperBoundViewRect = View.UnconstrainedViewRect = View.ViewRect;
 			}
 			else
 			{
-				View.ViewRect = View.UnscaledViewRect;
+				View.UpperBoundViewRect = View.ViewRect = View.UnscaledViewRect;
 			}
 		}	
 	}
