[HWRT] Reduce FRayTracingShaderBindingTable::AllocateDynamicRange(...) overhead.
- defer MarkDirty(...) calls so they can be batched into a single SetRange(...) call.
- refactored DynamicAllocationPool to be a TIndirectArray so we can use simple counter instead of active/free allocation pools.

- RayTracingScene_AddDynamicInstances cost on test level with 300 GCs each consisting of ~70 parts: 8.52ms -> 7.36ms.

#Synced-CL 48851180

diff --git a/Engine/Source/Runtime/Renderer/Private/RayTracing/RayTracingShaderBindingTable.cpp b/Engine/Source/Runtime/Renderer/Private/RayTracing/RayTracingShaderBindingTable.cpp
--- a/Engine/Source/Runtime/Renderer/Private/RayTracing/RayTracingShaderBindingTable.cpp
+++ b/Engine/Source/Runtime/Renderer/Private/RayTracing/RayTracingShaderBindingTable.cpp
@@ -113,11 +113,6 @@ FRayTracingShaderBindingTable::FRayTracingShaderBindingTable() : NumShaderSlotsP
 FRayTracingShaderBindingTable::~FRayTracingShaderBindingTable()
 {
 	ResetDynamicAllocationData();
-	for (FRayTracingSBTAllocation* SBTAllocation : FreeDynamicAllocationPool)
-	{
-		delete SBTAllocation;
-	}
-	FreeDynamicAllocationPool.Empty();
 	
 	// Assume empty?
 	check(TrackedAllocationMap.IsEmpty());
@@ -132,8 +127,7 @@ FRayTracingShaderBindingTable::~FRayTracingShaderBindingTable()
 void FRayTracingShaderBindingTable::ResetDynamicAllocationData()
 {				
 	// Release all dynamic allocation back to the pool
-	FreeDynamicAllocationPool.Append(ActiveDynamicAllocations);
-	ActiveDynamicAllocations.Empty(ActiveDynamicAllocations.Num());
+	NumUsedDynamicAllocations = 0;
 	NumDynamicGeometrySegments = 0;
 	
 	// Static allocations are not allowed anymore because dynamic allocations are stored right after all static allocations
@@ -143,6 +137,7 @@ void FRayTracingShaderBindingTable::ResetDynamicAllocationData()
 	uint32 AllocatedStaticSegmentSize = GetMaxAllocatedStaticSegmentCount();
 	StartDynamicRangeOffset = AllocatedStaticSegmentSize * NumShaderSlotsPerGeometrySegment;
 	CurrentDynamicRangeOffset = StartDynamicRangeOffset;
+	StartDynamicRangeOffsetToMarkDirty = StartDynamicRangeOffset;
 }
 
 void FRayTracingShaderBindingTable::ResetMissAndCallableShaders()
@@ -166,6 +161,23 @@ FRayTracingShaderBindingDataOneFrameArray FRayTracingShaderBindingTable::GetDirt
 	
 	ensure(bStaticAllocationsLocked);
 
+	// Flush dynamic records that are pending being marked dirty
+	{
+		check(StartDynamicRangeOffsetToMarkDirty <= CurrentDynamicRangeOffset);
+	
+		// Grow the bit array if necessary
+		const int32 MaxRecordIndex = CurrentDynamicRangeOffset;
+		if (DirtyPersistentRecords.Num() < MaxRecordIndex)
+		{
+			DirtyPersistentRecords.SetNumUninitialized(FMath::RoundUpToPowerOfTwo(MaxRecordIndex));
+		}
+
+		// Mark bits in dynamic range as dirty
+		DirtyPersistentRecords.SetRange(StartDynamicRangeOffsetToMarkDirty, CurrentDynamicRangeOffset - StartDynamicRangeOffsetToMarkDirty, true);
+
+		StartDynamicRangeOffsetToMarkDirty = CurrentDynamicRangeOffset;
+	}
+
 	FRayTracingShaderBindingDataOneFrameArray DirtyBindings;
 
 	// Persistent SBTs disabled (none of the cached SBTs have any persistent bindings so safe to just copy visible bindings and update type inside binding will be ignored)
@@ -190,10 +202,10 @@ FRayTracingShaderBindingDataOneFrameArray FRayTracingShaderBindingTable::GetDirt
 			}
 			else
 			{
-				check(VisibleBinding.BindingType!= ERayTracingLocalShaderBindingType::Persistent || SetPersistentRecords[VisibleBinding.SBTRecordIndex]);
+				check(VisibleBinding.BindingType != ERayTracingLocalShaderBindingType::Persistent || SetPersistentRecords[VisibleBinding.SBTRecordIndex]);
 
 				FBitReference DirtyRecordIndex = DirtyPersistentRecords[VisibleBinding.SBTRecordIndex];
-				if (DirtyRecordIndex || VisibleBinding.BindingType == ERayTracingLocalShaderBindingType::Transient)
+				if (VisibleBinding.BindingType == ERayTracingLocalShaderBindingType::Transient || DirtyRecordIndex)
 				{
 					DirtyBindings.Add(VisibleBinding);
 					DirtyRecordIndex = false;
@@ -532,11 +544,17 @@ FRayTracingSBTAllocation* FRayTracingShaderBindingTable::AllocateDynamicRange(ER
 
 	// Don't need lock right now because all dynamic allocation are allocated linearly on the same thread
 	// So the FreeDynamicAllocationPool can't be shared with the static allocations right now because those would require a lock
-	FRayTracingSBTAllocation* Allocation = FreeDynamicAllocationPool.IsEmpty() ? nullptr : FreeDynamicAllocationPool.Pop(EAllowShrinking::No);
-	if (Allocation == nullptr)
+	FRayTracingSBTAllocation* Allocation = nullptr;
+	if (NumUsedDynamicAllocations < (uint32)DynamicAllocationPool.Num())
+	{
+		Allocation = &DynamicAllocationPool[NumUsedDynamicAllocations];
+	}
+	else
 	{
 		Allocation = new FRayTracingSBTAllocation();
+		DynamicAllocationPool.Add(Allocation);
 	}
+	NumUsedDynamicAllocations++;
 
 	uint32 LayersCount = FMath::CountBits((uint32)AllocatedLayers);
 	uint32 BaseIndex = CurrentDynamicRangeOffset;
@@ -545,20 +563,19 @@ FRayTracingSBTAllocation* FRayTracingShaderBindingTable::AllocateDynamicRange(ER
 	CurrentDynamicRangeOffset += RecordCount;
 	Allocation->InitDynamic(AllocatedLayers, BaseIndex, RecordsPerLayer, RecordCount);
 
-	MarkDirty(Allocation);
+	// Don't need to call MarkDirty(Allocation) here since it will be batched with other allocations using StartDynamicRangeOffsetToMarkDirty
 
 	NumDynamicGeometrySegments += SegmentCount * LayersCount;
 
 	MaxNumDynamicGeometrySegments = FMath::Max(MaxNumDynamicGeometrySegments, NumDynamicGeometrySegments);
 
-	ActiveDynamicAllocations.Add(Allocation);
-
 	return Allocation;
 }
 
 bool FRayTracingShaderBindingTable::IsDirty(uint32 RecordIndex) const
 {
-	return DirtyPersistentRecords[RecordIndex];
+	const bool bIsDynamicPendingMarkDirty = RecordIndex >= StartDynamicRangeOffsetToMarkDirty;
+	return bIsDynamicPendingMarkDirty || DirtyPersistentRecords[RecordIndex];
 }
 
 bool FRayTracingShaderBindingTable::IsPersistent() const
diff --git a/Engine/Source/Runtime/Renderer/Private/RayTracing/RayTracingShaderBindingTable.h b/Engine/Source/Runtime/Renderer/Private/RayTracing/RayTracingShaderBindingTable.h
--- a/Engine/Source/Runtime/Renderer/Private/RayTracing/RayTracingShaderBindingTable.h
+++ b/Engine/Source/Runtime/Renderer/Private/RayTracing/RayTracingShaderBindingTable.h
@@ -291,8 +291,8 @@ private:
 	FSpanAllocator StaticRangeAllocator;								 //< Range allocator to find free static record ranges
 	TMap<FAllocationKey, FRefCountedAllocation> TrackedAllocationMap;	 //< All static allocation with refcount tracking
 	
-	TArray<FRayTracingSBTAllocation*> ActiveDynamicAllocations;			 //< All current active dynamic allocations
-	TArray<FRayTracingSBTAllocation*> FreeDynamicAllocationPool;		 //< Free dynamic allocation pool (for faster allocations)
+	TIndirectArray<FRayTracingSBTAllocation> DynamicAllocationPool;      //< Dynamic allocation pool (for faster allocations)
+	uint32 NumUsedDynamicAllocations = 0;                                //< Number of dynamic allocations currently used
 
 	uint32 TotalStaticAllocationCount = 0;								 //< Total amount of static allocations (without deduplications)
 	uint32 AllocatedStaticSegmentCount = 0;								 //< Total amount of allocated static segments (with deduplication)
@@ -301,6 +301,7 @@ private:
 	uint32 NumDynamicGeometrySegments = 0;								 //< Current number of allocated dynamic segments
 	uint32 StartDynamicRangeOffset = 0;								 	 //< Start SBT record offset for the first dynamic allocation
 	uint32 CurrentDynamicRangeOffset = 0;								 //< Current working SBT record offset for the next dynamic allocation
+	uint32 StartDynamicRangeOffsetToMarkDirty = 0;						 //< Start SBT record offset that needs to be marked dirty
 };
 
 #endif // RHI_RAYTRACING
