Fix TSAN error (false race detection due to bit-filed overlap) & remove deprecated constructor.

#rnx
#Synced-CL 43344078

diff --git a/Engine/Source/Runtime/Renderer/Public/InstanceCulling/InstanceCullingContext.h b/Engine/Source/Runtime/Renderer/Public/InstanceCulling/InstanceCullingContext.h
--- a/Engine/Source/Runtime/Renderer/Public/InstanceCulling/InstanceCullingContext.h
+++ b/Engine/Source/Runtime/Renderer/Public/InstanceCulling/InstanceCullingContext.h
@@ -95,19 +95,6 @@ public:
 	
 	FInstanceCullingContext() {}
 
-	UE_DEPRECATED(5.4, "Use constructor which provides pass name as first argument")
-	FInstanceCullingContext(		
-		EShaderPlatform ShaderPlatform,
-		FInstanceCullingManager* InInstanceCullingManager,
-		TArrayView<const int32> InViewIds,
-		const TRefCountPtr<IPooledRenderTarget>& InPrevHZB,
-		EInstanceCullingMode InInstanceCullingMode = EInstanceCullingMode::Normal,
-		EInstanceCullingFlags InFlags = EInstanceCullingFlags::None,
-		EBatchProcessingMode InSingleInstanceProcessingMode = EBatchProcessingMode::UnCulled) :
-		FInstanceCullingContext(TEXT("Unknown"), ShaderPlatform, InInstanceCullingManager, InViewIds, InPrevHZB, InInstanceCullingMode, InFlags, InSingleInstanceProcessingMode)
-	{
-	}
-
 	/**
 	 * Create an instance culling context to process draw commands that can be culled using GPU-Scene.
 	 * @param InPrevHZB if non-null enables HZB-occlusion culling for the context (if r.InstanceCulling.OcclusionCull is enabled),
@@ -247,9 +234,8 @@ public:
 	EShaderPlatform ShaderPlatform = SP_NumPlatforms;
 	TArray<int32, TInlineAllocator<6, SceneRenderingAllocator>> ViewIds;
 	TRefCountPtr<IPooledRenderTarget> PrevHZB = nullptr;
-	bool bIsEnabled:1 = false;
-	// This is set during SetupDrawCommands to indicate if there are draw commands using indirect args, as this requires setting up the relevant buffers (even if the instance count is zero).
-	bool bHasAnyIndirectDraw:1 = false;
+	// Note: bIsEnabled is immutable after construction and is accessed from multiple threads. To avoid false race detection, it should not be made into a bit-field.
+	bool bIsEnabled = false;
 	EInstanceCullingMode InstanceCullingMode = EInstanceCullingMode::Normal;
 	EInstanceCullingFlags Flags = EInstanceCullingFlags::None;
 
@@ -259,6 +245,9 @@ public:
 	int32 DynamicInstanceIdNum = -1;
 
 	SyncPrerequisitesFuncType SyncPrerequisitesFunc;
+protected:
+	// This is set during SetupDrawCommands to indicate if there are draw commands using indirect args, as this requires setting up the relevant buffers (even if the instance count is zero).
+	bool bHasAnyIndirectDraw:1 = false;
 public:
 
 	enum class EAsyncProcessingMode
