AnimBank refactor to add support for AnimBank animation tracks. These animation tracks generate a unique animation pose. Each track has an index that is assigned to an ISKM instance when a UAnimBankDataInstance is applied as the transform provider. UAnimBankData transform provider behavior is unchanged.
- Added a new API to create a UAnimBankDataInstance from a UAnimBankData.
- UAnimBankDataInstance contains an allocator API to allocate / update / and deallocate tracks.
- Added a FAnimBankTrackAutoPlayData and FAnimBankTrackManualData which allow control to auto-play and animation or manually control playback of one.
- Simplified the AnimBankTransformProvider to issue a single compute dispatch that writes directly into the final transform buffer.
- Modified the transform provider API to allow game thread data to be fed to a render proxy on creation in a thread-safe manner.

#Synced-CL 48814992

diff --git a/Engine/Plugins/Experimental/DynamicWind/Source/DynamicWind/Private/DynamicWindData.cpp b/Engine/Plugins/Experimental/DynamicWind/Source/DynamicWind/Private/DynamicWindData.cpp
--- a/Engine/Plugins/Experimental/DynamicWind/Source/DynamicWind/Private/DynamicWindData.cpp
+++ b/Engine/Plugins/Experimental/DynamicWind/Source/DynamicWind/Private/DynamicWindData.cpp
@@ -26,7 +26,7 @@ bool UDynamicWindData::UsesSkeletonBatching() const
 	return true;
 }
 
-const uint32 UDynamicWindData::GetUniqueAnimationCount() const
+uint32 UDynamicWindData::GetUniqueAnimationCount() const
 {
 	return DYNAMIC_WIND_DIRECTIONALITY_SLICES;
 }
@@ -54,25 +54,15 @@ uint32 UDynamicWindData::GetSkinningDataOffset(int32 InstanceIndex, const FTrans
 	return DirectionalityIndex * 2u;
 }
 
-FTransformProviderRenderProxy* UDynamicWindData::CreateRenderThreadResources(FSkinningSceneExtensionProxy* SceneProxy, FSceneInterface& Scene, FRHICommandListBase& RHICmdList)
+FTransformProviderRenderProxy* UDynamicWindData::CreateRenderProxy(FInstancedSkinningSceneExtensionProxy* SceneExtensionProxy) const
 {
-	FDynamicWindDataRenderProxy* ProviderProxy = new FDynamicWindDataRenderProxy(this, SceneProxy, Scene);
-	ProviderProxy->CreateRenderThreadResources(RHICmdList);
-	return ProviderProxy;
+	return new FDynamicWindDataProxy(this, SceneExtensionProxy);
 }
 
-void UDynamicWindData::DestroyRenderThreadResources(FTransformProviderRenderProxy* ProviderProxy)
+FDynamicWindDataProxy::FDynamicWindDataProxy(const UDynamicWindData* BankData, FInstancedSkinningSceneExtensionProxy* InSceneExtensionProxy)
+	: SceneExtensionProxy(InSceneExtensionProxy)
 {
-	ProviderProxy->DestroyRenderThreadResources();
-	delete ProviderProxy;
-}
-
-FDynamicWindDataRenderProxy::FDynamicWindDataRenderProxy(UDynamicWindData* BankData, FSkinningSceneExtensionProxy* InSceneProxy, FSceneInterface& InScene)
-: SceneProxy(InSceneProxy)
-, Scene(InScene)
-{
-	UniqueAnimationCount = SceneProxy->GetUniqueAnimationCount();
-	const USkinnedAsset* SkinnedAsset = SceneProxy->GetSkinnedAsset();
+	const USkinnedAsset* SkinnedAsset = SceneExtensionProxy->GetSkinnedAsset();
 	check(SkinnedAsset != nullptr);
 	
 	const USkeletalMesh* SkeletalMesh = CastChecked<USkeletalMesh>(SkinnedAsset);
@@ -81,20 +71,14 @@ FDynamicWindDataRenderProxy::FDynamicWindDataRenderProxy(UDynamicWindData* BankD
 	// GW-TODO: Need a validation method on the base data type to check per-proxy if ref pose fallback is needed.
 }
 
-FDynamicWindDataRenderProxy::~FDynamicWindDataRenderProxy()
+FDynamicWindDataProxy::~FDynamicWindDataProxy()
 {
 }
 
-void FDynamicWindDataRenderProxy::CreateRenderThreadResources(FRHICommandListBase& RHICmdList)
+void FDynamicWindDataProxy::CreateRenderThreadResources(FSceneInterface& InScene, FRHICommandListBase& RHICmdList)
 {
 }
 
-void FDynamicWindDataRenderProxy::DestroyRenderThreadResources()
+void FDynamicWindDataProxy::DestroyRenderThreadResources()
 {
-}
-
-const TConstArrayView<uint64> FDynamicWindDataRenderProxy::GetProviderData(bool& bOutValid) const
-{
-	bOutValid = true;
-	return {};
-}
+}
\ No newline at end of file
diff --git a/Engine/Plugins/Experimental/DynamicWind/Source/DynamicWind/Public/DynamicWindData.h b/Engine/Plugins/Experimental/DynamicWind/Source/DynamicWind/Public/DynamicWindData.h
--- a/Engine/Plugins/Experimental/DynamicWind/Source/DynamicWind/Public/DynamicWindData.h
+++ b/Engine/Plugins/Experimental/DynamicWind/Source/DynamicWind/Public/DynamicWindData.h
@@ -14,29 +14,25 @@ class UDynamicWindData : public UTransformProviderData
 public:
 	virtual bool IsEnabled() const override;
 	virtual const FGuid& GetTransformProviderID() const override;
+	virtual uint32 GetUniqueAnimationCount() const override;
 	virtual bool UsesSkeletonBatching() const override;
-	virtual const uint32 GetUniqueAnimationCount() const override;
 	virtual bool HasAnimationBounds() const override;
 	virtual bool GetAnimationBounds(uint32 AnimationIndex, FRenderBounds& OutBounds) const override;
 	virtual uint32 GetSkinningDataOffset(int32 InstanceIndex, const FTransform& ComponentTransform, const FSkinnedMeshInstanceData& InstanceData) const override;
-	virtual FTransformProviderRenderProxy* CreateRenderThreadResources(FSkinningSceneExtensionProxy* SceneProxy, FSceneInterface& Scene, FRHICommandListBase& RHICmdList) override;
-	virtual void DestroyRenderThreadResources(FTransformProviderRenderProxy* ProviderProxy) override;
+	virtual FTransformProviderRenderProxy* CreateRenderProxy(FInstancedSkinningSceneExtensionProxy* ExtensionProxy) const override;
 };
 
-class FDynamicWindDataRenderProxy : public FTransformProviderRenderProxy
+class FDynamicWindDataProxy : public FTransformProviderRenderProxy
 {
 	friend class UDynamicWindData;
 
 public:
-	DYNAMICWIND_API FDynamicWindDataRenderProxy(UDynamicWindData* BankData, FSkinningSceneExtensionProxy* SceneProxy, FSceneInterface& Scene);
-	DYNAMICWIND_API virtual ~FDynamicWindDataRenderProxy();
-	DYNAMICWIND_API virtual void CreateRenderThreadResources(FRHICommandListBase& RHICmdList) override;
+	DYNAMICWIND_API FDynamicWindDataProxy(const UDynamicWindData* BankData, FInstancedSkinningSceneExtensionProxy* SceneExtensionProxy);
+	DYNAMICWIND_API virtual ~FDynamicWindDataProxy();
+	DYNAMICWIND_API virtual void CreateRenderThreadResources(FSceneInterface& Scene, FRHICommandListBase& RHICmdList) override;
 	DYNAMICWIND_API virtual void DestroyRenderThreadResources() override;
-	DYNAMICWIND_API virtual const TConstArrayView<uint64> GetProviderData(bool& bOutValid) const override;
 
 private:
-	FSkinningSceneExtensionProxy* SceneProxy = nullptr;
+	FInstancedSkinningSceneExtensionProxy* SceneExtensionProxy = nullptr;
 	TObjectPtr<const class UDynamicWindSkeletalData> SkeletalData = nullptr;
-	FSceneInterface& Scene; 
-	uint32 UniqueAnimationCount = 0;
 };
diff --git a/Engine/Shaders/Private/GpuSkinVertexFactory.ush b/Engine/Shaders/Private/GpuSkinVertexFactory.ush
--- a/Engine/Shaders/Private/GpuSkinVertexFactory.ush
+++ b/Engine/Shaders/Private/GpuSkinVertexFactory.ush
@@ -510,7 +510,7 @@ FComputeBoneMatrixInputs GetComputeBoneMatrixInputs(FInstanceSceneData SceneData
 	FComputeBoneMatrixInputs Inputs;
 	Inputs.BoneMatrices = BoneMatrices;
 	Inputs.SectionBoneMap = GPUSkinVF.SectionBoneMap;
-	Inputs.BoneTransformOffset = GetSkinningBoneTransformOffset(GPUSkinVF.TransformBufferOffset, SceneData.SkinningData, GPUSkinVF.CurrentTransformSlot, GPUSkinVF.MaxTransformCount, bPreviousTransform);
+	Inputs.BoneTransformOffset = GetSkinningBoneTransformOffset(GPUSkinVF.TransformBufferOffset, SceneData.SkinningData * (GPUSkinVF.bIsRefPose ? 0 : 1), GPUSkinVF.CurrentTransformSlot, GPUSkinVF.MaxTransformCount, bPreviousTransform);
 	Inputs.TransformStorageMode = GPUSkinVF.TransformStorageMode;
 	Inputs.SectionBoneOffset = GPUSkinVF.SectionBoneOffset;
 	return Inputs;
diff --git a/Engine/Shaders/Private/Skinning/AnimBankEval.usf b/Engine/Shaders/Private/Skinning/AnimBankEval.usf
--- a/Engine/Shaders/Private/Skinning/AnimBankEval.usf
+++ b/Engine/Shaders/Private/Skinning/AnimBankEval.usf
@@ -8,13 +8,13 @@
 #include "/Engine/Private/Skinning/SkinningDataDecode.ush"
 #include "/Engine/Shared/SkinningDefinitions.h"
 
-ByteAddressBuffer HeaderBuffer;
-ByteAddressBuffer BankBuffer;
-FCompressedBoneTransformBuffer SrcTransformBuffer;
-ByteAddressBuffer SrcBoneMapBuffer;
+ByteAddressBuffer SequenceTransformBuffer;
+ByteAddressBuffer BankBlockBuffer;
+ByteAddressBuffer BoneMapBuffer;
 FCompressedBoneTransformRWBuffer TransformBuffer;
 
-float DeltaTime;
+float SampleInterval;
+uint bInterpolating;
 
 struct FKeyIndices
 {
@@ -74,125 +74,79 @@ FKeyIndices GetKeyIndicesFromTime(float Time, int KeyCount, float SequenceLength
 [numthreads(BONES_PER_GROUP, 1, 1)]
 void BankEvaluateCS(uint3 GroupId : SV_GroupID, uint GroupIndex : SV_GroupIndex)
 {
-	const uint BoneBlock		= GroupId.x;
-	const uint LocalBoneIndex	= GroupIndex;
+	const uint BoneBlock           = GroupId.x;
+	const uint LocalTransformIndex = GroupIndex;
 
-	const FBankBlockHeader BoneBlockHeader = HeaderBuffer.Load<FBankBlockHeader>(BoneBlock * uint(sizeof(FBankBlockHeader)));
-	if (LocalBoneIndex < BoneBlockHeader.BlockBoneCount)
+	const FAnimBankBlockHeader BoneBlockHeader = BankBlockBuffer.Load<FAnimBankBlockHeader>(BoneBlock * uint(sizeof(FAnimBankBlockHeader)));
+	if (LocalTransformIndex < BoneBlockHeader.BlockTransformCount)
 	{
-		const uint BankRecordHeaderOffset = BoneBlockHeader.BankRecordOffset;
-		const FBankRecordHeader BankRecordHeader = BankBuffer.Load<FBankRecordHeader>(BankRecordHeaderOffset);
-		if (BankRecordHeader.Playing)
+		const uint BoneTransformSize      = sizeof(FBoneTransform);
+		const uint MappingTransformIndex  = BoneBlockHeader.MappingTransformIndex;
+		const uint SequenceTransformIndex = BoneBlockHeader.SequenceTransformIndex;
+		const uint TransformIndex         = LocalTransformIndex + (BoneBlockHeader.BlockLocalIndex * BONES_PER_GROUP);
+
+		uint BoneIndex = TransformIndex;
+	
+		if (BoneBlockHeader.BoneMapOffset != ~uint(0))
 		{
-			const uint BoneTransformSize	= (uint)select(BankRecordHeader.HasScale != 0, sizeof(FBoneTransformWithScale), sizeof(FBoneTransform));
-			const uint BoneTransformRange	= BankRecordHeader.BoneCount * BoneTransformSize;
-			const uint MappingKeyOffset 	= BankRecordHeaderOffset + (uint)sizeof(FBankRecordHeader);
-			const uint BoneKeyOffset		= MappingKeyOffset + BoneTransformRange;
-
-			const float SampleInterval = 1.0f / BankRecordHeader.SampleRate;
-			const float TrackLength = (BankRecordHeader.FrameCount - 1) * SampleInterval;
-
-			float Time = fmod(BankRecordHeader.CurrentTime, TrackLength);
-			if (Time < 0.0f) 
-			{
-				Time += TrackLength;
-			}
-
-			const float FramesPerSecond = -1.0f;
-			FKeyIndices KeyIndices = GetKeyIndicesFromTime(Time, BankRecordHeader.FrameCount, TrackLength, FramesPerSecond);
-
-			if (!BankRecordHeader.Interpolating)
-			{
-				// Forcing alpha to zero disables pose interpolation (interpolation method is "step")
-				KeyIndices.Alpha = 0.0f;
-			}
-
-			const uint BoneIndex = (BoneBlockHeader.BlockLocalIndex * BONES_PER_GROUP) + LocalBoneIndex;
-
-			// TODO: Scale version
-
-			FBoneTransform InvGlobalRefPoseTransform = BankBuffer.Load<FBoneTransform>(MappingKeyOffset + ((uint)BoneIndex * BoneTransformSize));
-			FBoneTransform MeshGlobalAnimPoseTransform;
-		
-			if (KeyIndices.Alpha <= 0.0f)
-			{
-				const uint KeyIndex = (KeyIndices.KeyIndex1 * BankRecordHeader.BoneCount) + BoneIndex;
-				MeshGlobalAnimPoseTransform = BankBuffer.Load<FBoneTransform>(BoneKeyOffset + (KeyIndex * BoneTransformSize));
-				QuatNormalize(MeshGlobalAnimPoseTransform.Rotation);
-			}
-			else if (KeyIndices.Alpha >= 1.0f)
-			{
-				const uint KeyIndex = (KeyIndices.KeyIndex2 * BankRecordHeader.BoneCount) + BoneIndex;
-				MeshGlobalAnimPoseTransform = BankBuffer.Load<FBoneTransform>(BoneKeyOffset + (KeyIndex * BoneTransformSize));
-				QuatNormalize(MeshGlobalAnimPoseTransform.Rotation);
-			}
-			else
-			{
-				const uint KeyIndexA = (KeyIndices.KeyIndex1 * BankRecordHeader.BoneCount) + BoneIndex;
-				const uint KeyIndexB = (KeyIndices.KeyIndex2 * BankRecordHeader.BoneCount) + BoneIndex;
-
-				FBoneTransform MeshGlobalTransformA = BankBuffer.Load<FBoneTransform>(BoneKeyOffset + (KeyIndexA * BoneTransformSize));
-				FBoneTransform MeshGlobalTransformB = BankBuffer.Load<FBoneTransform>(BoneKeyOffset + (KeyIndexB * BoneTransformSize));
-
-				// Ensure rotations are normalized before blending
-				QuatNormalize(MeshGlobalTransformA.Rotation);
-				QuatNormalize(MeshGlobalTransformB.Rotation);
-
-				// NOTE: Rotation will already be normalized
-				MeshGlobalAnimPoseTransform = BlendBoneTransforms(MeshGlobalTransformA, MeshGlobalTransformB, KeyIndices.Alpha);
-			}
-
-			MeshGlobalAnimPoseTransform = MultiplyBoneTransform(InvGlobalRefPoseTransform, MeshGlobalAnimPoseTransform);
-			const float4x4 CurrentTransform44 = MatrixFromBoneTransform(MeshGlobalAnimPoseTransform);
-			const float3x4 CurrentTransform = (float3x4)transpose(CurrentTransform44);
-
-			StoreCompressedBoneTransform(TransformBuffer, BankRecordHeader.TransformOffset, BoneIndex, transpose(CurrentTransform));
+			BoneIndex = BoneMapBuffer.Load<uint>(BoneBlockHeader.BoneMapOffset + TransformIndex * uint(sizeof(uint)));
 		}
-	}
-}
 
-[numthreads(BONES_PER_GROUP, 1, 1)]
-void BankScatterCS(uint3 GroupId : SV_GroupID, uint GroupIndex : SV_GroupIndex)
-{
-	const uint GlobalBlockIndex		= GroupId.x;
-	const uint LocalTransformIndex	= GroupIndex;
+		const float TrackLength = (BoneBlockHeader.FrameCount - 1) * SampleInterval;
+		const float Time = BoneBlockHeader.Time;
 
-	const FBankScatterHeader BlockHeader = HeaderBuffer.Load<FBankScatterHeader>(GlobalBlockIndex * uint(sizeof(FBankScatterHeader)));
-	if (LocalTransformIndex < BlockHeader.BlockTransformCount)
-	{
-		float4x3 PreviousTransform;
-		float4x3 CurrentTransform;
+		const float FramesPerSecond = -1.0f;
+		FKeyIndices KeyIndices = GetKeyIndicesFromTime(Time, BoneBlockHeader.FrameCount, TrackLength, FramesPerSecond);
+
+		if (!bInterpolating)
+		{
+			// Forcing alpha to zero disables pose interpolation (interpolation method is "step")
+			KeyIndices.Alpha = 0.0f;
+		}
+
+		// TODO: Scale version
+
+		FBoneTransform InvGlobalRefPoseTransform = SequenceTransformBuffer.Load<FBoneTransform>((MappingTransformIndex + BoneIndex) * BoneTransformSize);
+		FBoneTransform MeshGlobalAnimPoseTransform;
 
-		BRANCH
-		if (BlockHeader.BlockSrcTransformOffset == ~uint(0))
+		if (KeyIndices.Alpha <= 0.0f)
 		{
-			PreviousTransform[0] = float3(1.0f, 0.0f, 0.0f);
-			PreviousTransform[1] = float3(0.0f, 1.0f, 0.0f);
-			PreviousTransform[2] = float3(0.0f, 0.0f, 1.0f);
-			PreviousTransform[3] = float3(0.0f, 0.0f, 0.0f);
-
-			CurrentTransform[0] = float3(1.0f, 0.0f, 0.0f);
-			CurrentTransform[1] = float3(0.0f, 1.0f, 0.0f);
-			CurrentTransform[2] = float3(0.0f, 0.0f, 1.0f);
-			CurrentTransform[3] = float3(0.0f, 0.0f, 0.0f);
+			const uint KeyIndex = (KeyIndices.KeyIndex1 * BoneBlockHeader.BoneCount) + BoneIndex;
+			MeshGlobalAnimPoseTransform = SequenceTransformBuffer.Load<FBoneTransform>((SequenceTransformIndex + KeyIndex) * BoneTransformSize);
+			QuatNormalize(MeshGlobalAnimPoseTransform.Rotation);
+		}
+		else if (KeyIndices.Alpha >= 1.0f)
+		{
+			const uint KeyIndex = (KeyIndices.KeyIndex2 * BoneBlockHeader.BoneCount) + BoneIndex;
+			MeshGlobalAnimPoseTransform = SequenceTransformBuffer.Load<FBoneTransform>((SequenceTransformIndex + KeyIndex) * BoneTransformSize);
+			QuatNormalize(MeshGlobalAnimPoseTransform.Rotation);
 		}
 		else
 		{
-			PreviousTransform	= LoadCompressedBoneTransform(TransformBuffer, BlockHeader.BlockDstTransformOffset, LocalTransformIndex);
-
-			if (BlockHeader.BlockSrcBoneMapOffset != ~uint(0))
-			{
-				uint BoneMapOffset         = BlockHeader.BlockSrcBoneMapOffset + LocalTransformIndex * uint(sizeof(uint));
-				uint BoneMapTransformIndex = SrcBoneMapBuffer.Load<uint>(BoneMapOffset);
-				CurrentTransform           = LoadCompressedBoneTransform(SrcTransformBuffer, BlockHeader.BlockSrcTransformOffset, BoneMapTransformIndex);
-			}
-			else
-			{
-				CurrentTransform = LoadCompressedBoneTransform(SrcTransformBuffer, BlockHeader.BlockSrcTransformOffset, LocalTransformIndex);
-			}
+			const uint KeyIndexA = (KeyIndices.KeyIndex1 * BoneBlockHeader.BoneCount) + BoneIndex;
+			const uint KeyIndexB = (KeyIndices.KeyIndex2 * BoneBlockHeader.BoneCount) + BoneIndex;
+
+			FBoneTransform MeshGlobalTransformA = SequenceTransformBuffer.Load<FBoneTransform>((SequenceTransformIndex + KeyIndexA) * BoneTransformSize);
+			FBoneTransform MeshGlobalTransformB = SequenceTransformBuffer.Load<FBoneTransform>((SequenceTransformIndex + KeyIndexB) * BoneTransformSize);
+
+			// Ensure rotations are normalized before blending
+			QuatNormalize(MeshGlobalTransformA.Rotation);
+			QuatNormalize(MeshGlobalTransformB.Rotation);
+
+			// NOTE: Rotation will already be normalized
+			MeshGlobalAnimPoseTransform = BlendBoneTransforms(MeshGlobalTransformA, MeshGlobalTransformB, KeyIndices.Alpha);
 		}
 
-		StoreCompressedBoneTransform(TransformBuffer, BlockHeader.BlockDstTransformOffset, LocalTransformIndex, CurrentTransform);
-		StoreCompressedBoneTransform(TransformBuffer, BlockHeader.BlockDstTransformOffset, BlockHeader.TotalTransformCount + LocalTransformIndex, PreviousTransform);
+		MeshGlobalAnimPoseTransform = MultiplyBoneTransform(InvGlobalRefPoseTransform, MeshGlobalAnimPoseTransform);
+		const float4x3 CurrentTransform = (float4x3)MatrixFromBoneTransform(MeshGlobalAnimPoseTransform);
+
+		const uint OutputTransformOffset = BoneBlockHeader.OutputTransformIndex * sizeof(FCompressedBoneTransform);
+		StoreCompressedBoneTransform(TransformBuffer, OutputTransformOffset, TransformIndex, CurrentTransform);
+
+		if (BoneBlockHeader.OutputTransformCopyIndexOffset != 0)
+		{
+			const uint SecondaryOutputTransformOffset = (BoneBlockHeader.OutputTransformIndex + BoneBlockHeader.OutputTransformCopyIndexOffset) * sizeof(FCompressedBoneTransform);
+			StoreCompressedBoneTransform(TransformBuffer, SecondaryOutputTransformOffset, TransformIndex, CurrentTransform);
+		}
 	}
-}
+}
\ No newline at end of file
diff --git a/Engine/Shaders/Private/Skinning/SkinningSceneDataDecode.ush b/Engine/Shaders/Private/Skinning/SkinningSceneDataDecode.ush
--- a/Engine/Shaders/Private/Skinning/SkinningSceneDataDecode.ush
+++ b/Engine/Shaders/Private/Skinning/SkinningSceneDataDecode.ush
@@ -39,7 +39,7 @@ FBoneTransform LoadSkinningBoneObjectSpace(uint BufferOffset, uint BoneIndex)
 
 uint GetSkinningBoneTransformOffset(FSkinningHeader SkinningHeader, uint InstanceIndex, bool bPrevious)
 {
-	return SkinningHeader.TransformBufferOffset + (InstanceIndex + (bPrevious ^ SkinningHeader.CurrentTransformSlot)) * SkinningHeader.MaxTransformCount;
+	return SkinningHeader.TransformBufferOffset + (InstanceIndex * (SkinningHeader.bIsRefPose ? 0 : 1) + (bPrevious ^ SkinningHeader.CurrentTransformSlot)) * SkinningHeader.MaxTransformCount;
 }
 
 uint GetSkinningBoneTransformOffset(uint PrimitiveIndex, uint InstanceIndex, bool bPrevious)
diff --git a/Engine/Shaders/Shared/SkinningDefinitions.h b/Engine/Shaders/Shared/SkinningDefinitions.h
--- a/Engine/Shaders/Shared/SkinningDefinitions.h
+++ b/Engine/Shaders/Shared/SkinningDefinitions.h
@@ -20,40 +20,21 @@ namespace UE::HLSL
 {
 #endif
 
-struct FBankRecordHeader
+struct FAnimBankBlockHeader
 {
-	uint BoneCount;
+	float Time;
+	uint BoneMapOffset;
+	uint SequenceTransformIndex;
+	uint MappingTransformIndex;
+	uint OutputTransformIndex;
 	uint FrameCount;
-
-	float SampleRate;
-	float PlayRate;
-
-	float CurrentTime;
-	float PreviousTime;
-
-	uint TransformOffset	: 29;
-	uint Playing			: 1;
-	uint Interpolating		: 1;
-	uint HasScale			: 1;
-};
-
-struct FBankBlockHeader
-{
-	uint BlockLocalIndex;
-	uint BlockBoneCount;
-	uint BlockTransformOffset;
-	uint BankRecordOffset;
+	uint BoneCount                      : 16;
+	uint OutputTransformCopyIndexOffset : 16;
+	uint BlockLocalIndex                : 16;
+	uint BlockTransformCount            : 16;
 };
 
-struct FBankScatterHeader
-{
-	uint BlockLocalIndex;
-	uint BlockSrcBoneMapOffset;
-	uint BlockSrcTransformOffset;
-	uint BlockDstTransformOffset;
-	uint BlockTransformCount : 8;
-	uint TotalTransformCount : 24;
-};
+static_assert(sizeof(FAnimBankBlockHeader) == 32);
 
 #define ANIM_BANK_FLAG_NONE			0x0
 #define ANIM_BANK_FLAG_LOOPING		0x1
@@ -69,6 +50,9 @@ struct FBankScatterHeader
 #define UINT_TYPE uint
 #endif
 
+// Animation is always sampled at 30hz
+#define ANIM_BANK_SAMPLE_RATE 30.0f
+
 #define SKINNING_BUFFER_HIERARCHY_OFFSET_BITS 22
 #define SKINNING_BUFFER_HIERARCHY_OFFSET_MAX (1 << SKINNING_BUFFER_HIERARCHY_OFFSET_BITS) - 1
 
@@ -89,10 +73,13 @@ struct FSkinningHeader
 	UINT_TYPE MaxTransformCount       : 16;
 	UINT_TYPE CurrentTransformSlot    : 1;
 	UINT_TYPE MaxInfluenceCount       : SKINNING_BUFFER_INFLUENCE_BITS;
-	UINT_TYPE UniqueAnimationCount    : 7;
 	UINT_TYPE bHasScale               : 1;
+	UINT_TYPE bIsRefPose              : 1;
+	UINT_TYPE Unused                  : 6;
 };
 
+static_assert(sizeof(FSkinningHeader) == 16);
+
 #endif // defined(__cplusplus) || COMPILER_SUPPORTS_HLSL2021
 
 #if USE_COMPRESSED_BONE_TRANSFORM
diff --git a/Engine/Source/Runtime/Engine/Classes/Animation/AnimBank.h b/Engine/Source/Runtime/Engine/Classes/Animation/AnimBank.h
--- a/Engine/Source/Runtime/Engine/Classes/Animation/AnimBank.h
+++ b/Engine/Source/Runtime/Engine/Classes/Animation/AnimBank.h
@@ -13,6 +13,7 @@
 #include "SkinningDefinitions.h"
 #include "UObject/ObjectPtr.h"
 #include "UObject/WeakObjectPtr.h"
+#include "AnimTrackPool.h"
 
 #include "AnimBank.generated.h"
 
@@ -26,7 +27,9 @@ struct FRenderBounds;
 struct FReferenceSkeleton;
 class USkinnedAsset;
 class UInstancedSkinnedMeshComponent;
+class FInstancedSkinningSceneExtensionProxy;
 class FAnimBankBuildAsyncCacheTask;
+class IAnimBankTransformProvider;
 
 struct FSkinnedAssetMapping
 {
@@ -84,6 +87,12 @@ struct FAnimBankEntry
 	uint32 KeyCount		= 0u;
 	uint32 Flags		= 0u;
 
+	inline float GetPlayLength() const
+	{
+		const float SampleInterval = 1.0f / ANIM_BANK_SAMPLE_RATE;
+		return float(FrameCount - 1) * SampleInterval;
+	}
+
 	inline bool IsLooping() const
 	{
 		return (Flags & ANIM_BANK_FLAG_LOOPING) != 0;
@@ -177,10 +186,10 @@ public:
 #if WITH_EDITORONLY_DATA
 	UPROPERTY(EditAnywhere, BlueprintReadWrite, Category = Sequences, meta = (ShowOnlyInnerProperties))
 	TArray<struct FAnimBankSequence> Sequences;
+#endif
 
 	UPROPERTY(EditAnywhere, BlueprintReadWrite, Category = Mapping, meta = (ShowOnlyInnerProperties))
 	TObjectPtr<USkinnedAsset> Asset;
-#endif
 
 public:
 	virtual void Serialize(FArchive& Ar) override;
@@ -203,6 +212,22 @@ public:
 		return Data;
 	}
 
+	UFUNCTION(BlueprintCallable, Category = "TransformProviders|AnimBank")
+	ENGINE_API float GetSequencePlayLength(int32 SequenceIndex) const;
+	
+	UFUNCTION(BlueprintCallable, Category = "TransformProviders|AnimBank")
+	ENGINE_API bool IsSequenceLooping(int32 SequenceIndex) const;
+
+	UFUNCTION(BlueprintCallable, Category = "TransformProviders|AnimBank")
+	inline float GetSequenceCount() const
+	{
+#if WITH_EDITORONLY_DATA
+		return Sequences.Num();
+#else
+		return GetData().Entries.Num();
+#endif
+	}
+
 #if WITH_EDITOR
 	virtual void PostEditChangeProperty(FPropertyChangedEvent& PropertyChangedEvent) override;
 	virtual void BeginCacheForCookedPlatformData(const ITargetPlatform* TargetPlatform) override;
@@ -525,79 +550,246 @@ public:
 	TSoftObjectPtr<class UTransformProviderData> TransformProvider;
 };
 
-struct FAnimBankDesc
+UENUM(BlueprintType)
+enum class EAnimBankTrackLoopMode : uint8
 {
-	typedef Experimental::FHashType FDescHash;
+	// The behavior is unchanged from the sequence default.
+	Auto,
 
-	TWeakObjectPtr<const UAnimBank> BankAsset;
-	TWeakObjectPtr<const USkinnedAsset> Asset;
+	// Looping is forced for the current track.
+	Loop,
 
-	uint32 SequenceIndex	= 0u;
-	float Position			= 0.0f;
-	float PlayRate			= 1.0f;
-	uint8 bLooping   : 1	= 1;
-	uint8 bAutoStart : 1	= 1;
+	// Clamp is forced for the current track.
+	Clamp
+};
 
-	ENGINE_API uint32 GetHash() const;
+USTRUCT(BlueprintType)
+struct FAnimBankTrackAutoPlayData
+{
+	GENERATED_USTRUCT_BODY()
 
-	friend uint32 GetTypeHash(const FAnimBankDesc& Desc)
-	{
-		return Desc.GetHash();
-	}
+	// The anim bank item index to play.
+	UPROPERTY(EditAnywhere, BlueprintReadWrite, Category = Animation)
+	int32 ItemIndex = 0;
+
+	// The starting time position to begin play.
+	UPROPERTY(EditAnywhere, BlueprintReadWrite, Category = Animation)
+	float Position = 0.0f;
+
+	// The rate multiplier to play the animation.
+	UPROPERTY(EditAnywhere, BlueprintReadWrite, Category = Animation)
+	float PlayRate = 1.0f;
+
+	// The mode to control looping behavior.
+	UPROPERTY(EditAnywhere, BlueprintReadWrite, Category = Animation)
+	EAnimBankTrackLoopMode LoopMode = EAnimBankTrackLoopMode::Auto;
 };
 
-struct FAnimBankRecordHandle
+USTRUCT(BlueprintType)
+struct FAnimBankTrackManualData
 {
-	int32 Id	= INDEX_NONE;
-	uint32 Hash	= 0u;
+	GENERATED_USTRUCT_BODY()
+
+	// The anim bank item index to play.
+	UPROPERTY(EditAnywhere, BlueprintReadWrite, Category = Animation)
+	int32 ItemIndex = 0;
+
+	// The time position so sample in the sequence.
+	UPROPERTY(EditAnywhere, BlueprintReadWrite, Category = Animation)
+	float Position = 0.0f;
+
+	// Whether to apply the previous position when updating the animation. When false, the previous position is derived from the last frame.
+	UPROPERTY(EditAnywhere, BlueprintReadWrite, Category = Animation)
+	uint32 bHasPreviousPosition : 1 = 0;
+
+	// A time position to update previous frame transforms.
+	UPROPERTY(EditAnywhere, BlueprintReadWrite, Category = Animation, meta = (EditCondition = "bHasPreviousPosition", EditConditionHides))
+	float PreviousPosition = 0.0f;
+};
+
+class FAnimBankTransformProviderProxy;
+class UAnimBankDataInstance;
+
+class FAnimBankTrackPackedData
+{
+public:
+	/////////////////////////////////////////////////////////////////////////////////////////////////
+	// AutoPlay
+
+	static FAnimBankTrackPackedData Pack(FAnimBankTrackAutoPlayData AutoPlayData)
+	{
+		FAnimBankTrackPackedData Data;
+		Data.ItemIndex            = AutoPlayData.ItemIndex;
+		Data.bAutoPlay            = 1;
+		Data.LoopMode             = static_cast<uint32>(AutoPlayData.LoopMode);
+		Data.bHasPreviousPosition = true;
+		Data.Position             = AutoPlayData.Position;
+		Data.PreviousPosition     = AutoPlayData.Position;
+		Data.PlayRate             = AutoPlayData.PlayRate;
+		return Data;
+	}
+
+	bool IsAutoPlay() const
+	{
+		return bAutoPlay;
+	}
+
+	FAnimBankTrackAutoPlayData UnpackAutoPlay() const
+	{
+		check(IsAutoPlay());
+
+		FAnimBankTrackAutoPlayData Data;
+		Data.ItemIndex = ItemIndex;
+		Data.Position  = Position;
+		Data.PlayRate  = PlayRate;
+		Data.LoopMode  = static_cast<EAnimBankTrackLoopMode>(LoopMode);
+		return Data;
+	}
+
+	/////////////////////////////////////////////////////////////////////////////////////////////////
+	// Manual
+
+	static FAnimBankTrackPackedData Pack(FAnimBankTrackManualData ManualData)
+	{
+		FAnimBankTrackPackedData Data;
+		Data.ItemIndex            = ManualData.ItemIndex;
+		Data.bAutoPlay            = 0;
+		Data.LoopMode             = 0;
+		Data.bHasPreviousPosition = ManualData.bHasPreviousPosition;
+		Data.Position             = ManualData.Position;
+		Data.PreviousPosition     = ManualData.bHasPreviousPosition ? ManualData.PreviousPosition : ManualData.Position;
+		Data.PlayRate             = 0.0f;
+		return Data;
+	}
+
+	bool IsManual() const
+	{
+		return !bAutoPlay;
+	}
+
+	FAnimBankTrackManualData UnpackManual() const
+	{
+		check(IsManual());
+
+		FAnimBankTrackManualData Data;
+		Data.ItemIndex            = ItemIndex;
+		Data.Position             = Position;
+		Data.PreviousPosition     = PreviousPosition;
+		Data.bHasPreviousPosition = bHasPreviousPosition;
+		return Data;
+	}
+
+	/////////////////////////////////////////////////////////////////////////////////////////////////
+
+	int32 GetItemIndex() const
+	{
+		return ItemIndex;
+	}
+
+	float GetPosition() const
+	{
+		return Position;
+	}
 
-	inline bool IsValid() const
+	float GetPreviousPosition() const
 	{
-		return Id != INDEX_NONE && Hash != 0u;
+		return PreviousPosition;
 	}
+
+	bool HasPreviousPosition() const
+	{
+		return bHasPreviousPosition;
+	}
+
+	EAnimBankTrackLoopMode GetLoopMode() const
+	{
+		return static_cast<EAnimBankTrackLoopMode>(LoopMode);
+	}
+
+	ENGINE_API void Update(float SequenceLength, bool bLooping, float DeltaTime);
+
+	ENGINE_API void Wrap(float SequenceLength, bool bLooping);
+
+	ENGINE_API void Serialize(FArchive& Ar);
+
+private:
+	uint32 ItemIndex            : 20;
+	uint32 LoopMode             : 2;
+	uint32 bAutoPlay            : 1;
+	uint32 bHasPreviousPosition : 1;
+	uint32 Unused               : 8;
+	float Position;
+	float PreviousPosition;
+	float PlayRate;
+
+	friend class FAnimBankTransformProvider;
 };
 
-struct FAnimBankRecord
-{
-	typedef Experimental::FHashElementId FRecordId;
+template<> struct TCanBulkSerialize<FAnimBankTrackPackedData> { enum { Value = true }; };
 
-	FAnimBankDesc Desc;
-	int32 RecordId = INDEX_NONE;
-	int32 KeyOffset = INDEX_NONE;
-	uint32 KeyCount = 0;
-	int32 FrameCount = 0;
-	int32 ReferenceCount = 0;
+using FAnimBankTrackPool = TAnimTrackPool<FAnimBankTrackPackedData>;
 
-	// TODO: De-dup using USkinnedAsset*
-	FSkinnedAssetMapping AssetMapping;
+class FAnimBankRenderData
+{
+public:
+	ENGINE_API ~FAnimBankRenderData();
 
-	TArray<FVector3f>	PositionKeys;
-	TArray<FQuat4f>		RotationKeys;
+private:
+	void Patch(FAnimBankTrackPool::FPatch&& Patch);
+	void InsertProxy(FInstancedSkinningSceneExtensionProxy* Proxy);
+	void RemoveProxy(FInstancedSkinningSceneExtensionProxy* Proxy);
 
-	// Playback
-	uint8				Playing : 1 = 0;
-	float				CurrentTime = 0.0f;
-	float				PreviousTime = 0.0f;
+	FAnimBankTrackPool Tracks;
+	Experimental::TRobinHoodHashSet<FInstancedSkinningSceneExtensionProxy*> Proxies;
+
+	friend FAnimBankTransformProviderProxy;
+	friend UAnimBankDataInstance;
 };
 
-struct FAnimBankDescKeyFuncs : TDefaultMapHashableKeyFuncs<FAnimBankDesc, FAnimBankRecord, false>
+class FAnimBankTransformProviderProxy : public FTransformProviderRenderProxy
 {
-	static inline bool Matches(KeyInitType A, KeyInitType B)
+	friend class UAnimBankDataInstance;
+	friend class UAnimBankData;
+
+	ENGINE_API FAnimBankTransformProviderProxy(TConstArrayView<FAnimBankItem> AnimBankItems, FInstancedSkinningSceneExtensionProxy* SceneExtensionProxy);
+	ENGINE_API FAnimBankTransformProviderProxy(TConstArrayView<FAnimBankItem> AnimBankItems, FInstancedSkinningSceneExtensionProxy* SceneExtensionProxy, FAnimBankRenderData& RenderData);
+
+public:
+	virtual void CreateRenderThreadResources(FSceneInterface& Scene, FRHICommandListBase& RHICmdList) override;
+	virtual void DestroyRenderThreadResources() override;
+
+	FAnimBankTrackPool& GetTracks()
 	{
-		return
-			A.BankAsset		== B.BankAsset &&
-			A.SequenceIndex	== B.SequenceIndex &&
-			A.Asset			== B.Asset &&
-			A.Position		== B.Position &&
-			A.PlayRate		== B.PlayRate &&
-			A.bLooping		== B.bLooping &&
-			A.bAutoStart	== B.bAutoStart;
+		return *Tracks;
 	}
 
-	static inline uint32 GetKeyHash(KeyInitType Key)
+	TConstArrayView<FAnimBankItem> GetAnimBankItems() const
 	{
-		return Key.GetHash();
+		return AnimBankItems;
 	}
+
+	const FInstancedSkinningSceneExtensionProxy* GetSceneExtensionProxy() const
+	{
+		return SceneExtensionProxy;
+	}
+
+private:
+	TConstArrayView<FAnimBankItem> AnimBankItems;
+	TUniquePtr<FAnimBankTrackPool> TrackAllocator;
+	FAnimBankRenderData* RenderData = nullptr;
+	FAnimBankTrackPool* Tracks;
+	FInstancedSkinningSceneExtensionProxy* SceneExtensionProxy = nullptr;
+	const USkinnedAsset* SkinnedAsset = nullptr;
+	IAnimBankTransformProvider* AnimBankTransformProvider = nullptr;
+};
+
+class IAnimBankTransformProvider
+{
+public:
+	virtual ~IAnimBankTransformProvider() = default;
+
+	virtual void RegisterAnimBank(FAnimBankTransformProviderProxy* Proxy) {}
+	virtual void UnregisterAnimBank(FAnimBankTransformProviderProxy* Proxy) {}
 };
 
 UCLASS(config=Engine, hidecategories=Object, MinimalAPI, BlueprintType)
@@ -608,13 +800,12 @@ class UAnimBankData : public UTransformProviderData
 public:
 	virtual bool IsEnabled() const override;
 	virtual const FGuid& GetTransformProviderID() const override;
+	virtual uint32 GetUniqueAnimationCount() const override;
 	virtual bool UsesSkeletonBatching() const override;
-	virtual const uint32 GetUniqueAnimationCount() const override;
 	virtual bool HasAnimationBounds() const override;
 	virtual bool GetAnimationBounds(uint32 AnimationIndex, FRenderBounds& OutBounds) const override;
 	virtual uint32 GetSkinningDataOffset(int32 InstanceIndex, const FTransform& ComponentTransform, const FSkinnedMeshInstanceData& InstanceData) const override;
-	virtual FTransformProviderRenderProxy* CreateRenderThreadResources(FSkinningSceneExtensionProxy* SceneProxy, FSceneInterface& Scene, FRHICommandListBase& RHICmdList) override;
-	virtual void DestroyRenderThreadResources(FTransformProviderRenderProxy* ProviderProxy) override;
+	virtual FTransformProviderRenderProxy* CreateRenderProxy(FInstancedSkinningSceneExtensionProxy* ExtensionProxy) const override;
 	virtual bool IsCompiling() const override;
 
 #if WITH_EDITOR
@@ -625,30 +816,79 @@ public:
 public:
 	UPROPERTY(EditAnywhere, BlueprintReadOnly, Category = TransformProvider, meta = (ShowOnlyInnerProperties))
 	TArray<struct FAnimBankItem> AnimBankItems;
+
+	// Returns the play length for an anim bank item.
+	UFUNCTION(BlueprintCallable, Category = "TransformProviders|AnimBank")
+	ENGINE_API float GetSequencePlayLength(int32 ItemIndex) const;
+
+protected:
+	bool IsAnimBankValid(FInstancedSkinningSceneExtensionProxy* ExtensionProxy) const;
 };
 
-class FAnimBankDataRenderProxy : public FTransformProviderRenderProxy
+UCLASS(config=Engine, hidecategories=Object, MinimalAPI, BlueprintType, Within=InstancedSkinnedMeshComponent)
+class UAnimBankDataInstance : public UAnimBankData
 {
-	friend class UAnimBankData;
+	GENERATED_BODY()
+
+public:
+	virtual bool IsEnabled() const override;
+	virtual uint32 GetUniqueAnimationCount() const override;
+	virtual bool GetAnimationBounds(uint32 AnimationIndex, FRenderBounds& OutBounds) const override;
+	virtual uint32 GetSkinningDataOffset(int32 InstanceIndex, const FTransform& ComponentTransform, const FSkinnedMeshInstanceData& InstanceData) const override;
+	virtual FTransformProviderRenderProxy* CreateRenderProxy(FInstancedSkinningSceneExtensionProxy* ExtensionProxy) const override;
+	virtual void BeginDestroy() override;
+	virtual void Serialize(FArchive& Ar) override;
 
 public:
-	ENGINE_API FAnimBankDataRenderProxy(UAnimBankData* BankData, FSkinningSceneExtensionProxy* SceneProxy, FSceneInterface& Scene);
-	ENGINE_API virtual ~FAnimBankDataRenderProxy();
-	ENGINE_API virtual void CreateRenderThreadResources(FRHICommandListBase& RHICmdList) override;
-	ENGINE_API virtual void DestroyRenderThreadResources() override;
-	ENGINE_API virtual const TConstArrayView<uint64> GetProviderData(bool& bOutValid) const override;
+	// Copies items from this anim bank to create a new anim bank instance.
+	UFUNCTION(BlueprintCallable, Category = "TransformProviders|AnimBank")
+	static ENGINE_API UAnimBankDataInstance* Create(UAnimBankData* AnimBank, UInstancedSkinnedMeshComponent* Owner);
+
+	// Allocates a new track with auto play data, or -1 if the track allocation failed.
+	UFUNCTION(BlueprintCallable, Category = "TransformProviders|AnimBank")
+	ENGINE_API int32 AllocateAutoPlayTrack(const FAnimBankTrackAutoPlayData& Data);
+
+	// Allocates a new track with manual play data, or -1 if the track allocation failed.
+	UFUNCTION(BlueprintCallable, Category = "TransformProviders|AnimBank")
+	ENGINE_API int32 AllocateManualTrack(const FAnimBankTrackManualData& Data);
+
+	// Updates a track with auto play data. Returns whether the update succeeded.
+	UFUNCTION(BlueprintCallable, Category = "TransformProviders|AnimBank")
+	ENGINE_API bool UpdateAutoPlayTrack(int32 TrackIndex, const FAnimBankTrackAutoPlayData& Data);
+
+	// Updates a track with manual data. Returns whether the update succeeded.
+	UFUNCTION(BlueprintCallable, Category = "TransformProviders|AnimBank")
+	ENGINE_API bool UpdateManualTrack(int32 TrackIndex, const FAnimBankTrackManualData& Data);
+
+	// Deallocates a track. Returns whether the deallocation succeeded.
+	UFUNCTION(BlueprintCallable, Category = "TransformProviders|AnimBank")
+	ENGINE_API bool DeallocateTrack(int32 TrackIndex);
+
+	// Submits all pending changes to rendering.
+	UFUNCTION(BlueprintCallable, Category = "TransformProviders|AnimBank")
+	ENGINE_API void SubmitChanges();
 
 private:
-	FSkinningSceneExtensionProxy* SceneProxy = nullptr;
-	TObjectPtr<const class USkinnedAsset> SkinnedAsset = nullptr;
-	FSceneInterface& Scene;
-	TArray<uint64> AnimBankIds;
-	TArray<FAnimBankRecordHandle> AnimBankHandles;
-	TArray<FAnimBankItem> AnimBankItems;
-	uint32 UniqueAnimationCount = 0;
-};
+	void InitRenderData();
+	void InitSequences();
+
+	struct FSequence
+	{
+		FSequence(float InPlayLength, bool bInLooping)
+			: PlayLength(InPlayLength)
+			, bLooping(bInLooping)
+		{}
 
-using FAnimBankRecordMap = Experimental::TRobinHoodHashMap<FAnimBankDesc, FAnimBankRecord, FAnimBankDescKeyFuncs>;
+		float PlayLength;
+		bool bLooping;
+	};
+
+	FAnimBankTrackPool Tracks;
+	TArray<FSequence> Sequences;
+	TUniquePtr<FAnimBankRenderData> RenderData{ new FAnimBankRenderData };
+
+	friend class UAnimBankData;
+};
 
 namespace UE::AnimBank
 {
diff --git a/Engine/Source/Runtime/Engine/Classes/Animation/AnimTrackPool.h b/Engine/Source/Runtime/Engine/Classes/Animation/AnimTrackPool.h
new file mode 100644
--- /dev/null
+++ b/Engine/Source/Runtime/Engine/Classes/Animation/AnimTrackPool.h
@@ -0,0 +1,306 @@
+// Copyright Epic Games, Inc. All Rights Reserved.
+
+#pragma once
+
+#include "Serialization/Archive.h"
+#include "Experimental/ConcurrentLinearAllocator.h"
+
+/**
+ * An pool / allocator class that manages a pool of tracks of type DataType. Supports allocation / deallocation / updates of track data, as well as
+ * a mechanism to build diffs into a Patch data structure which can be marshaled and applied to another instance for thread-safe game / render updates.
+ * Track are pool allocated and so only a subset may be active at one time.
+ */
+template <typename DataType>
+class TAnimTrackPool
+{
+public:
+	class FPatch
+	{
+	public:
+		int32 GetNumTracks() const
+		{
+			return NumTracks;
+		}
+
+	private:
+		struct FActivateOp
+		{
+			FActivateOp(int32 InIndex, const DataType& InData)
+				: Index(InIndex)
+				, Data(InData)
+			{}
+
+			uint32 Index;
+			DataType Data;
+		};
+
+		struct FDeactivateOp
+		{
+			FDeactivateOp(int32 InIndex)
+				: Index(InIndex)
+			{}
+
+			uint32 Index;
+		};
+
+		int32 NumTracks;
+		int32 NumActiveTracks;
+		TArray<FActivateOp, FConcurrentLinearArrayAllocator> ActivateOps;
+		TArray<FDeactivateOp, FConcurrentLinearArrayAllocator> DeactivateOps;
+
+		friend class TAnimTrackPool;
+	};
+
+	void Reserve(int32 Num)
+	{
+		Datas.Reserve(Num);
+		ActiveBits.Reserve(Num);
+		DirtyBits.Reserve(Num);
+	}
+
+	int32 AllocateTrack(const DataType& Data)
+	{
+		const int32 Index = ActiveBits.FindAndSetFirstZeroBit();
+		NumActiveTracks++;
+
+		if (Index != INDEX_NONE)
+		{
+			SetDirty(Index);
+			Datas[Index] = Data;
+			return Index;
+		}
+
+		Datas.Emplace(Data);
+		ActiveBits.Add(true);
+		DirtyBits.Add(true);
+		NumDirtyTracks++;
+
+		return Datas.Num() - 1;
+	}
+
+	bool UpdateTrack(int32 Index, const DataType& Data)
+	{
+		if (IsActiveIndex(Index))
+		{
+			SetDirty(Index);
+			Datas[Index] = Data;
+			return true;
+		}
+		return false;
+	}
+
+	bool DeallocateTrack(int32 Index)
+	{
+		if (IsActiveIndex(Index))
+		{
+			SetDirty(Index);
+			ActiveBits[Index] = false;
+			NumActiveTracks--;
+			return true;
+		}
+		return false;
+	}
+
+	FPatch Finalize()
+	{
+		FPatch Patch;
+		Patch.NumTracks = Datas.Num();
+		Patch.NumActiveTracks = NumActiveTracks;
+		Patch.DeactivateOps.Reserve(Datas.Num());
+		Patch.ActivateOps.Reserve(Datas.Num());
+
+		for (TConstSetBitIterator<> BitIt(DirtyBits); BitIt; ++BitIt)
+		{
+			const uint32 Index = BitIt.GetIndex();
+			if (ActiveBits[Index])
+			{
+				Patch.ActivateOps.Emplace(Index, Datas[Index]);
+			}
+			else
+			{
+				Patch.DeactivateOps.Emplace(Index);
+			}
+		}
+		SetAllClean();
+
+		return Patch;
+	}
+
+	void Patch(FPatch&& Patch)
+	{
+		Datas.SetNumUninitialized(Patch.NumTracks);
+		ActiveBits.SetNum(Patch.NumTracks, false);
+		DirtyBits.SetNum(Patch.NumTracks, false);
+		NumActiveTracks = Patch.NumActiveTracks;
+
+		for (typename FPatch::FDeactivateOp Op : Patch.DeactivateOps)
+		{
+			ActiveBits[Op.Index] = false;
+			SetDirty(Op.Index);
+		}
+
+		for (const typename FPatch::FActivateOp& Op : Patch.ActivateOps)
+		{
+			ActiveBits[Op.Index] = true;
+			Datas[Op.Index] = Op.Data;
+			SetDirty(Op.Index);
+		}
+	}
+
+	void Serialize(FArchive& Ar)
+	{
+		Ar << NumActiveTracks;
+		Ar << Datas;
+		Ar << ActiveBits;
+
+		if (Ar.IsLoading())
+		{
+			DirtyBits.Init(false, GetNumTracks());
+			NumDirtyTracks = 0;
+		}
+	}
+
+	int32 GetNumTracks() const
+	{
+		return Datas.Num();
+	}
+
+	int32 GetNumActiveTracks() const
+	{
+		return NumActiveTracks;
+	}
+
+	int32 GetNumDirtyTracks() const
+	{
+		return NumDirtyTracks;
+	}
+
+	template <typename LambdaType>
+	void EnumerateDirtyTracks(LambdaType&& Lambda)
+	{
+		for (TConstSetBitIterator<> BitIt(DirtyBits); BitIt; ++BitIt)
+		{
+			Lambda(BitIt.GetIndex());
+		}
+	}
+
+	template <typename LambdaType>
+	void EnumerateDirtyTracks(LambdaType&& Lambda) const
+	{
+		for (TConstSetBitIterator<> BitIt(DirtyBits); BitIt; ++BitIt)
+		{
+			Lambda(BitIt.GetIndex());
+		}
+	}
+
+	template <typename LambdaType>
+	void EnumerateActiveTracks(LambdaType&& Lambda)
+	{
+		for (TConstSetBitIterator<> BitIt(ActiveBits); BitIt; ++BitIt)
+		{
+			Lambda(BitIt.GetIndex());
+		}
+	}
+
+	template <typename LambdaType>
+	void EnumerateActiveTracks(LambdaType&& Lambda) const
+	{
+		for (TConstSetBitIterator<> BitIt(ActiveBits); BitIt; ++BitIt)
+		{
+			Lambda(BitIt.GetIndex());
+		}
+	}
+	
+	template <typename LambdaType>
+	void EnumerateActiveDirtyTracks(LambdaType&& Lambda)
+	{
+		for (TConstDualSetBitIterator<> BitIt(ActiveBits, DirtyBits); BitIt; ++BitIt)
+		{
+			Lambda(BitIt.GetIndex());
+		}
+	}
+
+	template <typename LambdaType>
+	void EnumerateActiveDirtyTracks(LambdaType&& Lambda) const
+	{
+		for (TConstDualSetBitIterator<> BitIt(ActiveBits, DirtyBits); BitIt; ++BitIt)
+		{
+			Lambda(BitIt.GetIndex());
+		}
+	}
+
+	bool IsValidIndex(int32 Index) const
+	{
+		return Datas.IsValidIndex(Index);
+	}
+
+	bool IsActiveIndex(int32 Index) const
+	{
+		return IsValidIndex(Index) && ActiveBits[Index];
+	}
+
+	bool IsActiveIndexChecked(int32 Index) const
+	{
+		return ActiveBits[Index];
+	}
+
+	bool IsDirtyIndex(int32 Index) const
+	{
+		return IsValidIndex(Index) && DirtyBits[Index];
+	}
+
+	bool IsDirtyIndexChecked(int32 Index) const
+	{
+		return DirtyBits[Index];
+	}
+	
+	DataType& GetData(int32 Index)
+	{
+		check(IsActiveIndex(Index));
+		return Datas[Index];
+	}
+
+	const DataType& GetData(int32 Index) const
+	{
+		check(IsActiveIndex(Index));
+		return Datas[Index];
+	}
+
+	void SetAllDirty()
+	{
+		NumDirtyTracks = DirtyBits.Num();
+		DirtyBits.Init(true, NumDirtyTracks);
+	}
+
+	void SetDirty(int32 Index)
+	{
+		FBitReference Bit = DirtyBits[Index];
+		NumDirtyTracks += Bit ? 0 : 1;
+		Bit = true;
+	}
+
+	bool SetClean(int32 Index)
+	{
+		FBitReference Bit = DirtyBits[Index];
+		bool bWasDirty = Bit;
+		NumDirtyTracks -= bWasDirty ? 1 : 0;
+		Bit = false;
+		return bWasDirty;
+	}
+
+	void SetAllClean()
+	{
+		if (NumDirtyTracks > 0)
+		{
+			NumDirtyTracks = 0;
+			DirtyBits.Init(false, DirtyBits.Num());
+		}
+	}
+
+protected:
+	TArray<DataType> Datas;
+	TBitArray<> ActiveBits;
+	TBitArray<> DirtyBits;
+	int32 NumActiveTracks = 0;
+	int32 NumDirtyTracks = 0;
+};
\ No newline at end of file
diff --git a/Engine/Source/Runtime/Engine/Private/Animation/AnimBank.cpp b/Engine/Source/Runtime/Engine/Private/Animation/AnimBank.cpp
--- a/Engine/Source/Runtime/Engine/Private/Animation/AnimBank.cpp
+++ b/Engine/Source/Runtime/Engine/Private/Animation/AnimBank.cpp
@@ -1,6 +1,5 @@
 // Copyright Epic Games, Inc. All Rights Reserved.
 
-#include "Animation/AnimBank.h"
 #include "Animation/AnimBank.h"
 #include "Animation/AnimBoneCompressionCodec.h"
 #include "Animation/AnimSequence.h"
@@ -25,6 +24,7 @@
 #include "UObject/ObjectSaveContext.h"
 #include "UObject/UE5MainStreamObjectVersion.h"
 #include "UObject/UObjectThreadContext.h"
+#include "UObject/AnimObjectVersion.h"
 
 #include UE_INLINE_GENERATED_CPP_BY_NAME(AnimBank)
 
@@ -566,54 +566,58 @@ bool FAnimBankBuildAsyncCacheTask::BuildData(const UE::FSharedString& Name, cons
 			BoneIndices[Index] = (FBoneIndexType)Index;
 		}
 
-		FBoneContainer BoneContainer(BoneIndices, UE::Anim::FCurveFilterSettings(), *AnimBank->Asset);
+		{
+			TRACE_CPUPROFILER_EVENT_SCOPE(FAnimBankBuildAsyncCacheTask::BuildData);
+			FBoneContainer BoneContainer(BoneIndices, UE::Anim::FCurveFilterSettings(), *AnimBank->Asset);
 
-		int32 KeyIndex = 0;
+			int32 KeyIndex = 0;
 
-		for (uint32 Frame = 0; Frame < BankEntry.FrameCount; ++Frame)
-		{
-			// Some paths in the decompression code use mem stack, so make sure we put a mark here.
-			FMemMark Mark(FMemStack::Get());
+			for (uint32 Frame = 0; Frame < BankEntry.FrameCount; ++Frame)
+			{
+				TRACE_CPUPROFILER_EVENT_SCOPE(FAnimBankBuildAsyncCacheTask::Frame);
+				// Some paths in the decompression code use mem stack, so make sure we put a mark here.
+				FMemMark Mark(FMemStack::Get());
 
-			const double SeekTime = Sequence->GetSamplingFrameRate().AsSeconds(FFrameTime(int32(Frame)));
+				const double SeekTime = Sequence->GetSamplingFrameRate().AsSeconds(FFrameTime(int32(Frame)));
 
-			FCompactPose AnimPose;
-			AnimPose.SetBoneContainer(&BoneContainer);
-			FBlendedCurve AnimCurve;
-			AnimCurve.InitFrom(BoneContainer);
-			UE::Anim::FStackAttributeContainer AnimAttributes;
+				FCompactPose AnimPose;
+				AnimPose.SetBoneContainer(&BoneContainer);
+				FBlendedCurve AnimCurve;
+				AnimCurve.InitFrom(BoneContainer);
+				UE::Anim::FStackAttributeContainer AnimAttributes;
 
-			FAnimExtractContext ExtractionContext(SeekTime);
-			FAnimationPoseData AnimationPoseData(AnimPose, AnimCurve, AnimAttributes);
-			Sequence->GetAnimationPose(AnimationPoseData, ExtractionContext);
+				FAnimExtractContext ExtractionContext(SeekTime);
+				FAnimationPoseData AnimationPoseData(AnimPose, AnimCurve, AnimAttributes);
+				Sequence->GetAnimationPose(AnimationPoseData, ExtractionContext);
 
-			for (int32 MeshBoneIndex = 0; MeshBoneIndex < NumMeshBones; ++MeshBoneIndex)
-			{
-				const FCompactPoseBoneIndex CompactIndex = BoneContainer.MakeCompactPoseIndex(FMeshPoseBoneIndex(MeshBoneIndex));
-				const int32 ParentIndex = MeshRefSkeleton.GetParentIndex(MeshBoneIndex);
-				if (SampledGlobalMeshPose.IsValidIndex(ParentIndex) && ParentIndex < MeshBoneIndex)
+				for (int32 MeshBoneIndex = 0; MeshBoneIndex < NumMeshBones; ++MeshBoneIndex)
 				{
-					SampledGlobalMeshPose[MeshBoneIndex] = AnimationPoseData.GetPose()[CompactIndex] * SampledGlobalMeshPose[ParentIndex];
+					const FCompactPoseBoneIndex CompactIndex = BoneContainer.MakeCompactPoseIndex(FMeshPoseBoneIndex(MeshBoneIndex));
+					const int32 ParentIndex = MeshRefSkeleton.GetParentIndex(MeshBoneIndex);
+					if (SampledGlobalMeshPose.IsValidIndex(ParentIndex) && ParentIndex < MeshBoneIndex)
+					{
+						SampledGlobalMeshPose[MeshBoneIndex] = AnimationPoseData.GetPose()[CompactIndex] * SampledGlobalMeshPose[ParentIndex];
+					}
+					else
+					{
+						SampledGlobalMeshPose[MeshBoneIndex] = MeshLocalRefPose[MeshBoneIndex];
+					}
 				}
-				else
+
+				for (int32 MeshBoneIndex = 0; MeshBoneIndex < NumMeshBones; ++MeshBoneIndex)
 				{
-					SampledGlobalMeshPose[MeshBoneIndex] = MeshLocalRefPose[MeshBoneIndex];
-				}
-			}
+					BankEntry.PositionKeys[KeyIndex] = FVector3f(SampledGlobalMeshPose[MeshBoneIndex].GetTranslation());
+					BankEntry.RotationKeys[KeyIndex] = FQuat4f(SampledGlobalMeshPose[MeshBoneIndex].GetRotation());
 
-			for (int32 MeshBoneIndex = 0; MeshBoneIndex < NumMeshBones; ++MeshBoneIndex)
-			{
-				BankEntry.PositionKeys[KeyIndex] = FVector3f(SampledGlobalMeshPose[MeshBoneIndex].GetTranslation());
-				BankEntry.RotationKeys[KeyIndex] = FQuat4f(SampledGlobalMeshPose[MeshBoneIndex].GetRotation());
+					// Expand animated bounds
+					{
+						const VectorRegister4Float VecBonePosition = VectorLoadFloat3(&BankEntry.PositionKeys[KeyIndex].X);
+						VecAnimatedBoundsMin = VectorMin(VecAnimatedBoundsMin, VecBonePosition);
+						VecAnimatedBoundsMax = VectorMax(VecAnimatedBoundsMax, VecBonePosition);
+					}
 
-				// Expand animated bounds
-				{
-					const VectorRegister4Float VecBonePosition = VectorLoadFloat3(&BankEntry.PositionKeys[KeyIndex].X);
-					VecAnimatedBoundsMin = VectorMin(VecAnimatedBoundsMin, VecBonePosition);
-					VecAnimatedBoundsMax = VectorMax(VecAnimatedBoundsMax, VecBonePosition);
+					++KeyIndex;
 				}
-
-				++KeyIndex;
 			}
 		}
 
@@ -662,6 +666,8 @@ void FAnimBankBuildAsyncCacheTask::Reschedule(FQueuedThreadPool* InThreadPool, E
 
 #endif
 
+///////////////////////////////////////////////////////////////////////////////////////////////////
+
 void UAnimBank::Serialize(FArchive& Ar)
 {
 	LLM_SCOPE(ELLMTag::Animation);
@@ -1152,6 +1158,247 @@ void UAnimBank::NotifyOnGPUDataChanged()
 
 #endif
 
+float UAnimBank::GetSequencePlayLength(int32 SequenceIndex) const
+{
+#if WITH_EDITORONLY_DATA
+	if (IsCompiling())
+	{
+		if (Sequences.IsValidIndex(SequenceIndex))
+		{
+			if (Sequences[SequenceIndex].Sequence)
+			{
+				return Sequences[SequenceIndex].Sequence->GetPlayLength();
+			}
+		}
+	}
+	else
+#endif
+	if (Data.Entries.IsValidIndex(SequenceIndex))
+	{
+		const float SampleInterval = 1.0f / ANIM_BANK_SAMPLE_RATE;
+		return static_cast<float>(Data.Entries[SequenceIndex].FrameCount - 1) * SampleInterval;
+	}
+
+	return 0.0f;
+}
+
+bool UAnimBank::IsSequenceLooping(int32 SequenceIndex) const
+{
+#if WITH_EDITORONLY_DATA
+	if (IsCompiling())
+	{
+		if (Sequences.IsValidIndex(SequenceIndex))
+		{
+			if (Sequences[SequenceIndex].Sequence)
+			{
+				return Sequences[SequenceIndex].bLooping;
+			}
+		}
+	}
+	else
+#endif
+	if (Data.Entries.IsValidIndex(SequenceIndex))
+	{
+		return Data.Entries[SequenceIndex].IsLooping();
+	}
+
+	return false;
+}
+///////////////////////////////////////////////////////////////////////////////////////////////////
+
+inline float GetLoopedPosition(float Position, float TrackLength)
+{
+	if (TrackLength >= 0.0f)
+	{
+		Position = TrackLength == 0.0f ? 0.0f : FMath::Fmod(Position, TrackLength);
+
+		// Returns a negative range if input is negative.
+		if (Position < 0.0f)
+		{
+			Position += TrackLength;
+		}
+	}
+	return Position;
+}
+
+void FAnimBankTrackPackedData::Update(float SequenceLength, bool bLooping, float DeltaTime)
+{
+	if (bAutoPlay)
+	{
+		PreviousPosition = Position;
+		Position += PlayRate * DeltaTime;
+		bHasPreviousPosition = false;
+		Wrap(SequenceLength, bLooping);
+	}
+}
+
+void FAnimBankTrackPackedData::Wrap(float SequenceLength, bool bLooping)
+{
+	const EAnimBankTrackLoopMode LoopModeEnum = GetLoopMode();
+
+	if (LoopModeEnum == EAnimBankTrackLoopMode::Loop || (LoopModeEnum == EAnimBankTrackLoopMode::Auto && bLooping))
+	{
+		Position         = GetLoopedPosition(Position, SequenceLength);
+		PreviousPosition = GetLoopedPosition(PreviousPosition, SequenceLength);
+	}
+	else
+	{
+		Position         = FMath::Clamp(Position, 0.0f, SequenceLength);
+		PreviousPosition = FMath::Clamp(PreviousPosition, 0.0f, SequenceLength);
+	}
+}
+
+void FAnimBankTrackPackedData::Serialize(FArchive& Ar)
+{
+	constexpr auto MakeBitMask = [] (uint32 Index) -> uint32 { return ~((1 << Index) - 1); };
+
+	constexpr uint32 ItemIndexBits = 20;
+	constexpr uint32 ItemIndexMask = MakeBitMask(ItemIndexBits);
+
+	constexpr uint32 LoopModeShift = ItemIndexBits;
+	constexpr uint32 LoopModeBits  = 2;
+	constexpr uint32 LoopModeMask  = MakeBitMask(LoopModeBits);
+
+	constexpr uint32 AutoPlayShift = LoopModeShift + LoopModeBits;
+	constexpr uint32 AutoPlayBits  = 1;
+	constexpr uint32 AutoPlayMask  = MakeBitMask(AutoPlayBits);
+	
+	constexpr uint32 HasPreviousPositionShift = AutoPlayShift + AutoPlayBits;
+	constexpr uint32 HasPreviousPositionBits  = 1;
+	constexpr uint32 HasPreviousPositionMask  = MakeBitMask(HasPreviousPositionBits);
+
+	uint32 PackedBits;
+	if (Ar.IsSaving())
+	{
+		PackedBits =
+			   (ItemIndex & ItemIndexMask)
+			| ((LoopMode  & LoopModeMask ) << LoopModeShift)
+			| ((bAutoPlay & AutoPlayMask ) << AutoPlayShift)
+			| ((bHasPreviousPosition & HasPreviousPositionMask) << HasPreviousPositionShift);
+		Ar << PackedBits;
+		Ar << Position;
+
+		if (bAutoPlay)
+		{
+			Ar << PlayRate;
+		}
+		else
+		{
+			Ar << PreviousPosition;
+		}
+	}
+	else
+	{
+		Ar << PackedBits;
+		ItemIndex =  PackedBits & ItemIndexMask;
+		LoopMode  = (PackedBits >> LoopModeShift) & LoopModeMask;
+		bAutoPlay = (PackedBits >> AutoPlayShift) & AutoPlayMask;
+		bHasPreviousPosition = (PackedBits >> HasPreviousPositionShift) & HasPreviousPositionMask;
+		Ar << Position;
+
+		if (bAutoPlay)
+		{
+			Ar << PlayRate;
+		}
+		else
+		{
+			Ar << PreviousPosition;
+		}
+	}
+}
+
+///////////////////////////////////////////////////////////////////////////////////////////////////
+
+FAnimBankRenderData::~FAnimBankRenderData()
+{
+	check(Proxies.Num() == 0);
+}
+
+void FAnimBankRenderData::InsertProxy(FInstancedSkinningSceneExtensionProxy* Proxy)
+{
+	Proxies.FindOrAdd(Proxy);
+}
+
+void FAnimBankRenderData::RemoveProxy(FInstancedSkinningSceneExtensionProxy* Proxy)
+{
+	Proxies.Remove(Proxy);
+}
+
+void FAnimBankRenderData::Patch(FAnimBankTrackPool::FPatch&& Patch)
+{
+	if (Patch.GetNumTracks() != Tracks.GetNumTracks())
+	{
+		for (FInstancedSkinningSceneExtensionProxy* Proxy : Proxies)
+		{
+			Proxy->SetUniqueAnimationCount(Patch.GetNumTracks());
+		}
+	}
+	Tracks.Patch(MoveTemp(Patch));
+}
+
+///////////////////////////////////////////////////////////////////////////////////////////////////
+
+FAnimBankTransformProviderProxy::FAnimBankTransformProviderProxy(TConstArrayView<FAnimBankItem> InAnimBankItems, FInstancedSkinningSceneExtensionProxy* InSceneExtensionProxy)
+	: AnimBankItems(InAnimBankItems)
+	, TrackAllocator(new FAnimBankTrackPool)
+	, Tracks(TrackAllocator.Get())
+	, SceneExtensionProxy(InSceneExtensionProxy)
+{
+	TrackAllocator->Reserve(AnimBankItems.Num());
+
+	for (int32 ItemIndex = 0; ItemIndex < AnimBankItems.Num(); ++ItemIndex)
+	{
+		const FAnimBankItem& Item = AnimBankItems[ItemIndex];
+		const FAnimBankData& BankData = Item.BankAsset->GetData();
+		const FAnimBankEntry& BankEntry = BankData.Entries[Item.SequenceIndex];
+
+		FAnimBankTrackAutoPlayData Data;
+		Data.ItemIndex            = ItemIndex;
+		Data.Position             = BankEntry.Position;
+		Data.PlayRate             = BankEntry.PlayRate;
+
+		TrackAllocator->AllocateTrack(FAnimBankTrackPackedData::Pack(Data));
+	}
+}
+
+FAnimBankTransformProviderProxy::FAnimBankTransformProviderProxy(TConstArrayView<FAnimBankItem> InAnimBankItems, FInstancedSkinningSceneExtensionProxy* InSceneExtensionProxy, FAnimBankRenderData& InRenderData)
+	: AnimBankItems(InAnimBankItems)
+	, RenderData(&InRenderData)
+	, Tracks(&InRenderData.Tracks)
+	, SceneExtensionProxy(InSceneExtensionProxy)
+{}
+
+void FAnimBankTransformProviderProxy::CreateRenderThreadResources(FSceneInterface& Scene, FRHICommandListBase& RHICmdList)
+{
+	if (RenderData)
+	{
+		RenderData->InsertProxy(SceneExtensionProxy);
+	}
+
+	AnimBankTransformProvider = Scene.GetAnimBankTransformProvider();
+
+	if (AnimBankTransformProvider)
+	{
+		AnimBankTransformProvider->RegisterAnimBank(this);
+	}
+}
+
+void FAnimBankTransformProviderProxy::DestroyRenderThreadResources()
+{
+	if (RenderData)
+	{
+		RenderData->RemoveProxy(SceneExtensionProxy);
+	}
+
+	if (AnimBankTransformProvider)
+	{
+		AnimBankTransformProvider->UnregisterAnimBank(this);
+		AnimBankTransformProvider = nullptr;
+	}
+}
+
+///////////////////////////////////////////////////////////////////////////////////////////////////
+
 bool UAnimBankData::IsEnabled() const
 {
 	for (const FAnimBankItem& BankItem : AnimBankItems)
@@ -1161,12 +1408,12 @@ bool UAnimBankData::IsEnabled() const
 			return false;
 		}
 
-	#if WITH_EDITORONLY_DATA
 		if (BankItem.BankAsset->Asset == nullptr)
 		{
 			return false;
 		}
 
+	#if WITH_EDITORONLY_DATA
 		if (BankItem.SequenceIndex >= BankItem.BankAsset->Sequences.Num())
 		{
 			return false;
@@ -1193,14 +1440,14 @@ const FGuid& UAnimBankData::GetTransformProviderID() const
 	}
 }
 
-bool UAnimBankData::UsesSkeletonBatching() const
+uint32 UAnimBankData::GetUniqueAnimationCount() const
 {
-	return false;
+	return AnimBankItems.Num();
 }
 
-const uint32 UAnimBankData::GetUniqueAnimationCount() const
+bool UAnimBankData::UsesSkeletonBatching() const
 {
-	return AnimBankItems.Num();
+	return false;
 }
 
 bool UAnimBankData::HasAnimationBounds() const
@@ -1230,7 +1477,7 @@ bool UAnimBankData::GetAnimationBounds(uint32 AnimationIndex, FRenderBounds& Out
 uint32 UAnimBankData::GetSkinningDataOffset(int32 InstanceIndex, const FTransform& ComponentTransform, const FSkinnedMeshInstanceData& InstanceData) const
 {
 	const uint32 AnimationIndex = InstanceData.AnimationIndex;
-	if (ensure(AnimationIndex < uint32(AnimBankItems.Num())))
+	if (AnimationIndex < uint32(AnimBankItems.Num()))
 	{
 		return AnimationIndex * 2u;
 	}
@@ -1238,17 +1485,33 @@ uint32 UAnimBankData::GetSkinningDataOffset(int32 InstanceIndex, const FTransfor
 	return 0u;
 }
 
-FTransformProviderRenderProxy* UAnimBankData::CreateRenderThreadResources(FSkinningSceneExtensionProxy* SceneProxy, FSceneInterface& Scene, FRHICommandListBase& RHICmdList)
+bool UAnimBankData::IsAnimBankValid(FInstancedSkinningSceneExtensionProxy* ExtensionProxy) const
 {
-	FAnimBankDataRenderProxy* ProviderProxy = new FAnimBankDataRenderProxy(this, SceneProxy, Scene);
-	ProviderProxy->CreateRenderThreadResources(RHICmdList);
-	return ProviderProxy;
+	for (const FAnimBankItem& Item : AnimBankItems)
+	{
+		if (Item.BankAsset == nullptr || Item.BankAsset->Asset != ExtensionProxy->GetSkinnedAsset())
+		{
+			return false;
+		}
+
+		const FAnimBankData& BankData = Item.BankAsset->GetData();
+
+		if (Item.SequenceIndex >= BankData.Entries.Num())
+		{
+			return false;
+		}
+	}
+
+	return true;
 }
 
-void UAnimBankData::DestroyRenderThreadResources(FTransformProviderRenderProxy* ProviderProxy)
+FTransformProviderRenderProxy* UAnimBankData::CreateRenderProxy(FInstancedSkinningSceneExtensionProxy* SceneExtensionProxy) const
 {
-	ProviderProxy->DestroyRenderThreadResources();
-	delete ProviderProxy;
+	if (IsAnimBankValid(SceneExtensionProxy))
+	{
+		return new FAnimBankTransformProviderProxy(AnimBankItems, SceneExtensionProxy);
+	}
+	return nullptr;
 }
 
 bool UAnimBankData::IsCompiling() const
@@ -1284,125 +1547,210 @@ bool UAnimBankData::IsCachedCookedPlatformDataLoaded(const ITargetPlatform* Targ
 
 #endif
 
-FAnimBankDataRenderProxy::FAnimBankDataRenderProxy(UAnimBankData* BankData, FSkinningSceneExtensionProxy* InSceneProxy, FSceneInterface& InScene)
-: SceneProxy(InSceneProxy)
-, Scene(InScene)
+float UAnimBankData::GetSequencePlayLength(int32 ItemIndex) const
 {
-	UniqueAnimationCount = SceneProxy->GetUniqueAnimationCount();
-	SkinnedAsset = SceneProxy->GetSkinnedAsset();
-	AnimBankItems = BankData->AnimBankItems;
-	check(SkinnedAsset != nullptr);
+	if (AnimBankItems.IsValidIndex(ItemIndex))
+	{
+		const FAnimBankItem& AnimBankItem = AnimBankItems[ItemIndex];
+
+		if (const UAnimBank* BankAsset = AnimBankItem.BankAsset)
+		{
+			return BankAsset->GetSequencePlayLength(AnimBankItem.SequenceIndex);
+		}
+	}
+	return 0.0f;
+}
+
+///////////////////////////////////////////////////////////////////////////////////////////////////
+
+UAnimBankDataInstance* UAnimBankDataInstance::Create(UAnimBankData* AnimBank, UInstancedSkinnedMeshComponent* Owner)
+{
+	if (IsValid(Owner))
+	{
+		UAnimBankDataInstance* AnimBankInstance = NewObject<UAnimBankDataInstance>(Owner);
+		AnimBankInstance->AnimBankItems = AnimBank->AnimBankItems;
+		AnimBankInstance->bEnabled = true;
+		AnimBankInstance->InitSequences();
+		return AnimBankInstance;
+	}
+
+	return nullptr;
 }
 
-FAnimBankDataRenderProxy::~FAnimBankDataRenderProxy()
+bool UAnimBankDataInstance::IsEnabled() const
 {
+	return Tracks.GetNumTracks() != 0 && Super::IsEnabled();
 }
 
-void FAnimBankDataRenderProxy::CreateRenderThreadResources(FRHICommandListBase& RHICmdList)
+uint32 UAnimBankDataInstance::GetUniqueAnimationCount() const
 {
-	TArray<FAnimBankDesc> Descs;
-	Descs.Reserve(AnimBankItems.Num());
+	return Tracks.GetNumTracks();
+}
 
-	for (const FAnimBankItem& Item : AnimBankItems)
+bool UAnimBankDataInstance::GetAnimationBounds(uint32 AnimationIndex, FRenderBounds& OutBounds) const
+{
+	if (Tracks.IsActiveIndex(AnimationIndex))
 	{
-		if (Item.BankAsset == nullptr || SkinnedAsset == nullptr)
-		{
-			Descs.Emplace(FAnimBankDesc());
-			continue;
-		}
+		return Super::GetAnimationBounds(Tracks.GetData(AnimationIndex).GetItemIndex(), OutBounds);
+	}
 
-		const FAnimBankData& BankData = Item.BankAsset->GetData();
+	return false;
+}
 
-		if (Item.SequenceIndex >= BankData.Entries.Num())
+uint32 UAnimBankDataInstance::GetSkinningDataOffset(int32 InstanceIndex, const FTransform& ComponentTransform, const FSkinnedMeshInstanceData& InstanceData) const
+{
+	const uint32 AnimationIndex = InstanceData.AnimationIndex;
+	if (AnimationIndex < uint32(Tracks.GetNumTracks()))
+	{
+		return AnimationIndex * 2u;
+	}
+
+	return 0u;
+}
+
+FTransformProviderRenderProxy* UAnimBankDataInstance::CreateRenderProxy(FInstancedSkinningSceneExtensionProxy* SceneExtensionProxy) const
+{
+	if (IsAnimBankValid(SceneExtensionProxy))
+	{
+		return new FAnimBankTransformProviderProxy(AnimBankItems, SceneExtensionProxy, *RenderData);
+	}
+	return nullptr;
+}
+
+void UAnimBankDataInstance::BeginDestroy()
+{
+	ENQUEUE_RENDER_COMMAND(AnimBankDataInstanceRelease)([RenderData = MoveTemp(RenderData)](FRHICommandList&) mutable
+	{
+		// Move destruction onto the render thread.
+	});
+
+	Super::BeginDestroy();
+}
+
+void UAnimBankDataInstance::SubmitChanges()
+{
+	check(RenderData);
+
+	if (Tracks.GetNumDirtyTracks() > 0)
+	{
+		ENQUEUE_RENDER_COMMAND(AnimBankDataInstanceSubmit)([&RenderData = *RenderData, Patch = Tracks.Finalize()] (FRHICommandList&) mutable
 		{
-			Descs.Emplace(FAnimBankDesc());
-			continue;
-		}
+			RenderData.Patch(MoveTemp(Patch));
+		});
+	}
+}
 
-		const FAnimBankEntry& BankEntry = BankData.Entries[Item.SequenceIndex];
+void UAnimBankDataInstance::InitRenderData()
+{
+	check(RenderData);
+	RenderData->Tracks = Tracks;
+}
 
-		FAnimBankDesc& Desc = Descs.Emplace_GetRef();
+void UAnimBankDataInstance::InitSequences()
+{
+	Sequences.Reserve(AnimBankItems.Num());
 
-		Desc.BankAsset = MakeWeakObjectPtr(Item.BankAsset);
-		Desc.SequenceIndex = Item.SequenceIndex;
-		Desc.Asset = MakeWeakObjectPtr(SkinnedAsset);
-		Desc.Position = BankEntry.Position;
-		Desc.PlayRate = BankEntry.PlayRate;
-		Desc.bLooping = BankEntry.IsLooping() ? 1 : 0;
-		Desc.bAutoStart = BankEntry.IsAutoStart() ? 1 : 0;
+	for (const FAnimBankItem& Item : AnimBankItems)
+	{
+		if (Item.BankAsset)
+		{
+			const bool bLooping = Item.BankAsset->IsSequenceLooping(Item.SequenceIndex);
+			const float PlayLength = Item.BankAsset->GetSequencePlayLength(Item.SequenceIndex);
+			Sequences.Emplace(PlayLength, bLooping);
+		}
 	}
+}
+
+void UAnimBankDataInstance::Serialize(FArchive& Ar)
+{
+	Ar.UsingCustomVersion(FAnimObjectVersion::GUID);
+	Super::Serialize(Ar);
+	Tracks.Serialize(Ar);
 
-	AnimBankHandles = Scene.RegisterAnimBank(Descs);
-	AnimBankIds.Reset(AnimBankHandles.Num());
-	for (const FAnimBankRecordHandle& Handle : AnimBankHandles)
+	if (Ar.IsLoading())
 	{
-		AnimBankIds.Emplace((uint64)Handle.Id);
+		InitRenderData();
+		InitSequences();
 	}
 }
 
-void FAnimBankDataRenderProxy::DestroyRenderThreadResources()
+int32 UAnimBankDataInstance::AllocateAutoPlayTrack(const FAnimBankTrackAutoPlayData& Data)
 {
-	Scene.UnregisterAnimBank(AnimBankHandles);
-	AnimBankHandles.Reset();
-	AnimBankIds.Reset();
+	FAnimBankTrackPackedData PackedData = FAnimBankTrackPackedData::Pack(Data);
+
+	if (Sequences.IsValidIndex(Data.ItemIndex))
+	{
+		const FSequence& Sequence = Sequences[Data.ItemIndex];
+		PackedData.Wrap(Sequence.PlayLength, Sequence.bLooping);
+		return Tracks.AllocateTrack(PackedData);
+	}
+
+	return INDEX_NONE;
 }
 
-const TConstArrayView<uint64> FAnimBankDataRenderProxy::GetProviderData(bool& bOutValid) const
+int32 UAnimBankDataInstance::AllocateManualTrack(const FAnimBankTrackManualData& Data)
 {
-	bOutValid = (AnimBankIds.Num() == UniqueAnimationCount);
-	return AnimBankIds;
+	FAnimBankTrackPackedData PackedData = FAnimBankTrackPackedData::Pack(Data);
+
+	if (Sequences.IsValidIndex(Data.ItemIndex))
+	{
+		const FSequence& Sequence = Sequences[Data.ItemIndex];
+		PackedData.Wrap(Sequence.PlayLength, Sequence.bLooping);
+		return Tracks.AllocateTrack(PackedData);
+	}
+
+	return INDEX_NONE;
 }
 
-void FAnimBankSequence::ValidatePosition()
+bool UAnimBankDataInstance::UpdateAutoPlayTrack(int32 TrackIndex, const FAnimBankTrackAutoPlayData& Data)
 {
-	float Min = 0.0f, Max = 0.0f;
+	FAnimBankTrackPackedData PackedData = FAnimBankTrackPackedData::Pack(Data);
 
-	if (Sequence)
+	if (Sequences.IsValidIndex(Data.ItemIndex))
 	{
-		Max = Sequence->GetPlayLength();
+		const FSequence& Sequence = Sequences[Data.ItemIndex];
+		PackedData.Wrap(Sequence.PlayLength, Sequence.bLooping);
+		return Tracks.UpdateTrack(TrackIndex, PackedData);
 	}
 
-	Position = FMath::Clamp<float>(Position, Min, Max);
+	return false;
 }
 
-uint32 FAnimBankDesc::GetHash() const
+bool UAnimBankDataInstance::UpdateManualTrack(int32 TrackIndex, const FAnimBankTrackManualData& Data)
 {
-	struct FHashKey
+	FAnimBankTrackPackedData PackedData = FAnimBankTrackPackedData::Pack(Data);
+
+	if (Sequences.IsValidIndex(Data.ItemIndex))
 	{
-		uint32	BankHash;
-		uint32	SequenceIndex;
-		uint32	AssetHash;
-		float	Position;
-		float	PlayRate;
-		uint8	Flags;
+		const FSequence& Sequence = Sequences[Data.ItemIndex];
+		PackedData.Wrap(Sequence.PlayLength, Sequence.bLooping);
+		return Tracks.UpdateTrack(TrackIndex, PackedData);
+	}
 
-		static inline uint32 PointerHash(const void* Key)
-		{
-			// Ignoring the lower 4 bits since they are likely zero anyway.
-			// Higher bits are more significant in 64 bit builds.
-			return reinterpret_cast<UPTRINT>(Key) >> 4;
-		};
+	return false;
+}
 
-	} HashKey;
+bool UAnimBankDataInstance::DeallocateTrack(int32 TrackIndex)
+{
+	return Tracks.DeallocateTrack(TrackIndex);
+}
 
-	FMemory::Memzero(HashKey);
+///////////////////////////////////////////////////////////////////////////////////////////////////
 
-	const UAnimBank* BankPtr = BankAsset.Get();
-	const USkinnedAsset* AssetPtr = Asset.Get();
+void FAnimBankSequence::ValidatePosition()
+{
+	float Min = 0.0f, Max = 0.0f;
 
-	HashKey.BankHash		= FHashKey::PointerHash(BankPtr);
-	HashKey.SequenceIndex	= SequenceIndex;
-	HashKey.AssetHash		= FHashKey::PointerHash(AssetPtr);
-	HashKey.Flags			= 0;
-	HashKey.Flags		   |= bLooping   ? 0x1u : 0x0u;
-	HashKey.Flags		   |= bAutoStart ? 0x2u : 0x0u;
-	HashKey.Position		= Position;
-	HashKey.PlayRate		= PlayRate;
+	if (Sequence)
+	{
+		Max = Sequence->GetPlayLength();
+	}
 
-	const uint64 DescHash = CityHash64((char*)&HashKey, sizeof(FHashKey));
-	return HashCombineFast(uint32(DescHash & 0xFFFFFFFF), uint32((DescHash >> 32) & 0xFFFFFFFF));
+	Position = FMath::Clamp<float>(Position, Min, Max);
 }
 
+///////////////////////////////////////////////////////////////////////////////////////////////////
+
 namespace UE::AnimBank
 {
 
diff --git a/Engine/Source/Runtime/Engine/Private/GPUSkinVertexFactory.cpp b/Engine/Source/Runtime/Engine/Private/GPUSkinVertexFactory.cpp
--- a/Engine/Source/Runtime/Engine/Private/GPUSkinVertexFactory.cpp
+++ b/Engine/Source/Runtime/Engine/Private/GPUSkinVertexFactory.cpp
@@ -547,6 +547,7 @@ BEGIN_SHADER_PARAMETER_STRUCT(FGPUSkinVertexFactoryCommonShaderParameters, )
 	// number of influences for this draw call, 4 or 8
 	SHADER_PARAMETER(uint32, NumBoneInfluencesParam)
 	SHADER_PARAMETER(uint32, bIsMorphTarget)
+	SHADER_PARAMETER(uint32, bIsRefPose)
 	SHADER_PARAMETER(uint32, TransformStorageMode)
 	SHADER_PARAMETER(uint32, SectionBoneOffset)
 	SHADER_PARAMETER(uint32, TransformBufferOffset)
@@ -585,6 +586,7 @@ void GetGPUSkinVertexFactoryCommonShaderParameters(
 		ShaderParameters.MaxTransformCount = 0;
 		ShaderParameters.CurrentTransformSlot = 0;
 		ShaderParameters.TransformStorageMode = GPUSKIN_TRANSFORM_STORAGE_BONE_MAP;
+		ShaderParameters.bIsRefPose = false;
 	}
 	else
 	{
@@ -605,6 +607,7 @@ void GetGPUSkinVertexFactoryCommonShaderParameters(
 		ShaderParameters.CurrentTransformSlot = SkinningHeader->CurrentTransformSlot;
 		ShaderParameters.SectionBoneOffset = VertexFactory->GetBoneOffset();
 		ShaderParameters.TransformStorageMode = ShaderData.SectionBoneMapBuffer.IsValid() ? GPUSKIN_TRANSFORM_STORAGE_DIRECT : GPUSKIN_TRANSFORM_STORAGE_BONE_MAP;
+		ShaderParameters.bIsRefPose = SkinningHeader->bIsRefPose;
 	}
 	ShaderParameters.BoneUpdatedFrameNumber = ShaderData.UpdatedFrameNumber;
 	ShaderParameters.SectionBoneMap = ShaderData.SectionBoneMapBuffer ? ShaderData.SectionBoneMapBuffer : GNullVertexBuffer.VertexBufferSRV;
diff --git a/Engine/Source/Runtime/Engine/Private/SkinningSceneExtensionProxy.cpp b/Engine/Source/Runtime/Engine/Private/SkinningSceneExtensionProxy.cpp
--- a/Engine/Source/Runtime/Engine/Private/SkinningSceneExtensionProxy.cpp
+++ b/Engine/Source/Runtime/Engine/Private/SkinningSceneExtensionProxy.cpp
@@ -105,11 +105,17 @@ void FSkinningSceneExtensionProxy::InitForBoneMapStorage()
 	check(BoneTransformStorageMode == EBoneTransformStorageMode::BoneMap);
 	TConstArrayView<FSkelMeshRenderSection> Sections = MeshObject->GetRenderSections(LODIndex);
 
+	MaxBoneTransformIndex = 0;
 	MaxBoneTransformCount = 0;
 	for (const FSkelMeshRenderSection& Section : Sections)
 	{
 		if (Section.IsValid())
 		{
+			for (uint32 Index : Section.BoneMap)
+			{
+				MaxBoneTransformIndex = FMath::Max<uint32>(MaxBoneTransformIndex, Index);
+			}
+
 			MaxBoneTransformCount += Section.BoneMap.Num();
 		}
 	}
@@ -126,32 +132,34 @@ void FSkinningSceneExtensionProxy::InitForBoneMapStorage()
 	check(MaxBoneTransformCount > 0);
 }
 
-TConstArrayView<uint64> FSkinningSceneExtensionProxy::GetAnimationProviderData(bool& bOutValid) const
-{
-	bOutValid = false;
-	return {};
-}
-
 const FGuid& FSkinningSceneExtensionProxy::GetTransformProviderId() const
 {
 	static FGuid AnimRuntimeId(ANIM_RUNTIME_TRANSFORM_PROVIDER_GUID);
 	return AnimRuntimeId;
 }
 
-FInstancedSkinningSceneExtensionProxy::FInstancedSkinningSceneExtensionProxy(TObjectPtr<UTransformProviderData> InTransformProvider, const FSkeletalMeshObject* InMeshObject, const USkinnedAsset* InSkinnedAsset, bool bAllowScaling)
+FInstancedSkinningSceneExtensionProxy::FInstancedSkinningSceneExtensionProxy(const UTransformProviderData* TransformProvider, const FSkeletalMeshObject* InMeshObject, const USkinnedAsset* InSkinnedAsset, bool bAllowScaling)
 	: FSkinningSceneExtensionProxy(InMeshObject, InSkinnedAsset, bAllowScaling)
-	, TransformProvider(InTransformProvider)
 {
 	bUseInstancing = true;
 
+	bool bSucceeded = false;
+
 	const bool bForceRefPose = CVarInstancedSkinnedMeshesForceRefPose.GetValueOnAnyThread() != 0;
 	if (!bForceRefPose && GetSkinnedAsset() && TransformProvider != nullptr && TransformProvider->IsEnabled())
 	{
-		TransformProviderId  = TransformProvider->GetTransformProviderID();
-		UniqueAnimationCount = TransformProvider->GetUniqueAnimationCount();
-		bUseSkeletonBatching = TransformProvider->UsesSkeletonBatching();
+		TransformProviderProxy = TransformProvider->CreateRenderProxy(this);
+
+		if (TransformProviderProxy)
+		{
+			TransformProviderId    = TransformProvider->GetTransformProviderID();
+			bUseSkeletonBatching   = TransformProvider->UsesSkeletonBatching();
+			SetUniqueAnimationCount(TransformProvider->GetUniqueAnimationCount());
+			bSucceeded             = true;
+		}
 	}
-	else
+
+	if (!bSucceeded)
 	{
 		UniqueAnimationCount = 1;
 		bUseSkeletonBatching = false;
@@ -163,50 +171,29 @@ FInstancedSkinningSceneExtensionProxy::FInstancedSkinningSceneExtensionProxy(TOb
 
 void FInstancedSkinningSceneExtensionProxy::CreateRenderThreadResources(FSceneInterface& Scene, FRHICommandListBase& RHICmdList)
 {
-	if (TransformProvider != nullptr && TransformProvider->GetTransformProviderID() == TransformProviderId)
-	{
-		TransformProviderProxy = TransformProvider->CreateRenderThreadResources(this, Scene, RHICmdList);
-	}
-	else
+	if (TransformProviderProxy)
 	{
-		TransformProviderProxy = nullptr;
+		TransformProviderProxy->CreateRenderThreadResources(Scene, RHICmdList);
 	}
 }
 
 void FInstancedSkinningSceneExtensionProxy::DestroyRenderThreadResources()
 {
-	if (TransformProviderProxy != nullptr)
+	if (TransformProviderProxy)
 	{
-		TransformProvider->DestroyRenderThreadResources(TransformProviderProxy);
+		TransformProviderProxy->DestroyRenderThreadResources();
+		delete TransformProviderProxy;
 		TransformProviderProxy = nullptr;
 	}
 }
 
-TConstArrayView<uint64> FInstancedSkinningSceneExtensionProxy::GetAnimationProviderData(bool& bOutValid) const
-{
-	if (TransformProviderProxy != nullptr)
-	{
-		return TransformProviderProxy->GetProviderData(bOutValid);
-	}
-
-	bOutValid = false;
-	return {};
-}
-
 const FGuid& FInstancedSkinningSceneExtensionProxy::GetTransformProviderId() const
 {
-	// If the proxy is current in an invalid state, use the
-	// reference pose transform provider
-	if (TransformProviderId.IsValid())
+	if (!TransformProviderProxy)
 	{
-		bool bIsValid = false;
-		GetAnimationProviderData(bIsValid);
-		if (!bIsValid)
-		{
-			static FGuid RefPoseProviderId(REF_POSE_TRANSFORM_PROVIDER_GUID);
-			return RefPoseProviderId;
-		}
+		static FGuid RefPoseProviderId(REF_POSE_TRANSFORM_PROVIDER_GUID);
+		return RefPoseProviderId;
 	}
 
 	return TransformProviderId;
-}
\ No newline at end of file
+}
diff --git a/Engine/Source/Runtime/Engine/Public/Animation/TransformProviderData.h b/Engine/Source/Runtime/Engine/Public/Animation/TransformProviderData.h
--- a/Engine/Source/Runtime/Engine/Public/Animation/TransformProviderData.h
+++ b/Engine/Source/Runtime/Engine/Public/Animation/TransformProviderData.h
@@ -10,7 +10,7 @@
 
 struct FRenderBounds;
 struct FSkinnedMeshInstanceData;
-class FSkinningSceneExtensionProxy;
+class FInstancedSkinningSceneExtensionProxy;
 class FSceneInterface;
 class FRHICommandListBase;
 class ITargetPlatform;
@@ -18,17 +18,11 @@ class ITargetPlatform;
 class FTransformProviderRenderProxy
 {
 public:
-	FTransformProviderRenderProxy()
-	{
-	}
-
 	virtual ~FTransformProviderRenderProxy() = default;
 
-	virtual void CreateRenderThreadResources(FRHICommandListBase& RHICmdList) = 0;
+	virtual void CreateRenderThreadResources(FSceneInterface& Scene, FRHICommandListBase& RHICmdList) = 0;
 
 	virtual void DestroyRenderThreadResources() = 0;
-
-	virtual const TConstArrayView<uint64> GetProviderData(bool& bOutValid) const = 0;
 };
 
 UCLASS(Abstract, config=Engine, hidecategories=Object, MinimalAPI, BlueprintType)
@@ -48,7 +42,7 @@ public:
 		return InvalidID;
 	}
 
-	virtual const uint32 GetUniqueAnimationCount() const
+	virtual uint32 GetUniqueAnimationCount() const
 	{
 		return 1u;
 	}
@@ -73,18 +67,14 @@ public:
 		return 0u;
 	}
 
-	virtual FTransformProviderRenderProxy* CreateRenderThreadResources(FSkinningSceneExtensionProxy* SceneProxy, FSceneInterface& Scene, FRHICommandListBase& RHICmdList)
-	{
-		return nullptr;
-	}
-
-	virtual void DestroyRenderThreadResources(FTransformProviderRenderProxy* ProviderProxy)
+	virtual bool IsCompiling() const
 	{
+		return false;
 	}
 
-	virtual bool IsCompiling() const
+	virtual FTransformProviderRenderProxy* CreateRenderProxy(FInstancedSkinningSceneExtensionProxy* ExtensionSceneProxy) const
 	{
-		return false;
+		return nullptr;
 	}
 
 #if WITH_EDITOR
diff --git a/Engine/Source/Runtime/Engine/Public/SceneInterface.h b/Engine/Source/Runtime/Engine/Public/SceneInterface.h
--- a/Engine/Source/Runtime/Engine/Public/SceneInterface.h
+++ b/Engine/Source/Runtime/Engine/Public/SceneInterface.h
@@ -14,8 +14,7 @@
 #include "PathTracingOutputInvalidateReason.h"
 
 class AWorldSettings;
-struct FAnimBankDesc;
-struct FAnimBankRecordHandle;
+class IAnimBankTransformProvider;
 class FArchive;
 class FFloat16Color;
 class FInstanceCullingManager;
@@ -432,8 +431,7 @@ public:
 	virtual void AddSparseVolumeTextureViewer(FSparseVolumeTextureViewerSceneProxy* SVTV) = 0;
 	virtual void RemoveSparseVolumeTextureViewer(FSparseVolumeTextureViewerSceneProxy* SVTV) = 0;
 
-	virtual TArray<FAnimBankRecordHandle> RegisterAnimBank(const TConstArrayView<FAnimBankDesc> Descs) = 0;
-	virtual void UnregisterAnimBank(const TConstArrayView<FAnimBankRecordHandle> Handles) = 0;
+	virtual IAnimBankTransformProvider* GetAnimBankTransformProvider() = 0;
 
 	/**
 	 * Adds the unique volumetric cloud component to the scene
diff --git a/Engine/Source/Runtime/Engine/Public/SkinningSceneExtensionProxy.h b/Engine/Source/Runtime/Engine/Public/SkinningSceneExtensionProxy.h
--- a/Engine/Source/Runtime/Engine/Public/SkinningSceneExtensionProxy.h
+++ b/Engine/Source/Runtime/Engine/Public/SkinningSceneExtensionProxy.h
@@ -10,6 +10,7 @@ class USkinnedAsset;
 class FSkeletalMeshRenderData;
 class UTransformProviderData;
 class FTransformProviderRenderProxy;
+class FSkinningSceneExtension;
 struct FSkinnedMeshSceneProxyDesc;
 struct FInstancedSkinnedMeshSceneProxyDesc;
 
@@ -45,6 +46,11 @@ public:
 		return MaxBoneTransformCount;
 	}
 
+	uint32 GetMaxBoneTransformIndex() const
+	{
+		return MaxBoneTransformIndex;
+	}
+
 	uint32 GetMaxBoneHierarchyCount() const
 	{
 		return MaxBoneTransformCount;
@@ -97,6 +103,13 @@ public:
 		return 4 /* quat */ + 3 /* XYZ translation */ + (HasScale() ? 3 : 0 /* XYZ scale */);
 	}
 
+	virtual void CreateRenderThreadResources(FSceneInterface& Scene, FRHICommandListBase& RHICmdList) {}
+
+	virtual void DestroyRenderThreadResources() {}
+
+	virtual const FGuid& GetTransformProviderId() const;
+
+private:
 	bool SetLOD(int32 InLODIndex)
 	{
 		const bool bLODChange = InLODIndex != LODIndex;
@@ -108,14 +121,6 @@ public:
 		return bLODChange;
 	}
 
-	virtual void CreateRenderThreadResources(FSceneInterface& Scene, FRHICommandListBase& RHICmdList) {}
-
-	virtual void DestroyRenderThreadResources() {}
-
-	virtual const FGuid& GetTransformProviderId() const;
-
-	virtual TConstArrayView<uint64> GetAnimationProviderData(bool &bOutValid) const;
-
 protected:
 	void InitForDirectStorage(bool bAllowScaling);
 	ENGINE_API void InitForBoneMapStorage();
@@ -128,9 +133,10 @@ protected:
 	TArray<uint32> BoneHierarchy;
 	TArray<float> BoneObjectSpace;
 
+	uint32 UniqueAnimationCount         = 1u;
 	uint16 MaxBoneTransformCount        = 0u;
+	uint16 MaxBoneTransformIndex        = 0u;
 	uint16 MaxBoneInfluenceCount        = 0u;
-	uint16 UniqueAnimationCount         = 1u;
 	uint16 LODIndex                     = InvalidLODIndex;
 
 	EBoneTransformStorageMode BoneTransformStorageMode = EBoneTransformStorageMode::Direct;
@@ -138,22 +144,31 @@ protected:
 	uint8 bHasScale            : 1 = false;
 	uint8 bUseSkeletonBatching : 1 = false;
 	uint8 bUseInstancing       : 1 = false;
+
+	friend class FSkinningSceneExtension;
 };
 
 class FInstancedSkinningSceneExtensionProxy : public FSkinningSceneExtensionProxy
 {
 public:
-	ENGINE_API FInstancedSkinningSceneExtensionProxy(TObjectPtr<UTransformProviderData> InTransformProvider, const FSkeletalMeshObject* InMeshObject, const USkinnedAsset* InSkinnedAsset, bool bAllowScaling);
+	ENGINE_API FInstancedSkinningSceneExtensionProxy(const UTransformProviderData* InTransformProvider, const FSkeletalMeshObject* InMeshObject, const USkinnedAsset* InSkinnedAsset, bool bAllowScaling);
 
 	ENGINE_API void CreateRenderThreadResources(FSceneInterface& Scene, FRHICommandListBase& RHICmdList);
 	ENGINE_API void DestroyRenderThreadResources();
 
-	ENGINE_API TConstArrayView<uint64> GetAnimationProviderData(bool& bOutValid) const override;
+	void SetUniqueAnimationCount(uint32 InUniqueAnimationCount)
+	{
+		UniqueAnimationCount = FMath::Max(InUniqueAnimationCount, 1u);
+	}
 
 	ENGINE_API const FGuid& GetTransformProviderId() const override;
 
+	FTransformProviderRenderProxy* GetTransformProviderProxy() const
+	{
+		return TransformProviderProxy;
+	}
+
 protected:
-	TObjectPtr<UTransformProviderData> TransformProvider;
 	FTransformProviderRenderProxy* TransformProviderProxy = nullptr;
 	FGuid TransformProviderId;
 };
\ No newline at end of file
diff --git a/Engine/Source/Runtime/Renderer/Private/RendererScene.cpp b/Engine/Source/Runtime/Renderer/Private/RendererScene.cpp
--- a/Engine/Source/Runtime/Renderer/Private/RendererScene.cpp
+++ b/Engine/Source/Runtime/Renderer/Private/RendererScene.cpp
@@ -2731,56 +2731,9 @@ void FScene::RemoveHairStrands(FHairStrandsInstance* Proxy)
 	}
 }
 
-TArray<FAnimBankRecordHandle> FScene::RegisterAnimBank(const TConstArrayView<FAnimBankDesc> Descs)
+IAnimBankTransformProvider* FScene::GetAnimBankTransformProvider()
 {
-	check(IsInRenderingThread());
-
-	TArray<FAnimBankRecordHandle> Handles;
-	Handles.Reserve(Descs.Num());
-
-	FAnimBankTransformProvider* BankProvider = GetExtensionPtr<FAnimBankTransformProvider>();
-	if (BankProvider != nullptr)
-	{
-		for (const FAnimBankDesc& Desc : Descs)
-		{
-			if (!Desc.BankAsset.IsValid() || !Desc.Asset.IsValid())
-			{
-				Handles.Emplace(FAnimBankRecordHandle());
-				continue;
-			}
-
-			const FAnimBankData& BankData = Desc.BankAsset->GetData();
-
-			if (Desc.SequenceIndex >= uint32(BankData.Entries.Num()))
-			{
-				Handles.Emplace(FAnimBankRecordHandle());
-				continue;
-			}
-
-			Handles.Emplace(BankProvider->RegisterBank(Desc));
-		}
-	}
-
-	return Handles;
-}
-
-void FScene::UnregisterAnimBank(const TConstArrayView<FAnimBankRecordHandle> Handles)
-{
-	check(IsInRenderingThread());
-
-	FAnimBankTransformProvider* BankProvider = GetExtensionPtr<FAnimBankTransformProvider>();
-	if (BankProvider != nullptr)
-	{
-		for (const FAnimBankRecordHandle& Handle : Handles)
-		{
-			if (!Handle.IsValid())
-			{
-				continue;
-			}
-
-			BankProvider->UnregisterBank(Handle);
-		}
-	}
+	return SceneExtensions.GetExtensionPtr<FAnimBankTransformProvider>();
 }
 
 void FScene::GetLightIESAtlasSlot(const FLightSceneProxy* Proxy, FLightRenderParameters* Out)
@@ -6823,8 +6776,7 @@ public:
 	virtual void AddSparseVolumeTextureViewer(FSparseVolumeTextureViewerSceneProxy* SVTV) override {}
 	virtual void RemoveSparseVolumeTextureViewer(FSparseVolumeTextureViewerSceneProxy* SVTV) override {}
 
-	virtual TArray<FAnimBankRecordHandle> RegisterAnimBank(const TConstArrayView<FAnimBankDesc> Descs) { return TArray<FAnimBankRecordHandle>(); }
-	virtual void UnregisterAnimBank(const TConstArrayView<FAnimBankRecordHandle> Handles) override {}
+	virtual IAnimBankTransformProvider* GetAnimBankTransformProvider() { return nullptr; }
 
 	virtual void AddHairStrands(FHairStrandsInstance* Proxy) override {}
 	virtual void RemoveHairStrands(FHairStrandsInstance* Proxy) override {}
diff --git a/Engine/Source/Runtime/Renderer/Private/ScenePrivate.h b/Engine/Source/Runtime/Renderer/Private/ScenePrivate.h
--- a/Engine/Source/Runtime/Renderer/Private/ScenePrivate.h
+++ b/Engine/Source/Runtime/Renderer/Private/ScenePrivate.h
@@ -3507,8 +3507,7 @@ public:
 	virtual void AddSparseVolumeTextureViewer(FSparseVolumeTextureViewerSceneProxy* SVTV) override;
 	virtual void RemoveSparseVolumeTextureViewer(FSparseVolumeTextureViewerSceneProxy* SVTV) override;
 
-	virtual TArray<FAnimBankRecordHandle> RegisterAnimBank(const TConstArrayView<struct FAnimBankDesc> Descs) override;
-	virtual void UnregisterAnimBank(const TConstArrayView<FAnimBankRecordHandle> Handles) override;
+	virtual IAnimBankTransformProvider* GetAnimBankTransformProvider() override;
 
 	virtual void SetPhysicsField(class FPhysicsFieldSceneProxy* PhysicsFieldSceneProxy) override;
 	virtual void ResetPhysicsField() override;
diff --git a/Engine/Source/Runtime/Renderer/Private/Skinning/AnimBankTransformProvider.cpp b/Engine/Source/Runtime/Renderer/Private/Skinning/AnimBankTransformProvider.cpp
--- a/Engine/Source/Runtime/Renderer/Private/Skinning/AnimBankTransformProvider.cpp
+++ b/Engine/Source/Runtime/Renderer/Private/Skinning/AnimBankTransformProvider.cpp
@@ -20,15 +20,13 @@
 #include "SceneView.h"
 #include "SkeletalRenderPublic.h"
 #include "SkinningDefinitions.h"
+#include "RendererModule.h"
 
 static FGuid AnimBankGPUProviderId(ANIM_BANK_GPU_TRANSFORM_PROVIDER_GUID);
 static FGuid AnimBankCPUProviderId(ANIM_BANK_CPU_TRANSFORM_PROVIDER_GUID);
 
 IMPLEMENT_SCENE_EXTENSION(FAnimBankTransformProvider);
 
-// Animation is always sampled at 30hz
-#define ANIM_BANK_SAMPLE_RATE 30.0f
-
 static TAutoConsoleVariable<bool> CVarAnimBankInterp(
 	TEXT("r.AnimBank.Interpolation"),
 	true,
@@ -79,52 +77,31 @@ private:
 	}
 
 	BEGIN_SHADER_PARAMETER_STRUCT(FParameters, )
-		SHADER_PARAMETER_RDG_BUFFER_SRV(ByteAddressBuffer, HeaderBuffer)
-		SHADER_PARAMETER_RDG_BUFFER_SRV(ByteAddressBuffer, BankBuffer)
+		SHADER_PARAMETER_RDG_BUFFER_SRV(ByteAddressBuffer, SequenceTransformBuffer)
+		SHADER_PARAMETER_RDG_BUFFER_SRV(ByteAddressBuffer, BankBlockBuffer)
+		SHADER_PARAMETER_RDG_BUFFER_SRV(ByteAddressBuffer, BoneMapBuffer)
 		SHADER_PARAMETER_RDG_COMPRESSED_BONE_TRANSFORM_UAV(TransformBuffer)
+		SHADER_PARAMETER(uint32, bInterpolating)
+		SHADER_PARAMETER(float, SampleInterval)
 	END_SHADER_PARAMETER_STRUCT()
 };
 
 IMPLEMENT_GLOBAL_SHADER(FAnimBankEvaluateCS, "/Engine/Private/Skinning/AnimBankEval.usf", "BankEvaluateCS", SF_Compute);
 
-class FAnimBankScatterCS : public FGlobalShader
-{
-public:
-	static constexpr uint32 BonesPerGroup = FAnimBankEvaluateCS::BonesPerGroup;
-
-private:
-	DECLARE_GLOBAL_SHADER(FAnimBankScatterCS);
-	SHADER_USE_PARAMETER_STRUCT(FAnimBankScatterCS, FGlobalShader);
-
-	static bool ShouldCompilePermutation(const FGlobalShaderPermutationParameters& Parameters)
-	{
-		return true;
-	}
-
-	static void ModifyCompilationEnvironment(const FGlobalShaderPermutationParameters& Parameters, FShaderCompilerEnvironment& OutEnvironment)
-	{
-		FGlobalShader::ModifyCompilationEnvironment(Parameters, OutEnvironment);
+FAnimBankTransformProvider::FProxyData::FProxyData(FAnimBankTransformProviderProxy* Proxy)
+	: Tracks(Proxy->GetTracks())
+	, SceneExtensionProxy(*Proxy->GetSceneExtensionProxy())
+	, TransformStorageMode(SceneExtensionProxy.GetBoneTransformStorageMode())
+{}
 
-		OutEnvironment.SetDefine(TEXT("BONES_PER_GROUP"), BonesPerGroup);
-
-		OutEnvironment.CompilerFlags.Add(CFLAG_WarningsAsErrors);
-		OutEnvironment.CompilerFlags.Add(CFLAG_HLSL2021);
-		OutEnvironment.CompilerFlags.Add(CFLAG_AllowTypedUAVLoads);
-	}
-
-	BEGIN_SHADER_PARAMETER_STRUCT(FParameters, )
-		SHADER_PARAMETER_RDG_BUFFER_SRV(ByteAddressBuffer, HeaderBuffer)
-		SHADER_PARAMETER_RDG_COMPRESSED_BONE_TRANSFORM_SRV(SrcTransformBuffer)
-		SHADER_PARAMETER_RDG_BUFFER_SRV(ByteAddressBuffer, SrcBoneMapBuffer)
-		SHADER_PARAMETER_RDG_COMPRESSED_BONE_TRANSFORM_UAV(TransformBuffer)
-	END_SHADER_PARAMETER_STRUCT()
-};
-
-IMPLEMENT_GLOBAL_SHADER(FAnimBankScatterCS, "/Engine/Private/Skinning/AnimBankEval.usf", "BankScatterCS", SF_Compute);
+FAnimBankTransformProvider::FAnimBankTransformProvider(FScene& InScene)
+	: ISceneExtension(InScene)
+	, SequenceTransformPersistentBuffer(0, TEXT("AnimBank.SequenceTransformBuffer"))
+{}
 
 bool FAnimBankTransformProvider::ShouldCreateExtension(FScene& InScene)
 {
-	return IsGPUSkinSceneExtensionEnabled() || (NaniteSkinnedMeshesSupported() && DoesRuntimeSupportNanite(GetFeatureLevelShaderPlatform(InScene.GetFeatureLevel()), true, true));
+	return IsGPUSkinSceneExtensionEnabled();
 }
 
 void FAnimBankTransformProvider::InitExtension(FScene& InScene)
@@ -147,357 +124,369 @@ void FAnimBankTransformProvider::InitExtension(FScene& InScene)
 	}
 }
 
-FAnimBankRecordHandle FAnimBankTransformProvider::RegisterBank(const FAnimBankDesc& Desc)
+void FAnimBankTransformProvider::RegisterAnimBank(FAnimBankTransformProviderProxy* Proxy)
 {
-	FAnimBankRecordHandle Handle;
-
-	const FAnimBankDesc::FDescHash DescHash = BankRecordMap.ComputeHash(Desc);
-	FAnimBankRecord::FRecordId RecordId = BankRecordMap.FindOrAddIdByHash(DescHash, Desc, FAnimBankRecord());
+	TConstArrayView<FAnimBankItem> AnimBankItems = Proxy->GetAnimBankItems();
 
-	Handle.Id	= RecordId.GetIndex();
-	Handle.Hash	= DescHash.AsUInt();
+	FProxyData& ProxyData = ProxyDataMap.Emplace(Proxy->GetSceneExtensionProxy(), Proxy);
+	ProxyData.Sequences.Reserve(AnimBankItems.Num());
+	check(ProxyData.Sequences.IsEmpty());
 
-	FAnimBankRecord& BankRecord = BankRecordMap.GetByElementId(RecordId).Value;
-	if (BankRecord.ReferenceCount == 0)
+	for (const FAnimBankItem& Item : Proxy->GetAnimBankItems())
 	{
-		// First reference
-		BankRecord.Desc = Desc;
-		BankRecord.RecordId = Handle.Id;
-
-		// Calculate total number of sequence rot/pos keys
-		const FAnimBankData& BankData = Desc.BankAsset.Get()->GetData();
-
-		check(Desc.SequenceIndex < uint32(BankData.Entries.Num()));
-		const FAnimBankEntry& BankEntry = BankData.Entries[Desc.SequenceIndex];
-
-		const float SampleInterval = 1.0f / ANIM_BANK_SAMPLE_RATE;
-		const float TrackLength = float(BankEntry.FrameCount - 1) * SampleInterval;
-
-		BankRecord.KeyOffset	= BankAllocator.Allocate(BankEntry.KeyCount);
-		BankRecord.KeyCount		= BankEntry.KeyCount;
-		BankRecord.CurrentTime	= FMath::Clamp<float>(BankEntry.Position, 0.0f, TrackLength);
-		BankRecord.PreviousTime = BankRecord.CurrentTime;
-		BankRecord.Playing		= BankEntry.IsAutoStart() ? 1 : 0;
-		BankRecord.FrameCount	= BankEntry.FrameCount;
-		BankRecord.PositionKeys = BankEntry.PositionKeys; // TODO: Avoid copy
-		BankRecord.RotationKeys = BankEntry.RotationKeys; // TODO: Avoid copy
-		BankRecord.AssetMapping = BankData.Mapping; // TODO: Avoid copy
-
-		check(BankEntry.KeyCount == BankRecord.PositionKeys.Num() && BankEntry.KeyCount == BankRecord.RotationKeys.Num());
-	}
+		const FAnimBankData& AnimBankData = Item.BankAsset->GetData();
+		const FSequenceRegistryKey SequenceRegistryKey(Item);
 
-	++BankRecord.ReferenceCount;
-	return Handle;
-}
+		FMapping& Mapping = MappingRegistry.FindOrAdd(SequenceRegistryKey.BankAsset);
+		if (Mapping.RefCount == 0)
+		{
+			const FSkinnedAssetMapping& AnimBankMapping = AnimBankData.Mapping;
 
-void FAnimBankTransformProvider::UnregisterBank(const FAnimBankRecordHandle& Handle)
-{
-	FAnimBankRecord::FRecordId RecordId(Handle.Id);
-	check(RecordId.IsValid());
+			Mapping.RegistryKey   = SequenceRegistryKey.BankAsset;
+			Mapping.Keys.Position = AnimBankMapping.PositionKeys;
+			Mapping.Keys.Rotation = AnimBankMapping.RotationKeys;
+			Mapping.BoneCount     = AnimBankMapping.BoneCount;
+
+			MappingsToUpload.Add(&Mapping);
+		}
+		++Mapping.RefCount;
 
-	FAnimBankRecord& BankRecord = BankRecordMap.GetByElementId(RecordId).Value;
-	check(BankRecord.ReferenceCount > 0);
-	--BankRecord.ReferenceCount;
+		FSequence& Sequence = SequenceRegistry.FindOrAdd(SequenceRegistryKey);
+		if (Sequence.RefCount == 0)
+		{
+			const FAnimBankEntry& AnimBankEntry = AnimBankData.Entries[SequenceRegistryKey.SequenceIndex];
+
+			Sequence.RegistryKey   = SequenceRegistryKey;
+			Sequence.Keys.Position = AnimBankEntry.PositionKeys;
+			Sequence.Keys.Rotation = AnimBankEntry.RotationKeys;
+			Sequence.PlayLength    = AnimBankEntry.GetPlayLength();
+			Sequence.FrameCount    = AnimBankEntry.FrameCount;
+			Sequence.KeyCount      = AnimBankEntry.KeyCount;
+			Sequence.BoneCount     = Mapping.BoneCount;
+			Sequence.Mapping       = &Mapping;
+			Sequence.bLooping      = AnimBankEntry.IsLooping();
+
+			SequencesToUpload.AddUnique(&Sequence);
+		}
+		++Sequence.RefCount;
 
-	if (BankRecord.ReferenceCount == 0)
-	{
-		BankAllocator.Free(BankRecord.KeyOffset, BankRecord.KeyCount);
-		BankRecordMap.RemoveByElementId(RecordId);
+		ProxyData.Sequences.Emplace(SequenceRegistryKey, Mapping, Sequence);
 	}
 }
 
-struct FAnimBankGPUData
+void FAnimBankTransformProvider::UnregisterAnimBank(FAnimBankTransformProviderProxy* Proxy)
 {
-	TArray<uint32, SceneRenderingAllocator> IdToOffsetMapping;
+	const FInstancedSkinningSceneExtensionProxy* SceneExtensionProxy = Proxy->GetSceneExtensionProxy();
+	const uint32 SceneExtensionProxyHash = GetTypeHash(SceneExtensionProxy);
 
-	FRDGBufferRef BoneBlockBuffer	= nullptr;
-	FRDGBufferRef BankDataBuffer	= nullptr;
-	FRDGBufferRef TransformBuffer	= nullptr;
+	FProxyData& ProxyData = ProxyDataMap.FindByHashChecked(SceneExtensionProxyHash, SceneExtensionProxy);
 
-	uint32 RecordCount = 0;
-	uint32 TransformCount = 0;
-	uint32 BlockCount = 0;
-	uint32 KeyCount = 0;
+	for (const FSequenceView& SequenceView : ProxyData.Sequences)
+	{
+		const uint32 BankAssetHash = GetTypeHash(SequenceView.Key.BankAsset);
+		FMapping& Mapping = MappingRegistry.FindByHashChecked(BankAssetHash, SequenceView.Key.BankAsset);
+		check(Mapping.RefCount > 0);
+		--Mapping.RefCount;
 
-	bool bValid = false;
-};
+		if (Mapping.RefCount == 0)
+		{
+			if (Mapping.BufferAllocation.IsValid())
+			{
+				SequenceTransformAllocator.Free(Mapping.BufferAllocation.Offset, Mapping.BufferAllocation.Count);
+				Mapping.BufferAllocation = {};
+			}
 
-struct FAnimBankCPUData
-{
-	TArray<uint32, SceneRenderingAllocator> IdToOffsetMapping;
+			MappingsToUpload.Remove(&Mapping);
+			MappingRegistry.RemoveByHash(BankAssetHash, SequenceView.Key.BankAsset);
+		}
 
-	FRDGBufferRef TransformBuffer = nullptr;
+		const uint32 SequenceRegistryKeyHash = GetTypeHash(SequenceView.Key);
+		FSequence& Sequence = SequenceRegistry.FindByHashChecked(SequenceRegistryKeyHash, SequenceView.Key);
+		check(Sequence.RefCount > 0);
+		--Sequence.RefCount;
 
-	uint32 TransformCount = 0;
-	uint32 RecordCount = 0;
+		if (Sequence.RefCount == 0)
+		{
+			if (Sequence.BufferAllocation.IsValid())
+			{
+				SequenceTransformAllocator.Free(Sequence.BufferAllocation.Offset, Mapping.BufferAllocation.Count);
+				Sequence.BufferAllocation = {};
+			}
 
-	bool bValid = false;
-};
+			SequencesToUpload.Remove(&Sequence);
+			SequenceRegistry.RemoveByHash(SequenceRegistryKeyHash, SequenceView.Key);
+		}
+	}
+
+	ProxyDataMap.RemoveByHash(SceneExtensionProxyHash, SceneExtensionProxy);
+}
 
-static FAnimBankGPUData BuildAnimBankGPUData(const FAnimBankRecordMap& BankRecordMap, FRDGBuilder& GraphBuilder)
+void FAnimBankTransformProvider::ProvideGPUBankTransforms(FSkinningTransformProvider::FProviderContext& Context)
 {
-	const uint32 BonesPerGroup = FAnimBankEvaluateCS::BonesPerGroup;
+	TRACE_CPUPROFILER_EVENT_SCOPE(FAnimBankTransformProvider::ProvideGPUBankTransforms);
+	FRDGBuilder& GraphBuilder = Context.GraphBuilder;
 
-	FAnimBankGPUData BankData{};
+	int32 TotalSequenceTransformCount = 0;
+	SequenceTransformAllocator.Consolidate();
 
-	int32 MaxOffsetIndex = INDEX_NONE;
+	static_assert(sizeof(FAnimBankTransformProvider::FSequenceTransform) == 7 * sizeof(float));
 
-	for (const auto& RecordPair : BankRecordMap)
+	for (FMapping* Mapping : MappingsToUpload)
 	{
-		const FAnimBankRecord& BankRecord = RecordPair.Value;
-		check(BankRecord.PositionKeys.Num() == BankRecord.RotationKeys.Num());
+		const int32 TransformCount = Mapping->BoneCount;
+		Mapping->BufferAllocation.Count  = TransformCount;
+		Mapping->BufferAllocation.Offset = SequenceTransformAllocator.Allocate(TransformCount);
+		TotalSequenceTransformCount += TransformCount;
+	}
 
-		MaxOffsetIndex = FMath::Max<int32>(MaxOffsetIndex, BankRecord.RecordId);
+	for (FSequence* Sequence : SequencesToUpload)
+	{
+		const int32 TransformCount = Sequence->KeyCount;
+		Sequence->BufferAllocation.Count  = TransformCount;
+		Sequence->BufferAllocation.Offset = SequenceTransformAllocator.Allocate(TransformCount);
+		TotalSequenceTransformCount += TransformCount;
+	}
+
+	FRDGBuffer* SequenceTransformBuffer = SequenceTransformPersistentBuffer.ResizeBufferIfNeeded(GraphBuilder, SequenceTransformAllocator.GetMaxSize());
 
-		if (BankRecord.PositionKeys.Num() == 0 || !BankRecord.Playing)
+	if (TotalSequenceTransformCount > 0)
+	{
+		Context.ScatterUploadBuilder.AddPass(GraphBuilder, SequenceTransformUploadBuffer, SequenceTransformBuffer, TotalSequenceTransformCount, sizeof(FSequenceTransform), TEXT("AnimBank.SequenceTransforms"),
+			[MappingsToUpload = MoveTemp(MappingsToUpload), SequencesToUpload = MoveTemp(SequencesToUpload)] (FRDGScatterUploader& ScatterUploader)
 		{
-			continue;
-		}
+			TRACE_CPUPROFILER_EVENT_SCOPE(FAnimBankTransformProvider::UploadSequences);
 
-		++BankData.RecordCount;
+			for (const FMapping* Mapping : MappingsToUpload)
+			{
+				TArrayView<FSequenceTransform> DstTransforms = ScatterUploader.Add_GetRef<FSequenceTransform>(Mapping->BufferAllocation.Offset, Mapping->BoneCount);
 
-		const FAnimBankData& AnimBankData = BankRecord.Desc.BankAsset->GetData();
-		const uint32 BoneCount = AnimBankData.Mapping.BoneCount;
+				for (uint32 Index = 0; Index < Mapping->BoneCount; ++Index)
+				{
+					DstTransforms[Index].SetRotation(Mapping->Keys.Rotation[Index]);
+					DstTransforms[Index].SetPosition(Mapping->Keys.Position[Index]);
+				}
+			}
+
+			for (const FSequence* Sequence : SequencesToUpload)
+			{
+				TArrayView<FSequenceTransform> DstTransforms = ScatterUploader.Add_GetRef<FSequenceTransform>(Sequence->BufferAllocation.Offset, Sequence->KeyCount);
+
+				for (uint32 Index = 0; Index < Sequence->KeyCount; ++Index)
+				{
+					DstTransforms[Index].SetRotation(Sequence->Keys.Rotation[Index]);
+					DstTransforms[Index].SetPosition(Sequence->Keys.Position[Index]);
+				}
+			}
+		});
 
-		BankData.TransformCount += BoneCount;
-		BankData.BlockCount += FMath::DivideAndRoundUp(BoneCount, BonesPerGroup);
-		BankData.KeyCount += BankRecord.PositionKeys.Num();
+		Context.ScatterUploadBuilder.Execute(Context.GraphBuilder);
 	}
 
-	BankData.IdToOffsetMapping.SetNumZeroed(MaxOffsetIndex + 1);
+	uint32 BlockCountMax = 0;
+	const uint32 BonesPerGroup = FAnimBankEvaluateCS::BonesPerGroup;
+
+	TArrayView<FProxyData*> ProxyDatas = GraphBuilder.AllocPODArrayView<FProxyData*>(Context.Proxies.Num());
 
-	for (auto& RecordPair : BankRecordMap)
+	for (const FSkinningTransformProvider::FProviderIndirection Indirection : Context.Indirections)
 	{
-		const FAnimBankRecord& BankRecord = RecordPair.Value;
-		BankData.IdToOffsetMapping[BankRecord.RecordId] = ~uint32(0);
+		FProxyData& ProxyData = ProxyDataMap.FindChecked(Context.Proxies[Indirection.Index]);
+		ProxyDatas[Indirection.Index] = &ProxyData;
+
+		const uint32 ActiveTrackCount  = ProxyData.Tracks.GetNumActiveTracks();
+		const uint32 LocalBlockCount   = FMath::DivideAndRoundUp(ProxyData.SceneExtensionProxy.GetMaxBoneTransformCount(), BonesPerGroup);
+
+		BlockCountMax += ActiveTrackCount * LocalBlockCount * 2;
 	}
 
-	BankData.bValid = BankData.BlockCount > 0;
-	if (BankData.bValid)
+	if (!BlockCountMax)
 	{
-		const bool bInterpolating = CVarAnimBankInterp.GetValueOnRenderThread();
-
-		// Transform Buffer
-		{
-			BankData.TransformBuffer = GraphBuilder.CreateBuffer(GetCompressedBoneTransformBufferDesc(BankData.TransformCount), TEXT("AnimBank.Transforms"));
-		}
+		return;
+	}
 
-		// Bank Record Data
-		const uint32 BankHeaderSize = sizeof(UE::HLSL::FBankRecordHeader) * BankData.RecordCount;
-		const uint32 BankMappingSize = sizeof(float) * 7 * BankData.TransformCount;
-		const uint32 BankKeySize = sizeof(float) * 7 * BankData.KeyCount;
-		const uint32 BankDataSize = BankHeaderSize + BankMappingSize + BankKeySize;
-		uint8* BankRecordData = GraphBuilder.AllocPODArray<uint8>(BankDataSize);
+	TArrayView<UE::HLSL::FAnimBankBlockHeader> BlockHeaders = GraphBuilder.AllocPODArrayView<UE::HLSL::FAnimBankBlockHeader>(BlockCountMax);
+	int32& BlockCount = *GraphBuilder.AllocPOD<int32>();
+	BlockCount = 0;
 
-		// Bone Block Headers
-		UE::HLSL::FBankBlockHeader* BlockHeaders = GraphBuilder.AllocPODArray<UE::HLSL::FBankBlockHeader>(BankData.BlockCount);
+	GraphBuilder.AddSetupTask([this, ProxyDatas, BlockHeaders, &ViewFamily = Context.ViewFamily, Indirections = Context.Indirections, &BlockCount, BonesPerGroup]
+	{
+		TRACE_CPUPROFILER_EVENT_SCOPE(FAnimBankTransformProvider::BuildBoneBlocks);
+		const float GlobalTimeScale = CVarAnimBankTimeScale.GetValueOnRenderThread();
+		const float DeltaTime = ViewFamily.Time.GetDeltaWorldTimeSeconds() * GlobalTimeScale;
 
-		// Construct Headers and Indirections
+		for (const FSkinningTransformProvider::FProviderIndirection Indirection : Indirections)
 		{
-			uint32 BlockWrite = 0;
-			uint32 TransformOffset = 0;
+			FProxyData& ProxyData = *ProxyDatas[Indirection.Index];
+			FAnimBankTrackPool& Tracks = ProxyData.Tracks;
+			ProxyData.TrackHistories.SetNum(Tracks.GetNumTracks());
+
+			const uint32 MaxTransformCount     = ProxyData.SceneExtensionProxy.GetMaxBoneTransformCount();
+			const uint32 FullBlockCount        = MaxTransformCount / BonesPerGroup;
+			const uint32 PartialTransformCount = MaxTransformCount - (FullBlockCount * BonesPerGroup);
+			const uint32 BoneMapOffset         = ProxyData.TransformStorageMode == EBoneTransformStorageMode::BoneMap ? Indirection.HierarchyOffset : ~uint32(0);
 
-			uint8* BankRecordWrite = BankRecordData;
+			const uint32 CurrentTransformIndex  = Indirection.CurrentTransformOffset  / sizeof(FCompressedBoneTransform);
+			const uint32 PreviousTransformIndex = Indirection.PreviousTransformOffset / sizeof(FCompressedBoneTransform);
+			const uint32 CurrentTransformSlot   = CurrentTransformIndex < PreviousTransformIndex ? 0 : 1;
 
-			for (auto& RecordPair : BankRecordMap)
+			Tracks.EnumerateActiveTracks([&] (int32 TrackIndex)
 			{
-				const FAnimBankRecord& BankRecord = RecordPair.Value;
-				if (BankRecord.KeyCount == 0 || !BankRecord.Playing)
-				{
-					continue;
-				}
+				FAnimBankTrackPackedData& Data      = Tracks.GetData(TrackIndex);
+				const FSequenceView& SequenceView   = ProxyData.GetSequenceView(Data);
+				const FSequence& Sequence           = SequenceView.Sequence;
+				const uint32 SequenceTransformIndex = Sequence.BufferAllocation.Offset;
+				const uint32 MappingTransformIndex  = SequenceView.Mapping.BufferAllocation.Offset;
+				const uint32 LocalTransformIndex    = TrackIndex * MaxTransformCount * 2;
 
-				const FAnimBankData& AnimBankData = BankRecord.Desc.BankAsset->GetData();
-				const uint32 BoneCount = AnimBankData.Mapping.BoneCount;
-
-				const uint32 BankRecordOffset = UE_PTRDIFF_TO_UINT32(BankRecordWrite - BankRecordData);
-
-				UE::HLSL::FBankRecordHeader& Header = *reinterpret_cast<UE::HLSL::FBankRecordHeader*>(BankRecordWrite);
-				BankRecordWrite += sizeof(UE::HLSL::FBankRecordHeader);
-
-				Header.BoneCount		= BoneCount;
-				Header.FrameCount		= BankRecord.FrameCount;
-				Header.SampleRate		= ANIM_BANK_SAMPLE_RATE;
-				Header.PlayRate			= BankRecord.Desc.PlayRate;
-				Header.CurrentTime		= BankRecord.CurrentTime;
-				Header.PreviousTime		= BankRecord.PreviousTime;
-				Header.TransformOffset	= TransformOffset;
-				Header.Playing			= BankRecord.Playing;
-				Header.Interpolating	= bInterpolating;
-				Header.HasScale			= false;
-				checkSlow(Header.TransformOffset == TransformOffset);
-
-				BankData.IdToOffsetMapping[BankRecord.RecordId] = Header.TransformOffset;
-				
-				TransformOffset += sizeof(FCompressedBoneTransform) * Header.BoneCount;
-
-				// Full bone blocks
-				uint32 BlockTransformOffset = Header.TransformOffset;
-				const uint32 FullBlockCount = Header.BoneCount / BonesPerGroup;
-				for (uint32 BlockIndex = 0; BlockIndex < FullBlockCount; ++BlockIndex)
-				{
-					BlockHeaders[BlockWrite].BlockLocalIndex		= BlockIndex;
-					BlockHeaders[BlockWrite].BlockBoneCount			= BonesPerGroup;
-					BlockHeaders[BlockWrite].BlockTransformOffset	= BlockTransformOffset;
-					BlockHeaders[BlockWrite].BankRecordOffset		= BankRecordOffset;
-					++BlockWrite;
+				bool bNeedsCurrentUpdate, bNeedsPreviousUpdate;
+				ProxyData.TrackHistories[TrackIndex].Update(Data, CurrentTransformSlot, bNeedsCurrentUpdate, bNeedsPreviousUpdate);
 
-					BlockTransformOffset += (sizeof(FCompressedBoneTransform) * BonesPerGroup);
-				}
+				// Dirty previous transforms means the whole allocation is dirty. Rebuild everything.
+				const bool bDirtyTransformAllocation = EnumHasAnyFlags(Indirection.DirtyBoneTransforms, EDirtyBoneTransforms::Previous);
+				bNeedsCurrentUpdate  |= bDirtyTransformAllocation;
+				bNeedsPreviousUpdate |= bDirtyTransformAllocation;
 
-				// Partial bone block (remainder)
-				const uint32 PartialBoneCount = Header.BoneCount - (FullBlockCount * BonesPerGroup);
-				if (PartialBoneCount > 0)
+				const float Position              = Data.GetPosition();
+				const float PreviousPosition      = Data.GetPreviousPosition();
+
+				EPreviousBoneTransformUpdateMode UpdateMode = EPreviousBoneTransformUpdateMode::None;
+
+				if (bNeedsPreviousUpdate)
 				{
-					BlockHeaders[BlockWrite].BlockLocalIndex		= FullBlockCount;
-					BlockHeaders[BlockWrite].BlockBoneCount			= PartialBoneCount;
-					BlockHeaders[BlockWrite].BlockTransformOffset	= BlockTransformOffset;
-					BlockHeaders[BlockWrite].BankRecordOffset		= BankRecordOffset;
-					++BlockWrite;
+					UpdateMode = bNeedsCurrentUpdate && bNeedsPreviousUpdate && Position == PreviousPosition
+						? EPreviousBoneTransformUpdateMode::DuplicateCurrentToPrevious
+						: EPreviousBoneTransformUpdateMode::UpdatePrevious;
 				}
 
-				// Asset Mapping
+				const auto AddBlockHeaders = [&] (float Time, uint32 TransformIndex, uint32 OutputTransformCopyIndexOffset)
 				{
-					const FVector3f* PositionKeys	= BankRecord.AssetMapping.PositionKeys.GetData();
-					const FQuat4f* RotationKeys		= BankRecord.AssetMapping.RotationKeys.GetData();
-
-					for (uint32 KeyIndex = 0; KeyIndex < Header.BoneCount; ++KeyIndex)
+					for (uint32 BlockIndex = 0; BlockIndex < FullBlockCount; ++BlockIndex)
 					{
-						float* KeyWrite = reinterpret_cast<float*>(BankRecordWrite);
-
-						KeyWrite[0] = RotationKeys[KeyIndex].X;
-						KeyWrite[1] = RotationKeys[KeyIndex].Y;
-						KeyWrite[2] = RotationKeys[KeyIndex].Z;
-						KeyWrite[3] = RotationKeys[KeyIndex].W;
-
-						KeyWrite[4] = PositionKeys[KeyIndex].X;
-						KeyWrite[5] = PositionKeys[KeyIndex].Y;
-						KeyWrite[6] = PositionKeys[KeyIndex].Z;
+						UE::HLSL::FAnimBankBlockHeader& BlockHeader = BlockHeaders[BlockCount++];
+
+						BlockHeader.Time                      = Time;
+						BlockHeader.BoneMapOffset             = BoneMapOffset;
+						BlockHeader.SequenceTransformIndex    = SequenceTransformIndex;
+						BlockHeader.MappingTransformIndex     = MappingTransformIndex;
+						BlockHeader.OutputTransformIndex      = TransformIndex;
+						BlockHeader.FrameCount                = Sequence.FrameCount;
+						BlockHeader.BoneCount                 = Sequence.BoneCount;
+						BlockHeader.OutputTransformCopyIndexOffset = OutputTransformCopyIndexOffset;
+						BlockHeader.BlockLocalIndex           = BlockIndex;
+						BlockHeader.BlockTransformCount       = BonesPerGroup;
+					}
 
-						BankRecordWrite += sizeof(float) * 7;
+					if (PartialTransformCount > 0)
+					{
+						UE::HLSL::FAnimBankBlockHeader& BlockHeader = BlockHeaders[BlockCount++];
+						
+						BlockHeader.Time                      = Time;
+						BlockHeader.BoneMapOffset             = BoneMapOffset;
+						BlockHeader.SequenceTransformIndex    = SequenceTransformIndex;
+						BlockHeader.MappingTransformIndex     = MappingTransformIndex;
+						BlockHeader.OutputTransformIndex      = TransformIndex;
+						BlockHeader.FrameCount                = Sequence.FrameCount;
+						BlockHeader.BoneCount                 = Sequence.BoneCount;
+						BlockHeader.OutputTransformCopyIndexOffset = OutputTransformCopyIndexOffset;
+						BlockHeader.BlockLocalIndex           = FullBlockCount;
+						BlockHeader.BlockTransformCount       = PartialTransformCount;
 					}
-				}
+				};
 
-				// Position and Rotation Keys
+				if (bNeedsCurrentUpdate)
 				{
-					const FVector3f* PositionKeys	= BankRecord.PositionKeys.GetData();
-					const FQuat4f* RotationKeys		= BankRecord.RotationKeys.GetData();
+					uint32 PrimaryOutputIndex = CurrentTransformIndex;
+					uint32 SecondaryOutputIndexOffset = 0;
 
-					for (uint32 KeyIndex = 0; KeyIndex < BankRecord.KeyCount; ++KeyIndex)
+					if (UpdateMode == EPreviousBoneTransformUpdateMode::DuplicateCurrentToPrevious)
 					{
-						float* KeyWrite = reinterpret_cast<float*>(BankRecordWrite);
-
-						KeyWrite[0] = RotationKeys[KeyIndex].X;
-						KeyWrite[1] = RotationKeys[KeyIndex].Y;
-						KeyWrite[2] = RotationKeys[KeyIndex].Z;
-						KeyWrite[3] = RotationKeys[KeyIndex].W;
+						PrimaryOutputIndex = CurrentTransformSlot == 0 ? CurrentTransformIndex : PreviousTransformIndex;
+						SecondaryOutputIndexOffset = MaxTransformCount;
+					}
 
-						KeyWrite[4] = PositionKeys[KeyIndex].X;
-						KeyWrite[5] = PositionKeys[KeyIndex].Y;
-						KeyWrite[6] = PositionKeys[KeyIndex].Z;
+					AddBlockHeaders(Position, PrimaryOutputIndex + LocalTransformIndex, SecondaryOutputIndexOffset);
+				}
 
-						BankRecordWrite += sizeof(float) * 7;
-					}
+				if (UpdateMode == EPreviousBoneTransformUpdateMode::UpdatePrevious)
+				{
+					AddBlockHeaders(PreviousPosition, PreviousTransformIndex + LocalTransformIndex, 0);
 				}
-				
-			}
-		
-			check(UE_PTRDIFF_TO_UINT32(BankRecordWrite - BankRecordData) == BankDataSize);
+
+				Data.Update(Sequence.PlayLength, Sequence.bLooping, DeltaTime);
+			});
 		}
+	});
 
-		// Create and upload buffers
+	FRDGBuffer* BankBlockBuffer = CreateByteAddressBuffer(
+		GraphBuilder,
+		TEXT("AnimBank.BlockHeaders"),
+		// NumElements
+		[&BlockCount] { return BlockCount * (sizeof(UE::HLSL::FAnimBankBlockHeader) / 4); },
+		// InitialData
+		[Data = BlockHeaders.GetData()] { return Data; },
+		// InitialDataSize
+		[&BlockCount] { return BlockCount * sizeof(UE::HLSL::FAnimBankBlockHeader); }
+	);
 
-		BankData.BoneBlockBuffer = CreateByteAddressBuffer(
-			GraphBuilder,
-			TEXT("AnimBank.BlockHeaders"),
-			FMath::RoundUpToPowerOfTwo(sizeof(UE::HLSL::FBankBlockHeader) * BankData.BlockCount),
-			BlockHeaders,
-			sizeof(UE::HLSL::FBankBlockHeader) * BankData.BlockCount,
-			// The buffer data is allocated above on the RDG timeline
-			ERDGInitialDataFlags::NoCopy
-		);
+	const bool bInterpolating = CVarAnimBankInterp.GetValueOnRenderThread();
 
-		BankData.BankDataBuffer = CreateByteAddressBuffer(
-			GraphBuilder,
-			TEXT("AnimBank.BankData"),
-			FMath::RoundUpToPowerOfTwo(BankDataSize),
-			BankRecordData,
-			BankDataSize,
-			// The buffer data is allocated above on the RDG timeline
-			ERDGInitialDataFlags::NoCopy
-		);
-	}
+	FAnimBankEvaluateCS::FParameters* PassParameters = GraphBuilder.AllocParameters<FAnimBankEvaluateCS::FParameters>();
+	PassParameters->SequenceTransformBuffer = GraphBuilder.CreateSRV(SequenceTransformBuffer);
+	PassParameters->BankBlockBuffer = GraphBuilder.CreateSRV(BankBlockBuffer);
+	PassParameters->BoneMapBuffer = Context.HierarchyBufferSRV;
+	PassParameters->TransformBuffer = GetCompressedBoneTransformUAV(GraphBuilder, Context.TransformBuffer);
+	PassParameters->bInterpolating = bInterpolating;
+	PassParameters->SampleInterval = 1.0f / ANIM_BANK_SAMPLE_RATE;
 
-	return BankData;
+	auto ComputeShader = GetGlobalShaderMap(GMaxRHIFeatureLevel)->GetShader<FAnimBankEvaluateCS>();
+	FComputeShaderUtils::AddPass(
+		GraphBuilder,
+		RDG_EVENT_NAME("AnimBankEvaluate"),
+		ComputeShader,
+		PassParameters,
+		[&BlockCount] { return FIntVector(BlockCount, 1, 1); }
+	);
 }
 
-static FAnimBankCPUData BuildAnimBankCPUData(const FAnimBankRecordMap& BankRecordMap, FRDGBuilder& GraphBuilder)
+void FAnimBankTransformProvider::ProvideCPUBankTransforms(FSkinningTransformProvider::FProviderContext& Context)
 {
-	FAnimBankCPUData BankData{};
+	FRDGBuilder& GraphBuilder = Context.GraphBuilder;
 
-	int32 MaxOffsetIndex = INDEX_NONE;
+	uint32 GlobalTransformCount = 0;
 
-	for (const auto& RecordPair : BankRecordMap)
+	for (const FSkinningTransformProvider::FProviderIndirection Indirection : Context.Indirections)
 	{
-		const FAnimBankRecord& BankRecord = RecordPair.Value;
-		check(BankRecord.PositionKeys.Num() == BankRecord.RotationKeys.Num());
-
-		MaxOffsetIndex = FMath::Max<int32>(MaxOffsetIndex, BankRecord.RecordId);
+		const FInstancedSkinningSceneExtensionProxy* Proxy = static_cast<FInstancedSkinningSceneExtensionProxy*>(Context.Proxies[Indirection.Index]);
+		FAnimBankTransformProviderProxy* AnimBankProxy = static_cast<FAnimBankTransformProviderProxy*>(Proxy->GetTransformProviderProxy());
 
-		if (BankRecord.PositionKeys.Num() == 0 || !BankRecord.Playing)
-		{
-			continue;
-		}
-
-		const FAnimBankData& AnimBankData = BankRecord.Desc.BankAsset->GetData();
-		const uint32 BoneCount = AnimBankData.Mapping.BoneCount;
-
-		++BankData.RecordCount;
-
-		BankData.TransformCount += BoneCount;
+		const uint32 TransformCount = Proxy->GetMaxBoneTransformCount();
+		const uint32 ActiveAnimationCount = AnimBankProxy->GetTracks().GetNumActiveTracks();
+		GlobalTransformCount += ActiveAnimationCount * TransformCount * 2; // Current and Previous
 	}
 
-	BankData.IdToOffsetMapping.SetNumZeroed(MaxOffsetIndex + 1);
-
-	for (auto& RecordPair : BankRecordMap)
+	if (GlobalTransformCount == 0)
 	{
-		const FAnimBankRecord& BankRecord = RecordPair.Value;
-		BankData.IdToOffsetMapping[BankRecord.RecordId] = ~uint32(0);
+		return;
 	}
 
-	BankData.bValid = BankData.RecordCount > 0;
-	if (BankData.bValid)
+	Context.ScatterUploadBuilder.AddPass(GraphBuilder, SkinningTransformUploadBuffer, Context.TransformBuffer, GlobalTransformCount, sizeof(FCompressedBoneTransform), TEXT("Skinning.AnimTransforms"),
+		[this, Indirections = Context.Indirections, Proxies = Context.Proxies, &ViewFamily = Context.ViewFamily, GlobalTransformCount] (FRDGScatterUploader& ScatterUploader)
 	{
+		TRACE_CPUPROFILER_EVENT_SCOPE(FSkinningSceneExtension::ProvideAnimBankRuntimeTransformsTask);
+		TArray<FCompressedBoneTransform, SceneRenderingAllocator> SourceTransforms;
+		TBitArray<SceneRenderingBitArrayAllocator> SourceTransformMap;
 		const bool bInterpolating = CVarAnimBankInterp.GetValueOnRenderThread();
 
-		TArrayView<FCompressedBoneTransform> Transforms = GraphBuilder.AllocPODArrayView<FCompressedBoneTransform>(BankData.TransformCount);
-
-		uint32 TransformOffset = 0;
-
-		for (const auto& RecordPair : BankRecordMap)
+		const auto SampleAnimBank = [&SourceTransforms, &SourceTransformMap, bInterpolating] (
+			const FInstancedSkinningSceneExtensionProxy* Proxy,
+			TArrayView<FCompressedBoneTransform> Transforms,
+			TArrayView<FCompressedBoneTransform> DuplicatePreviousTransforms,
+			const FSequence& Sequence,
+			const FMapping& Mapping,
+			float TimeOffset)
 		{
-			const FAnimBankRecord& BankRecord = RecordPair.Value;
-			if (BankRecord.PositionKeys.Num() == 0 || !BankRecord.Playing)
-			{
-				continue;
-			}
-
-			check(BankRecord.PositionKeys.Num() == BankRecord.RotationKeys.Num());
-
-			BankData.IdToOffsetMapping[BankRecord.RecordId] = (TransformOffset * sizeof(FCompressedBoneTransform));
-
-			const float SampleInterval = 1.0f / ANIM_BANK_SAMPLE_RATE;
-			const float TrackLength = float(BankRecord.FrameCount - 1) * SampleInterval;
-
-			double Time = TrackLength == 0.0f ? 0.0 : (double)FMath::Fmod(BankRecord.CurrentTime, TrackLength);
-			if (Time < 0.0)
-			{
-				Time += TrackLength;
-			}
-
 			int32 KeyIndex0, KeyIndex1;
 			float Alpha;
-			FAnimationRuntime::GetKeyIndicesFromTime(KeyIndex0, KeyIndex1, Alpha, Time, BankRecord.FrameCount, TrackLength);
+			FAnimationRuntime::GetKeyIndicesFromTime(KeyIndex0, KeyIndex1, Alpha, TimeOffset, Sequence.FrameCount, Sequence.PlayLength);
 
 			if (!bInterpolating)
 			{
@@ -505,29 +494,26 @@ static FAnimBankCPUData BuildAnimBankCPUData(const FAnimBankRecordMap& BankRecor
 				Alpha = 0.0f;
 			}
 
-			const FAnimBankData& AnimBankData = BankRecord.Desc.BankAsset->GetData();
-			const uint32 BoneCount = AnimBankData.Mapping.BoneCount;
-
-			for (uint32 TransformIndex = 0; TransformIndex < BoneCount; ++TransformIndex)
+			const auto SampleTransform = [&] (int32 TransformIndex) -> FMatrix44f
 			{
-				const FTransform InvGlobalRefPoseXform = FTransform(FQuat(BankRecord.AssetMapping.RotationKeys[TransformIndex]), FVector(BankRecord.AssetMapping.PositionKeys[TransformIndex]));
+				const FTransform InvGlobalRefPoseXform = FTransform(FQuat(Mapping.Keys.Rotation[TransformIndex]), FVector(Mapping.Keys.Position[TransformIndex]));
 
 				FTransform MeshGlobalAnimPoseXform;
 
 				if (Alpha <= 0.f)
 				{
-					MeshGlobalAnimPoseXform = FTransform(FQuat(BankRecord.RotationKeys[(KeyIndex0 * BoneCount) + TransformIndex]), FVector(BankRecord.PositionKeys[(KeyIndex0 * BoneCount) + TransformIndex]));
+					MeshGlobalAnimPoseXform = FTransform(FQuat(Sequence.Keys.Rotation[(KeyIndex0 * Mapping.BoneCount) + TransformIndex]), FVector(Sequence.Keys.Position[(KeyIndex0 * Mapping.BoneCount) + TransformIndex]));
 					MeshGlobalAnimPoseXform.NormalizeRotation();
 				}
 				else if (Alpha >= 1.f)
 				{
-					MeshGlobalAnimPoseXform = FTransform(FQuat(BankRecord.RotationKeys[(KeyIndex1 * BoneCount) + TransformIndex]), FVector(BankRecord.PositionKeys[(KeyIndex1 * BoneCount) + TransformIndex]));
+					MeshGlobalAnimPoseXform = FTransform(FQuat(Sequence.Keys.Rotation[(KeyIndex1 * Mapping.BoneCount) + TransformIndex]), FVector(Sequence.Keys.Position[(KeyIndex1 * Mapping.BoneCount) + TransformIndex]));
 					MeshGlobalAnimPoseXform.NormalizeRotation();
 				}
 				else
 				{
-					FTransform MeshGlobalXformA = FTransform(FQuat(BankRecord.RotationKeys[(KeyIndex0 * BoneCount) + TransformIndex]), FVector(BankRecord.PositionKeys[(KeyIndex0 * BoneCount) + TransformIndex]));
-					FTransform MeshGlobalXformB = FTransform(FQuat(BankRecord.RotationKeys[(KeyIndex1 * BoneCount) + TransformIndex]), FVector(BankRecord.PositionKeys[(KeyIndex1 * BoneCount) + TransformIndex]));
+					FTransform MeshGlobalXformA = FTransform(FQuat(Sequence.Keys.Rotation[(KeyIndex0 * Mapping.BoneCount) + TransformIndex]), FVector(Sequence.Keys.Position[(KeyIndex0 * Mapping.BoneCount) + TransformIndex]));
+					FTransform MeshGlobalXformB = FTransform(FQuat(Sequence.Keys.Rotation[(KeyIndex1 * Mapping.BoneCount) + TransformIndex]), FVector(Sequence.Keys.Position[(KeyIndex1 * Mapping.BoneCount) + TransformIndex]));
 
 					// Ensure rotations are normalized
 					MeshGlobalXformA.NormalizeRotation();
@@ -539,246 +525,108 @@ static FAnimBankCPUData BuildAnimBankCPUData(const FAnimBankRecordMap& BankRecor
 
 				MeshGlobalAnimPoseXform = InvGlobalRefPoseXform * MeshGlobalAnimPoseXform;
 
-				FMatrix44f Transform = (FMatrix44f)MeshGlobalAnimPoseXform.ToMatrixNoScale();
-
-				StoreCompressedBoneTransform(&Transforms[TransformOffset], Transform);
-				++TransformOffset;
-			}
-		}
-
-		// Transform Buffer
-		{
-			BankData.TransformBuffer = GraphBuilder.CreateBuffer(GetCompressedBoneTransformBufferDesc(BankData.TransformCount), TEXT("AnimBank.Transforms"));
-			GraphBuilder.QueueBufferUpload(BankData.TransformBuffer, Transforms.GetData(), Transforms.Num() * sizeof(FCompressedBoneTransform), ERDGInitialDataFlags::NoCopy);
-		}
-	}
-
-	return BankData;
-}
-
-void FAnimBankTransformProvider::AdvanceAnimation(FSkinningTransformProvider::FProviderContext& Context)
-{
-	if (LastAnimationFrame == GFrameNumberRenderThread)
-	{
-		return;
-	}
-
-	const float GlobalTimeScale = CVarAnimBankTimeScale.GetValueOnRenderThread();
-	const float DeltaTime = Context.ViewFamily.Time.GetDeltaWorldTimeSeconds();
-
-	for (auto& RecordPair : BankRecordMap)
-	{
-		FAnimBankRecord& BankRecord = RecordPair.Value;
-		if (!BankRecord.Playing)
-		{
-			continue;
-		}
-
-		BankRecord.PreviousTime = BankRecord.CurrentTime;
-		BankRecord.CurrentTime += (DeltaTime * BankRecord.Desc.PlayRate) * GlobalTimeScale;
-	}
-
-	LastAnimationFrame = GFrameNumberRenderThread;
-}
-
-void FAnimBankTransformProvider::ScatterAnimation(
-	FSkinningTransformProvider::FProviderContext& Context,
-	const TConstArrayView<uint32> IdToOffsetMapping,
-	FRDGBufferRef TransformBuffer
-)
-{
-	FRDGBuilder& GraphBuilder = Context.GraphBuilder;
-
-	const uint32 BonesPerGroup = FAnimBankScatterCS::BonesPerGroup;
-
-	uint32 BlockCount = 0;
-	for (const FSkinningTransformProvider::FProviderIndirection Indirection : Context.Indirections)
-	{
-		const FPrimitiveSceneInfo* Primitive = Context.Primitives[Indirection.Index];
-		const FSkinningSceneExtensionProxy* Proxy = Context.Proxies[Indirection.Index];
-		const uint32 TransformCount = Proxy->GetMaxBoneTransformCount();
-		const uint32 AnimationCount = Proxy->GetUniqueAnimationCount();
-		BlockCount += FMath::DivideAndRoundUp(TransformCount, BonesPerGroup) * AnimationCount;
-	}
-
-	TArrayView<UE::HLSL::FBankScatterHeader> BlockHeaders = GraphBuilder.AllocPODArrayView<UE::HLSL::FBankScatterHeader>(BlockCount);
-	uint32 BlockWriteIndex = 0;
-
-	for (const FSkinningTransformProvider::FProviderIndirection Indirection : Context.Indirections)
-	{
-		const FSkinningSceneExtensionProxy* Proxy = Context.Proxies[Indirection.Index];
-		const uint32 TransformCount = Proxy->GetMaxBoneTransformCount();
-		check(Proxy->UseInstancing());
-
-		bool bIsValid = false;
-		const TConstArrayView<uint64> BankIds = Proxy->GetAnimationProviderData(bIsValid);
-		check(bIsValid && BankIds.Num() == Proxy->GetUniqueAnimationCount());
-
-		const uint32 FullBlockCount = TransformCount / BonesPerGroup;
-		const uint32 PartialTransformCount = TransformCount - (FullBlockCount * BonesPerGroup);
-
-		// The proxy may need bones laid out with the bone map indirection built-in.
-		const bool bUseBoneMapBuffer = Proxy->GetBoneTransformStorageMode() == EBoneTransformStorageMode::BoneMap;
-
-		uint32 DstTransformOffset = Indirection.CurrentTransformOffset;
+				return (FMatrix44f)MeshGlobalAnimPoseXform.ToMatrixNoScale();
+			};
 
-		for (const uint64 BankId : BankIds)
-		{
-			if (BankId == ~uint64(0))
+			if (Proxy->GetBoneTransformStorageMode() == EBoneTransformStorageMode::BoneMap)
 			{
-				continue;
-			}
-
-			check(BankId < uint64(IdToOffsetMapping.Num()));
-			uint32 SrcTransformOffset = IdToOffsetMapping[int32(BankId)];
-
-			const bool bIsPending = (SrcTransformOffset == ~uint32(0));
-			const uint32 BlockStart = BlockWriteIndex;
+				const uint32 NumSourceTransforms = Proxy->GetMaxBoneTransformIndex() + 1;
+				TConstArrayView<uint32> BoneHierarchy = Proxy->GetBoneHierarchy();
+				SourceTransformMap.Init(false, NumSourceTransforms);
+				SourceTransforms.SetNumUninitialized(NumSourceTransforms, EAllowShrinking::No);
+				int32 DstTransformIndex = 0;
 
-			uint32 BlockBoneMapOffset = bUseBoneMapBuffer ? Indirection.HierarchyOffset : ~uint32(0);
-
-			auto WriteBlockHeaderFunction = [&SrcTransformOffset, &BlockBoneMapOffset, &DstTransformOffset, TransformCount, bUseBoneMapBuffer]
-				(UE::HLSL::FBankScatterHeader& BlockHeader, uint32 BlockIndex, uint32 BlockTransformCount) mutable
-			{
-				BlockHeader.BlockLocalIndex         = BlockIndex;
-				BlockHeader.BlockSrcTransformOffset = SrcTransformOffset;
-				BlockHeader.BlockSrcBoneMapOffset   = BlockBoneMapOffset;
-				BlockHeader.BlockDstTransformOffset = DstTransformOffset;
-				BlockHeader.BlockTransformCount     = BlockTransformCount;
-				BlockHeader.TotalTransformCount     = TransformCount;
-
-				// When using the bone map, increment the source bone map offset for each block.
-				if (bUseBoneMapBuffer)
+				for (uint32 SrcTransformIndex : Proxy->GetBoneHierarchy())
 				{
-					BlockBoneMapOffset += (BlockTransformCount * sizeof(uint32));
-				}
-				// Otherwise increment the source transform offset for each block.
-				else
-				{
-					SrcTransformOffset += (BlockTransformCount * sizeof(FCompressedBoneTransform));
-				}
+					if (!SourceTransformMap[SrcTransformIndex])
+					{
+						SourceTransformMap[SrcTransformIndex] = true;
+						StoreCompressedBoneTransform(&SourceTransforms[SrcTransformIndex], SampleTransform(SrcTransformIndex));
+					}
 
-				DstTransformOffset += (BlockTransformCount * sizeof(FCompressedBoneTransform));
-			};
+					Transforms[DstTransformIndex] = SourceTransforms[SrcTransformIndex];
 
-			for (uint32 BlockIndex = 0; BlockIndex < FullBlockCount; ++BlockIndex)
-			{
-				WriteBlockHeaderFunction(BlockHeaders[BlockWriteIndex], BlockIndex, BonesPerGroup);
-				++BlockWriteIndex;
-			}
+					if (!DuplicatePreviousTransforms.IsEmpty())
+					{
+						DuplicatePreviousTransforms[DstTransformIndex] = SourceTransforms[SrcTransformIndex];
+					}
 
-			if (PartialTransformCount > 0)
-			{
-				WriteBlockHeaderFunction(BlockHeaders[BlockWriteIndex], FullBlockCount, PartialTransformCount);
-				++BlockWriteIndex;
+					DstTransformIndex++;
+				}
 			}
-
-			if (bIsPending)
+			else
 			{
-				for (uint32 BlockIndex = 0; BlockIndex < FullBlockCount; ++BlockIndex)
+				for (uint32 TransformIndex = 0; TransformIndex < Mapping.BoneCount; ++TransformIndex)
 				{
-					BlockHeaders[BlockStart + BlockIndex].BlockSrcTransformOffset = ~uint32(0);
-				}
+					StoreCompressedBoneTransform(&Transforms[TransformIndex], SampleTransform(TransformIndex));
 
-				if (PartialTransformCount > 0)
-				{
-					BlockHeaders[BlockStart + FullBlockCount].BlockSrcTransformOffset = ~uint32(0);
+					if (!DuplicatePreviousTransforms.IsEmpty())
+					{
+						StoreCompressedBoneTransform(&DuplicatePreviousTransforms[TransformIndex], SampleTransform(TransformIndex));
+					}
 				}
 			}
+		};
 
-			// Advance past all the previous transforms as well
-			DstTransformOffset += (sizeof(FCompressedBoneTransform) * TransformCount);
-		}
-	}
+		const float GlobalTimeScale = CVarAnimBankTimeScale.GetValueOnRenderThread();
+		const float DeltaTime = ViewFamily.Time.GetDeltaWorldTimeSeconds() * GlobalTimeScale;
 
-	FRDGBufferRef ScatterBlockHeaders = CreateByteAddressBuffer(
-		GraphBuilder,
-		TEXT("AnimBank.ScatterHeaders"),
-		FMath::RoundUpToPowerOfTwo(sizeof(UE::HLSL::FBankScatterHeader) * BlockCount),
-		BlockHeaders.GetData(),
-		sizeof(UE::HLSL::FBankScatterHeader) * BlockCount,
-		// The buffer data is allocated above on the RDG timeline
-		ERDGInitialDataFlags::NoCopy
-	);
+		for (const FSkinningTransformProvider::FProviderIndirection Indirection : Indirections)
+		{
+			const FInstancedSkinningSceneExtensionProxy* Proxy = static_cast<FInstancedSkinningSceneExtensionProxy*>(Proxies[Indirection.Index]);
+			FProxyData& ProxyData = ProxyDataMap.FindChecked(Proxy);
+			FAnimBankTrackPool& Tracks = ProxyData.Tracks;
+			ProxyData.TrackHistories.SetNum(Tracks.GetNumTracks());
 
-	// Can be null if there are banks, but all are pending (so evaluation is skipped)
-	// None of the blocks will be sourcing from this buffer, ref pose will be written
-	FRDGBufferRef SrcTransformBuffer = TransformBuffer ? TransformBuffer : GetDefaultCompressedBoneTransformBuffer(GraphBuilder);
+			const uint32 MaxTransformCount = Proxy->GetMaxBoneTransformCount();
+			const uint32 UniqueAnimationCount = Proxy->GetUniqueAnimationCount();
+			const uint32 MaxTotalTransformCount = MaxTransformCount * 2u; // Current and Previous
 
-	FAnimBankScatterCS::FParameters* PassParameters = GraphBuilder.AllocParameters<FAnimBankScatterCS::FParameters>();
-	PassParameters->TransformBuffer = GetCompressedBoneTransformUAV(GraphBuilder, Context.TransformBuffer);
-	PassParameters->SrcTransformBuffer = GetCompressedBoneTransformSRV(GraphBuilder, SrcTransformBuffer);
-	PassParameters->SrcBoneMapBuffer = Context.HierarchyBufferSRV;
-	PassParameters->HeaderBuffer = GraphBuilder.CreateSRV(ScatterBlockHeaders);
+			const uint32 CurrentTransformIndex  = Indirection.CurrentTransformOffset  / sizeof(FCompressedBoneTransform);
+			const uint32 PreviousTransformIndex = Indirection.PreviousTransformOffset / sizeof(FCompressedBoneTransform);
+			const uint32 CurrentTransformSlot   = CurrentTransformIndex < PreviousTransformIndex ? 0 : 1;
 
-	auto ComputeShader = GetGlobalShaderMap(GMaxRHIFeatureLevel)->GetShader<FAnimBankScatterCS>();
-	FComputeShaderUtils::AddPass(
-		GraphBuilder,
-		RDG_EVENT_NAME("AnimBankScatter"),
-		ComputeShader,
-		PassParameters,
-		FIntVector(BlockCount, 1, 1)
-	);
-}
-
-void FAnimBankTransformProvider::ProvideGPUBankTransforms(FSkinningTransformProvider::FProviderContext& Context)
-{
-	AdvanceAnimation(Context);
+			Tracks.EnumerateActiveTracks([&] (int32 TrackIndex)
+			{
+				FAnimBankTrackPackedData& Data    = Tracks.GetData(TrackIndex);
+				const FSequenceView& SequenceView = ProxyData.GetSequenceView(Data);
+	
+				bool bNeedsCurrentUpdate, bNeedsPreviousUpdate;
+				ProxyData.TrackHistories[TrackIndex].Update(Data, CurrentTransformSlot, bNeedsCurrentUpdate, bNeedsPreviousUpdate);
 
-	FRDGBuilder& GraphBuilder = Context.GraphBuilder;
-	const FAnimBankGPUData BankData = BuildAnimBankGPUData(BankRecordMap, GraphBuilder);
-	if (!BankData.bValid)
-	{
-		return;
-	}
+				// Dirty previous transforms means the whole allocation is dirty. Rebuild everything.
+				const bool bDirtyTransformAllocation = EnumHasAnyFlags(Indirection.DirtyBoneTransforms, EDirtyBoneTransforms::Previous);
+				bNeedsCurrentUpdate  |= bDirtyTransformAllocation;
+				bNeedsPreviousUpdate |= bDirtyTransformAllocation;
 
-	// Evaluate animation banks
-	{
-		FAnimBankEvaluateCS::FParameters* PassParameters = GraphBuilder.AllocParameters<FAnimBankEvaluateCS::FParameters>();
-		PassParameters->TransformBuffer = GetCompressedBoneTransformUAV(GraphBuilder, BankData.TransformBuffer);
-		PassParameters->BankBuffer = GraphBuilder.CreateSRV(BankData.BankDataBuffer);
-		PassParameters->HeaderBuffer = GraphBuilder.CreateSRV(BankData.BoneBlockBuffer);
-
-		auto ComputeShader = GetGlobalShaderMap(GMaxRHIFeatureLevel)->GetShader<FAnimBankEvaluateCS>();
-		FComputeShaderUtils::AddPass(
-			GraphBuilder,
-			RDG_EVENT_NAME("AnimBankEvaluate"),
-			ComputeShader,
-			PassParameters,
-			FIntVector(BankData.BlockCount, 1, 1)
-		);
-	}
+				const float Position              = Data.GetPosition();
+				const float PreviousPosition      = Data.GetPreviousPosition();
 
-	// Scatter animation bank results
-	ScatterAnimation(Context, MakeConstArrayView(BankData.IdToOffsetMapping), BankData.TransformBuffer);
-}
+				EPreviousBoneTransformUpdateMode UpdateMode = EPreviousBoneTransformUpdateMode::None;
+				TArrayView<FCompressedBoneTransform> DstTransformsPrevious;
 
-void FAnimBankTransformProvider::ProvideCPUBankTransforms(FSkinningTransformProvider::FProviderContext& Context)
-{
-	AdvanceAnimation(Context);
+				if (bNeedsPreviousUpdate)
+				{
+					UpdateMode = bNeedsCurrentUpdate && bNeedsPreviousUpdate && Position == PreviousPosition
+						? EPreviousBoneTransformUpdateMode::DuplicateCurrentToPrevious
+						: EPreviousBoneTransformUpdateMode::UpdatePrevious;
 
-	// Evaluate animation banks on the CPU, and then upload
-	// the results to the GPU as input to the scatter.
+					DstTransformsPrevious = ScatterUploader.Add_GetRef<FCompressedBoneTransform>(PreviousTransformIndex + TrackIndex * MaxTotalTransformCount, MaxTransformCount);
+				}
 
-	FRDGBuilder& GraphBuilder = Context.GraphBuilder;
-	const FAnimBankCPUData BankData = BuildAnimBankCPUData(BankRecordMap, GraphBuilder);
-	if (!BankData.bValid)
-	{
-		return;
-	}
+				if (bNeedsCurrentUpdate)
+				{
+					TArrayView<FCompressedBoneTransform> DstTransforms = ScatterUploader.Add_GetRef<FCompressedBoneTransform>(CurrentTransformIndex + TrackIndex * MaxTotalTransformCount, MaxTransformCount);
+					SampleAnimBank(Proxy, DstTransforms, UpdateMode == EPreviousBoneTransformUpdateMode::DuplicateCurrentToPrevious ? DstTransformsPrevious : TArrayView<FCompressedBoneTransform>{}, SequenceView.Sequence, SequenceView.Mapping, Position);
+				}
 
-	// Scatter animation bank results
-	ScatterAnimation(Context, MakeConstArrayView(BankData.IdToOffsetMapping), BankData.TransformBuffer);
-}
+				if (UpdateMode == EPreviousBoneTransformUpdateMode::UpdatePrevious)
+				{
+					SampleAnimBank(Proxy, DstTransformsPrevious, {}, SequenceView.Sequence, SequenceView.Mapping, PreviousPosition);
+				}
 
-const FSkinningTransformProvider::FProviderId& GetAnimBankProviderId()
-{
-	if (CVarAnimBankGPU.GetValueOnRenderThread())
-	{
-		return AnimBankGPUProviderId;
-	}
-	else
-	{
-		return AnimBankCPUProviderId;
-	}
-}
+				Data.Update(SequenceView.Sequence.PlayLength, SequenceView.Sequence.bLooping, DeltaTime);
+			});
+		}
+	});
+}
\ No newline at end of file
diff --git a/Engine/Source/Runtime/Renderer/Private/Skinning/AnimBankTransformProvider.h b/Engine/Source/Runtime/Renderer/Private/Skinning/AnimBankTransformProvider.h
--- a/Engine/Source/Runtime/Renderer/Private/Skinning/AnimBankTransformProvider.h
+++ b/Engine/Source/Runtime/Renderer/Private/Skinning/AnimBankTransformProvider.h
@@ -8,39 +8,184 @@
 #include "RendererPrivateUtils.h"
 #include "SceneExtensions.h"
 #include "SkinningTransformProvider.h"
+#include "SkinningSceneExtensionProxy.h"
+#include "UnifiedBuffer.h"
 #include "SpanAllocator.h"
 
-class FAnimBankTransformProvider : public ISceneExtension
+class FAnimBankTransformProvider
+	: public ISceneExtension
+	, public IAnimBankTransformProvider
 {
 	DECLARE_SCENE_EXTENSION(RENDERER_API, FAnimBankTransformProvider);
 
 public:
-	using ISceneExtension::ISceneExtension;
+	FAnimBankTransformProvider(FScene& InScene);
 
 	static bool ShouldCreateExtension(FScene& InScene);
 
 	virtual void InitExtension(FScene& InScene) override;
 
-	RENDERER_API FAnimBankRecordHandle RegisterBank(const FAnimBankDesc& Desc);
-	RENDERER_API void UnregisterBank(const FAnimBankRecordHandle& Handle);
-
-	RENDERER_API static FAnimBankRecordHandle ComputeDescHash(const FAnimBankDesc& Desc);
+	void RegisterAnimBank(FAnimBankTransformProviderProxy* Proxy) override;
+	void UnregisterAnimBank(FAnimBankTransformProviderProxy* Proxy) override;
 
 private:
-	void AdvanceAnimation(FSkinningTransformProvider::FProviderContext& Context);
-	void ScatterAnimation(
-		FSkinningTransformProvider::FProviderContext& Context,
-		const TConstArrayView<uint32> IdToOffsetMapping,
-		FRDGBufferRef SrcTransformBuffer
-	);
-
 	void ProvideGPUBankTransforms(FSkinningTransformProvider::FProviderContext& Context);
 	void ProvideCPUBankTransforms(FSkinningTransformProvider::FProviderContext& Context);
 
-private:
-	FAnimBankRecordMap BankRecordMap;
-	FSpanAllocator BankAllocator;
-	uint64 LastAnimationFrame = 0;
-};
+	struct FBufferAllocation
+	{
+		bool IsValid() const
+		{
+			return Count != 0;
+		}
+
+		uint32 Offset = 0;
+		uint32 Count  = 0;
+	};
+
+	struct FSequenceTransform
+	{
+		void SetRotation(FQuat4f Q)
+		{
+			Rotation[0] = Q.X;
+			Rotation[1] = Q.Y;
+			Rotation[2] = Q.Z;
+			Rotation[3] = Q.W;
+		}
+
+		void SetPosition(FVector3f P)
+		{
+			Position[0] = P.X;
+			Position[1] = P.Y;
+			Position[2] = P.Z;
+		}
+
+		float Rotation[4];
+		float Position[3];
+	};
+	
+	struct FSequenceTransformKeys
+	{
+		TConstArrayView<FVector3f> Position;
+		TConstArrayView<FQuat4f>   Rotation;
+	};
+
+	struct FSequenceRegistryKey
+	{
+		FSequenceRegistryKey() = default;
+
+		FSequenceRegistryKey(const FAnimBankItem& Item)
+			: BankAsset(Item.BankAsset.Get())
+			, SequenceIndex(Item.SequenceIndex)
+		{}
+
+		const UAnimBank* BankAsset = nullptr;
+		int32 SequenceIndex = 0;
+
+		bool operator==(const FSequenceRegistryKey& Other) const
+		{
+			return BankAsset == Other.BankAsset && SequenceIndex == Other.SequenceIndex;
+		}
+
+		friend uint32 GetTypeHash(const FSequenceRegistryKey& Key)
+		{
+			return PointerHash(Key.BankAsset, Key.SequenceIndex);
+		}
+	};
+	
+	using FMappingRegistryKey = const UAnimBank*;
+
+	struct FMapping
+	{
+		FMappingRegistryKey RegistryKey = nullptr;
+		FSequenceTransformKeys Keys;
+		FBufferAllocation BufferAllocation;
+
+		uint32 RefCount = 0;
+		uint32 BoneCount = 0;
+	};
+
+	struct FSequence
+	{
+		FSequenceRegistryKey RegistryKey;
+		FSequenceTransformKeys Keys;
+		FBufferAllocation BufferAllocation;
+		FMapping* Mapping = nullptr;
+
+		float PlayLength = 0.0f;
+		uint32 RefCount = 0;
+		uint32 BoneCount = 0;
+		uint32 FrameCount = 0;
+		uint32 KeyCount = 0;
+		bool bLooping = false;
+	};
+
+	struct FSequenceView
+	{
+		FSequenceView(const FSequenceRegistryKey& InKey, const FMapping& InMapping, const FSequence& InSequence)
+			: Key(InKey)
+			, Mapping(InMapping)
+			, Sequence(InSequence)
+		{}
+
+		FSequenceRegistryKey Key;
+		const FMapping& Mapping;
+		const FSequence& Sequence;
+	};
+
+	struct FTrackHistory
+	{
+		static constexpr float Invalid = UE_MAX_FLT;
+
+		void Update(const FAnimBankTrackPackedData& Data, int32 InCurrentSlot, bool& bUpdateCurrent, bool& bUpdatePrevious)
+		{
+			if (ItemIndex != Data.GetItemIndex())
+			{
+				ItemIndex = Data.GetItemIndex();
+				Slots = FVector2f{ Invalid, Invalid };
+			}
+
+			bUpdateCurrent  = Slots[ InCurrentSlot] != Data.GetPosition();
+			bUpdatePrevious = Slots[!InCurrentSlot] != Data.GetPreviousPosition() && Data.HasPreviousPosition();
+			Slots[InCurrentSlot] = Data.GetPosition();
+
+			if (Data.HasPreviousPosition())
+			{
+				Slots[!InCurrentSlot] = Data.GetPreviousPosition();
+			}
+		}
+
+		int32 ItemIndex = INDEX_NONE;
+		FVector2f Slots{ Invalid, Invalid };
+	};
+
+	struct FProxyData
+	{
+		FProxyData(FAnimBankTransformProviderProxy* Proxy);
+
+		const FSequenceView& GetSequenceView(const FAnimBankTrackPackedData& Data) const
+		{
+			return Sequences[FMath::Clamp(Data.GetItemIndex(), 0, Sequences.Num() - 1)];
+		}
+
+		FAnimBankTrackPool& Tracks;
+		TArray<FTrackHistory> TrackHistories;
+		const FInstancedSkinningSceneExtensionProxy& SceneExtensionProxy;
+		TArray<FSequenceView> Sequences;
+		const EBoneTransformStorageMode TransformStorageMode;
+		bool bHasEmptyTracks = false;
+	};
+
+	TMap<FMappingRegistryKey, FMapping> MappingRegistry;
+	TMap<FSequenceRegistryKey, FSequence> SequenceRegistry;
+	TMap<const FSkinningSceneExtensionProxy*, FProxyData> ProxyDataMap;
+
+	TArray<FSequence*> SequencesToUpload;
+	TArray<FMapping*> MappingsToUpload;
 
-RENDERER_API const FSkinningTransformProvider::FProviderId& GetAnimBankProviderId();
\ No newline at end of file
+	FSpanAllocator SequenceTransformAllocator;
+	FRDGAsyncScatterUploadBuffer SequenceTransformUploadBuffer;
+	TPersistentByteAddressBuffer<FSequenceTransform> SequenceTransformPersistentBuffer;
+	
+	FRDGAsyncScatterUploadBuffer SkinningTransformUploadBuffer;
+};
\ No newline at end of file
diff --git a/Engine/Source/Runtime/Renderer/Private/Skinning/SkinningSceneExtension.cpp b/Engine/Source/Runtime/Renderer/Private/Skinning/SkinningSceneExtension.cpp
--- a/Engine/Source/Runtime/Renderer/Private/Skinning/SkinningSceneExtension.cpp
+++ b/Engine/Source/Runtime/Renderer/Private/Skinning/SkinningSceneExtension.cpp
@@ -16,7 +16,7 @@
 DECLARE_STATS_GROUP(TEXT("SkinningSceneExtension"), STATGROUP_SkinningSceneExtension, STATCAT_Advanced);
 
 DECLARE_DWORD_COUNTER_STAT(TEXT("Num Primitives"), STAT_SkinningSceneExtension_NumPrimitives, STATGROUP_SkinningSceneExtension);
-DECLARE_DWORD_COUNTER_STAT(TEXT("Num LOD Changes"), STAT_SkinningSceneExtension_NumLODChanges, STATGROUP_SkinningSceneExtension);
+DECLARE_DWORD_COUNTER_STAT(TEXT("Num Allocations"), STAT_SkinningSceneExtension_NumAllocations, STATGROUP_SkinningSceneExtension);
 DECLARE_DWORD_COUNTER_STAT(TEXT("Num Updates"), STAT_SkinningSceneExtension_NumUpdates, STATGROUP_SkinningSceneExtension);
 
 #if STATS
@@ -33,6 +33,7 @@ DECLARE_SKINNING_MEMORY_STAT(TEXT("Bone Transform Buffer Size"), STAT_SkinningSc
 DECLARE_SKINNING_MEMORY_STAT(TEXT("Bone Hierarchy Buffer Size"), STAT_SkinningSceneExtension_BoneHierarchyBufferSize, STATGROUP_SkinningSceneExtension);
 DECLARE_SKINNING_MEMORY_STAT(TEXT("Bone Object Space Buffer Size"), STAT_SkinningSceneExtension_BoneObjectSpaceBufferSize, STATGROUP_SkinningSceneExtension);
 DECLARE_SKINNING_MEMORY_STAT(TEXT("Bone Transform Upload Size"), STAT_SkinningSceneExtension_BoneTransformUploadSize, STATGROUP_SkinningSceneExtension);
+DECLARE_SKINNING_MEMORY_STAT(TEXT("Bone Transform Allocator Size"), STAT_SkinningSceneExtension_BoneTransformAllocatorSize, STATGROUP_SkinningSceneExtension);
 
 static int32 GSkinningBuffersTransformDataMinSizeBytes = 4 * 1024;
 static FAutoConsoleVariableRef CVarSkinningBuffersTransformDataMinSizeBytes(
@@ -243,7 +244,8 @@ void FSkinningSceneExtension::SetEnabled(bool bEnabled)
 			TransformAllocator.Reset();
 			HeaderDatas.Reset();
 			BatchHeaderDatas.Reset();
-			HeaderDataIndices.Reset();
+			AllocatedHeaderDataIndices.Reset();
+			InstancedHeaderDataIndices.Reset();
 		}
 	}
 }
@@ -321,6 +323,7 @@ void FSkinningSceneExtension::FinishSkinningBufferUpload(FRDGBuilder& GraphBuild
 		INC_MEMORY_STAT_BY(STAT_SkinningSceneExtension_BoneTransformBufferSize, TransformBuffer->Desc.GetSize());
 		INC_MEMORY_STAT_BY(STAT_SkinningSceneExtension_BoneHierarchyBufferSize, BoneHierarchyBuffer->Desc.GetSize());
 		INC_MEMORY_STAT_BY(STAT_SkinningSceneExtension_BoneObjectSpaceBufferSize, BoneObjectSpaceBuffer->Desc.GetSize());
+		INC_MEMORY_STAT_BY(STAT_SkinningSceneExtension_BoneTransformAllocatorSize, MinTransformDataSize * sizeof(FCompressedBoneTransform));
 	}
 
 	if (!Buffers->TransformDataBufferRHI || Buffers->TransformDataBufferRHI != TransformBuffer->GetRHIUnchecked())
@@ -343,6 +346,45 @@ void FSkinningSceneExtension::PerformSkinning(FSkinningSceneParameters& Paramete
 {
 	RDG_EVENT_SCOPE(GraphBuilder, "Skinning");
 
+	FRDGScatterUploadBuilder& ScatterUploadBuilder = *FRDGScatterUploadBuilder::Create(GraphBuilder);
+
+	ON_SCOPE_EXIT
+	{
+		// Vertex factories are updated prior to rendering regardless of whether the skinning is actually executed. They could
+		// technically be part of the updater but a command list task is expensive, so it's batched together with other skinning work.
+		if (!VertexFactoryUpdateList.IsEmpty() || bVertexFactoryFullUpload)
+		{
+			ScatterUploadBuilder.AddFunction([this, TaskVertexFactoryUpdateList = MoveTemp(VertexFactoryUpdateList)] (FRHICommandListBase& RHICmdList)
+			{
+				const auto UpdateVertexFactory = [&](FHeaderData& Data)
+				{
+					Data.MeshObject->UpdateSceneExtensionHeader(RHICmdList, Data.Pack(), Buffers->TransformDataBufferSRV);
+					Data.bVertexFactoryDirty = false;
+				};
+
+				if (bVertexFactoryFullUpload)
+				{
+					for (FHeaderData& Data : HeaderDatas)
+					{
+						UpdateVertexFactory(Data);
+					}
+				}
+				else
+				{
+					for (int32 HeaderDataIndex : TaskVertexFactoryUpdateList)
+					{
+						UpdateVertexFactory(HeaderDatas[HeaderDataIndex]);
+					}
+				}
+
+				bVertexFactoryFullUpload = false;
+			});
+		}
+
+		RDG_EVENT_SCOPE(GraphBuilder, "ScatterUpload");
+		ScatterUploadBuilder.Execute(GraphBuilder);
+	};
+
 	if (!GSkinningTransformProviders)
 	{
 		return;
@@ -384,14 +426,6 @@ void FSkinningSceneExtension::PerformSkinning(FSkinningSceneParameters& Paramete
 			EDirtyBoneTransforms DirtyBoneTransforms;
 		};
 
-		TArray<FHeaderData*, FConcurrentLinearArrayAllocator> VertexFactoryUpdates;
-
-		if (!bVertexFactoryFullUpload)
-		{
-			VertexFactoryUpdates.Reserve(VertexFactoryUpdateRequests);
-			VertexFactoryUpdateRequests = 0;
-		}
-
 		// Skeleton
 		if (BatchHeaderDatas.Num() > 0)
 		{
@@ -432,11 +466,6 @@ void FSkinningSceneExtension::PerformSkinning(FSkinningSceneParameters& Paramete
 				Offsets[TotalBatches].HierarchyOffset = Header.HierarchyBufferOffset;
 				Offsets[TotalBatches].DirtyBoneTransforms = EDirtyBoneTransforms::All;
 
-				if (!bVertexFactoryFullUpload && Header.bVertexFactoryDirty)
-				{
-					VertexFactoryUpdates.Emplace(&Header);
-				}
-
 				++TotalBatches;
 			}
 
@@ -487,6 +516,7 @@ void FSkinningSceneExtension::PerformSkinning(FSkinningSceneParameters& Paramete
 				Batches,
 				ViewFamily,
 				GraphBuilder,
+				ScatterUploadBuilder,
 				Parameters.BoneTransforms->GetParent(),
 				Parameters.BoneHierarchy
 			);
@@ -495,7 +525,7 @@ void FSkinningSceneExtension::PerformSkinning(FSkinningSceneParameters& Paramete
 		}
 
 		// Primitive
-		if (HeaderDataIndices.Num() > 0)
+		if (AllocatedHeaderDataIndices.Num() > 0)
 		{
 			TArray<uint8, FConcurrentLinearArrayAllocator> PrimitivesToRangeIndex;
 			PrimitivesToRangeIndex.AddUninitialized(HeaderDatas.Num());
@@ -504,21 +534,21 @@ void FSkinningSceneExtension::PerformSkinning(FSkinningSceneParameters& Paramete
 			PrimitiveRanges.Reserve(PrimitiveProviderIds.Num());
 			ResetRanges(PrimitiveProviderIds, PrimitiveRanges);
 
-			TArrayView<FPrimitiveSceneInfo*> Primitives = GraphBuilder.AllocPODArrayView<FPrimitiveSceneInfo*>(HeaderDataIndices.Num());
-			TArrayView<FSkinningSceneExtensionProxy*> Proxies = GraphBuilder.AllocPODArrayView<FSkinningSceneExtensionProxy*>(HeaderDataIndices.Num());
-			TArrayView<FOffsets> Offsets = GraphBuilder.AllocPODArrayView<FOffsets>(HeaderDataIndices.Num());
+			TArrayView<FPrimitiveSceneInfo*> Primitives = GraphBuilder.AllocPODArrayView<FPrimitiveSceneInfo*>(AllocatedHeaderDataIndices.Num());
+			TArrayView<FSkinningSceneExtensionProxy*> Proxies = GraphBuilder.AllocPODArrayView<FSkinningSceneExtensionProxy*>(AllocatedHeaderDataIndices.Num());
+			TArrayView<FOffsets> Offsets = GraphBuilder.AllocPODArrayView<FOffsets>(AllocatedHeaderDataIndices.Num());
 
 			uint32 TotalOffset = 0;
 
 			uint32 PrimitiveCount = 0;
-			for (const int32 HeaderDataIndex : HeaderDataIndices)
+			for (const int32 HeaderDataIndex : AllocatedHeaderDataIndices)
 			{
 				FHeaderData& Header = HeaderDatas[HeaderDataIndex];
 				Header.Validate();
 
 				int32 RangeIndex = 0;
 
-				if (!EnumHasAnyFlags(Header.DirtyBoneTransforms, EDirtyBoneTransforms::All) && !Header.bAlwaysUpdateTransforms)
+				if (!EnumHasAnyFlags(Header.DirtyBoneTransforms, EDirtyBoneTransforms::All))
 				{
 					continue;
 				}
@@ -546,11 +576,6 @@ void FSkinningSceneExtension::PerformSkinning(FSkinningSceneParameters& Paramete
 
 				Header.DirtyBoneTransforms = EDirtyBoneTransforms::None;
 
-				if (!bVertexFactoryFullUpload && Header.bVertexFactoryDirty)
-				{
-					VertexFactoryUpdates.Emplace(&Header);
-				}
-
 				++PrimitiveCount;
 			}
 
@@ -584,41 +609,13 @@ void FSkinningSceneExtension::PerformSkinning(FSkinningSceneParameters& Paramete
 				TConstArrayView<FSkeletonBatch>(),
 				ViewFamily,
 				GraphBuilder,
+				ScatterUploadBuilder,
 				Parameters.BoneTransforms->GetParent(),
 				Parameters.BoneHierarchy
 			);
 
 			TransformProvider->Broadcast(PrimitiveRanges, Context);
 		}
-
-		if (!VertexFactoryUpdates.IsEmpty() || bVertexFactoryFullUpload)
-		{
-			GraphBuilder.AddCommandListSetupTask([this, VertexFactoryUpdates = MoveTemp(VertexFactoryUpdates)](FRHICommandListBase& RHICmdList)
-			{
-				const auto UpdateVertexFactory = [&](FHeaderData& Data)
-				{
-					Data.MeshObject->UpdateSceneExtensionHeader(RHICmdList, Data.Pack(), Buffers->TransformDataBufferSRV);
-					Data.bVertexFactoryDirty = 0;
-				};
-
-				if (bVertexFactoryFullUpload)
-				{
-					for (FHeaderData& Data : HeaderDatas)
-					{
-						UpdateVertexFactory(Data);
-					}
-				}
-				else
-				{
-					for (FHeaderData* HeaderData : VertexFactoryUpdates)
-					{
-						UpdateVertexFactory(*HeaderData);
-					}
-				}
-
-				bVertexFactoryFullUpload = false;
-			});
-		}
 	}
 }
 
@@ -667,7 +664,7 @@ bool FSkinningSceneExtension::ProcessBufferDefragmentation()
 	HierarchyAllocator.Reset();
 	TransformAllocator.Reset();
 	BatchHeaderDatas.Reset();
-	HeaderDataIndices.Reset();
+	AllocatedHeaderDataIndices.Reset();
 
 	for (auto& Data : HeaderDatas)
 	{
@@ -759,6 +756,8 @@ void FSkinningSceneExtension::FUpdater::End()
 	// Ensure these tasks finish before we fall out of scope.
 	// NOTE: This should be unnecessary if the updater shares the graph builder's lifetime but we don't enforce that
 	SceneData->SyncAllTasks();
+
+	SceneData->LastUpdateFrameNumber = GFrameNumberRenderThread;
 }
 
 void FSkinningSceneExtension::FUpdater::PreSceneUpdate(FRDGBuilder& GraphBuilder, const FScenePreUpdateChangeSet& ChangeSet, FSceneUniformBuffer& SceneUniforms)
@@ -805,13 +804,21 @@ void FSkinningSceneExtension::FUpdater::PreSceneUpdate(FRDGBuilder& GraphBuilder
 						{
 							SceneData->TransformAllocator.Free(Data.TransformBufferOffset, Data.TransformBufferCount);
 						}
+
+						SceneData->AllocatedHeaderDataIndices.Remove(PersistentIndex.Index);
+
+						if (Data.bIsInstanced)
+						{
+							SceneData->InstancedHeaderDataIndices.Remove(PersistentIndex.Index);
+						}
+
+						if (Data.bVertexFactoryDirty)
+						{
+							SceneData->VertexFactoryUpdateList.Remove(PersistentIndex.Index);
+						}
 					}
 
 					SceneData->HeaderDatas.RemoveAt(PersistentIndex.Index);
-					if (!Data.bIsBatched)
-					{
-						SceneData->HeaderDataIndices.Remove(PersistentIndex.Index);
-					}
 				}
 			}
 
@@ -881,12 +888,19 @@ void FSkinningSceneExtension::FUpdater::PostSceneUpdate(FRDGBuilder& GraphBuilde
 					NewHeader.MaxObjectSpaceCount         = Proxy->GetMaxBoneObjectSpaceCount();
 					NewHeader.MaxInfluenceCount           = Proxy->GetMaxBoneInfluenceCount();
 					NewHeader.UniqueAnimationCount        = Proxy->GetUniqueAnimationCount();
+					NewHeader.bIsBatched                  = Proxy->UseSkeletonBatching();
+					NewHeader.bIsInstanced                = Proxy->UseInstancing();
+					NewHeader.bIsRefPose                  = NewHeader.ProviderId == RefPoseProviderId;
 					NewHeader.bHasScale                   = Proxy->HasScale();
-					NewHeader.bAlwaysUpdateTransforms     = Proxy->UseInstancing();
 
 					MeshObject->SkinningSceneExtensionPrivateData.CachedPrimitiveIndex = PersistentIndex;
 
 					SceneData->HeaderDatas.EmplaceAt(PersistentIndex, NewHeader);
+
+					if (NewHeader.bIsInstanced && !NewHeader.bIsBatched)
+					{
+						SceneData->InstancedHeaderDataIndices.Emplace(PersistentIndex);
+					}
 				}
 			},
 			SceneData->TaskHandles[FreeBufferSpaceTask],
@@ -896,6 +910,35 @@ void FSkinningSceneExtension::FUpdater::PostSceneUpdate(FRDGBuilder& GraphBuilde
 	}
 }
 
+bool FSkinningSceneExtension::GetHeaderFromMeshObject(const FSkeletalMeshObject* MeshObject, FHeaderData*& OutHeaderData, int32& OutPrimitiveIndex)
+{
+	// This skeletal mesh is not tracked by the scene extension.
+	if (MeshObject->SkinningSceneExtensionPrivateData.CachedPrimitiveIndex == INDEX_NONE)
+	{
+		return false;
+	}
+
+	const int32 PrimitiveIndex = MeshObject->SkinningSceneExtensionPrivateData.CachedPrimitiveIndex;
+	if (!HeaderDatas.IsValidIndex(PrimitiveIndex))
+	{
+		return false;
+	}
+
+	// We are unable to reset the mesh object primitive index on removal (the lifetimes of mesh objects don't match primitives,
+	// so it's unsafe to dereference it inside of the PreSceneUpdate on removal), so instead verify that the pointers match here.
+	// This avoids cases where the primitive gets removed from the scene but the mesh object is still valid and updating with a
+	// stale index. An explicit map would be cleaner but it's significantly slower and these updates are frequent.
+	FHeaderData& HeaderData = HeaderDatas[PrimitiveIndex];
+	if (HeaderData.MeshObject != MeshObject)
+	{
+		return false;
+	}
+
+	OutPrimitiveIndex = PrimitiveIndex;
+	OutHeaderData = &HeaderData;
+	return true;
+}
+
 void FSkinningSceneExtension::FUpdater::PostMeshUpdate(FRDGBuilder& GraphBuilder, const TConstArrayView<FPrimitiveSceneInfo*>& UpdatedSceneInfoList, const FSkeletalMeshUpdater::FSubmitData* SkeletalMeshSubmitData)
 {
 	if (!SceneData->IsEnabled())
@@ -913,81 +956,100 @@ void FSkinningSceneExtension::FUpdater::PostMeshUpdate(FRDGBuilder& GraphBuilder
 			TRACE_CPUPROFILER_EVENT_SCOPE(Skinning::InitUpdateList);
 
 			int32 NumUpdates = 0;
-			int32 NumLODChanges = 0;
+			int32 NumAllocations = 0;
 
-			const auto UpdateLambda = [this, &NumLODChanges] (FHeaderData& HeaderData, int32 PrimitiveIndex, int32 LODIndex)
+			const auto UpdateLambda = [this, &NumAllocations, &NumUpdates] (FHeaderData& HeaderData, int32 PrimitiveIndex, int32 LODIndex, bool bFlipTransformSlots)
 			{
+				HeaderData.UpdateCounter        = SceneData->UpdateCounter;
 				HeaderData.DirtyBoneTransforms |= EDirtyBoneTransforms::Current;
-				HeaderData.UpdateCounter = SceneData->UpdateCounter;
 
-				if (!HeaderData.bAlwaysUpdateTransforms)
+				if (bFlipTransformSlots && !HeaderData.bIsBatched)
 				{
 					HeaderData.CurrentTransformSlot = !HeaderData.CurrentTransformSlot;
 				}
 
-				if (HeaderData.Proxy->SetLOD(LODIndex))
-				{
-					NumLODChanges++;
+				const bool bLODChanged                  = LODIndex != INDEX_NONE ? HeaderData.Proxy->SetLOD(LODIndex) : false;
+				const bool bUniqueAnimationCountChanged = HeaderData.UniqueAnimationCount != HeaderData.Proxy->GetUniqueAnimationCount();
 
-					HeaderData.MaxTransformCount = HeaderData.Proxy->GetMaxBoneTransformCount();
-					HeaderData.MaxHierarchyCount = HeaderData.Proxy->GetMaxBoneHierarchyCount();
-					HeaderData.MaxObjectSpaceCount = HeaderData.Proxy->GetMaxBoneObjectSpaceCount();
+				if (bLODChanged || bUniqueAnimationCountChanged)
+				{
+					NumAllocations += bLODChanged ? 1 : 0;
+					HeaderData.UniqueAnimationCount = HeaderData.Proxy->GetUniqueAnimationCount();
+					HeaderData.MaxTransformCount    = HeaderData.Proxy->GetMaxBoneTransformCount();
+					HeaderData.MaxHierarchyCount    = HeaderData.Proxy->GetMaxBoneHierarchyCount();
+					HeaderData.MaxObjectSpaceCount  = HeaderData.Proxy->GetMaxBoneObjectSpaceCount();
 
 					AllocationUpdateList.Emplace(PrimitiveIndex);
 				}
 
 				HeaderDataUpdateList.Emplace(PrimitiveIndex);
+				NumUpdates++;
 			};
 
 			TConstArrayView<FSkeletalMeshUpdateEvent> UpdateEvents = SkeletalMeshSubmitData->GetUpdateEvents();
 			AllocationUpdateList.Reserve(UpdatedSceneInfoList.Num() + UpdateEvents.Num());
 			HeaderDataUpdateList.Reserve(UpdatedSceneInfoList.Num() + UpdateEvents.Num());
 
+			// Process update events from the skeletal mesh updater first, including LOD change events.
 			for (FSkeletalMeshUpdateEvent UpdateEvent : UpdateEvents)
 			{
-				// This skeletal mesh is not tracked by the scene extension.
-				if (UpdateEvent.MeshObject->SkinningSceneExtensionPrivateData.CachedPrimitiveIndex == INDEX_NONE)
+				FHeaderData* HeaderData = nullptr;
+				int32 PrimitiveIndex;
+				if (SceneData->GetHeaderFromMeshObject(UpdateEvent.MeshObject, HeaderData, PrimitiveIndex))
 				{
-					continue;
+					const bool bFlipTransformSlots = true;
+					UpdateLambda(*HeaderData, PrimitiveIndex, UpdateEvent.LODIndex, bFlipTransformSlots);
+					HeaderData->UpdateCounter = SceneData->UpdateCounter;
 				}
+			}
 
-				const int32 PrimitiveIndex = UpdateEvent.MeshObject->SkinningSceneExtensionPrivateData.CachedPrimitiveIndex;
-				if (!SceneData->HeaderDatas.IsValidIndex(PrimitiveIndex))
-				{
-					continue;
-				}
+			const auto GetFirstUpdateLOD = [] (const FHeaderData& HeaderData)
+			{
+				// We only need to pull the LOD on first update if the skeletal mesh updater failed to send a mesh object update event
+				// alongside the primitive AddToScene request. Technically it should always have an update queued but this relies on the high
+				// level code managing mesh object updates and scene attachment to be ordered perfectly and enforcing this is impossible.
+				// Instead, we allow for attachment without a mesh update event and instead pull the current LOD on the mesh object if no
+				// update event was provided. That way, if the mesh object was updated and then the AttachToScene call occurs later, it pulls
+				// the correct LOD and doesn't crash.
+
+				return HeaderData.IsFirstUpdate() ? HeaderData.MeshObject->GetLOD() : INDEX_NONE;
+			};
 
-				// We are unable to reset the mesh object primitive index on removal (the lifetimes of mesh objects don't match primitives,
-				// so it's unsafe to dereference it inside of the PreSceneUpdate on removal), so instead verify that the pointers match here.
-				// This avoids cases where the primitive gets removed from the scene but the mesh object is still valid and updating with a
-				// stale index. An explicit map would be cleaner but it's significantly slower and these updates are frequent.
-				FHeaderData& HeaderData = SceneData->HeaderDatas[PrimitiveIndex];
-				if (HeaderData.MeshObject != UpdateEvent.MeshObject)
+			// Update all unbatched instanced meshes that didn't otherwise receive an LOD update. These don't receive update events
+			// from the skeletal mesh updater in the same fashion, so they are manually ticked. Only do this once per frame as we might
+			// received multiple back-to-back scene updates from the high level.
+			if (!SceneData->HasUpdatedThisFrame())
+			{
+				for (int32 PrimitiveIndex : SceneData->InstancedHeaderDataIndices)
 				{
-					continue;
-				}
+					FHeaderData& HeaderData = SceneData->HeaderDatas[PrimitiveIndex];
 
-				UpdateLambda(HeaderData, PrimitiveIndex, UpdateEvent.LODIndex);
-				HeaderData.UpdateCounter = SceneData->UpdateCounter;
-				NumUpdates++;
+					if (HeaderData.UpdateCounter != SceneData->UpdateCounter)
+					{
+						const bool bFlipTransformSlots = true;
+						UpdateLambda(HeaderData, PrimitiveIndex, GetFirstUpdateLOD(HeaderData), bFlipTransformSlots);
+					}
+				}
 			}
-			
+
+			// Finally, update meshes that are new to the scene or were invalidated for other reasons.
 			for (FPrimitiveSceneInfo* PrimitiveSceneInfo : UpdatedSceneInfoList)
 			{
-				const int32 Index = PrimitiveSceneInfo->GetPersistentIndex().Index;
-				if (SceneData->HeaderDatas.IsValidIndex(Index))
+				const int32 PrimitiveIndex = PrimitiveSceneInfo->GetPersistentIndex().Index;
+				if (SceneData->HeaderDatas.IsValidIndex(PrimitiveIndex))
 				{
-					FHeaderData& HeaderData = SceneData->HeaderDatas[Index];
+					FHeaderData& HeaderData = SceneData->HeaderDatas[PrimitiveIndex];
 
 					if (HeaderData.UpdateCounter != SceneData->UpdateCounter)
 					{
-						UpdateLambda(HeaderData, Index, HeaderData.MeshObject->GetLOD());
+						const bool bFlipTransformSlots = false;
+						UpdateLambda(HeaderData, PrimitiveIndex, GetFirstUpdateLOD(HeaderData), bFlipTransformSlots);
 					}
 				}
 			}
 
 			INC_DWORD_STAT_BY(STAT_SkinningSceneExtension_NumUpdates, NumUpdates);
-			INC_DWORD_STAT_BY(STAT_SkinningSceneExtension_NumLODChanges, NumLODChanges);
+			INC_DWORD_STAT_BY(STAT_SkinningSceneExtension_NumAllocations, NumAllocations);
 		},
 		MakeArrayView(
 			{
@@ -1020,6 +1082,8 @@ void FSkinningSceneExtension::FUpdater::PostMeshUpdate(FRDGBuilder& GraphBuilder
 
 			if (SceneData->BatchHeaderDatas.Contains(BatchKey))
 			{
+				check(Data.bIsBatched == Data.Proxy->UseSkeletonBatching());
+
 				const FHeaderData& SrcHeaderData = SceneData->BatchHeaderDatas[BatchKey];
 
 				Data.ObjectSpaceBufferOffset	= SrcHeaderData.ObjectSpaceBufferOffset;
@@ -1028,7 +1092,6 @@ void FSkinningSceneExtension::FUpdater::PostMeshUpdate(FRDGBuilder& GraphBuilder
 				Data.HierarchyBufferCount		= SrcHeaderData.HierarchyBufferCount;
 				Data.TransformBufferOffset		= SrcHeaderData.TransformBufferOffset;
 				Data.TransformBufferCount		= SrcHeaderData.TransformBufferCount;
-				Data.bIsBatched					= true;
 
 			#if DO_CHECK
 				{
@@ -1115,8 +1178,6 @@ void FSkinningSceneExtension::FUpdater::PostMeshUpdate(FRDGBuilder& GraphBuilder
 
 		if (Data.Proxy->UseSkeletonBatching())
 		{
-			Data.bIsBatched = true;
-
 			const FSkeletonBatchKey BatchKey
 			{
 			#if ENABLE_SKELETON_DEBUG_NAME
@@ -1130,7 +1191,7 @@ void FSkinningSceneExtension::FUpdater::PostMeshUpdate(FRDGBuilder& GraphBuilder
 		}
 		else
 		{
-			SceneData->HeaderDataIndices.Add(HeaderDataIndex);
+			SceneData->AllocatedHeaderDataIndices.Add(HeaderDataIndex);
 		}
 	};
 
@@ -1188,7 +1249,11 @@ void FSkinningSceneExtension::FUpdater::PostMeshUpdate(FRDGBuilder& GraphBuilder
 		check(SceneData->Uploader.IsValid()); // Sanity check
 		SceneData->Uploader->HeaderDataUploader.Add(Data.Pack(), PersistentIndex);
 
-		Data.bVertexFactoryDirty = 1;
+		if (!Data.bIsBatched && !Data.bVertexFactoryDirty)
+		{
+			Data.bVertexFactoryDirty = 1;
+			SceneData->VertexFactoryUpdateList.Emplace(PersistentIndex);
+		}
 	};
 
 	// Kick off the header data upload task (synced when accessing the buffer)
@@ -1288,7 +1353,7 @@ void FSkinningSceneExtension::FUpdater::PostMeshUpdate(FRDGBuilder& GraphBuilder
 		FCompressedBoneTransform* DstPreviousBoneTransformsPtr = DstCurrentBoneTransformsPtr + Data.MaxTransformCount;
 		const uint32 StridedPtrStep = Data.MaxTransformCount * 2u;
 
-		for (int32 UniqueAnimation = 0; UniqueAnimation < Data.UniqueAnimationCount; ++UniqueAnimation)
+		for (uint32 UniqueAnimation = 0; UniqueAnimation < Data.UniqueAnimationCount; ++UniqueAnimation)
 		{
 			for (int32 TransformIndex = 0; TransformIndex < Data.MaxTransformCount; ++TransformIndex)
 			{
@@ -1634,10 +1699,9 @@ void FSkinningSceneExtension::ProvideAnimRuntimeTransforms(FSkinningTransformPro
 	}
 
 	FRDGBuilder& GraphBuilder = Context.GraphBuilder;
-	FRDGAsyncScatterUploadBuffer TransformUploadBuffer;
-	FRDGScatterUploadBuilder* Builder = FRDGScatterUploadBuilder::Create(GraphBuilder);
+	FRDGAsyncScatterUploadBuffer& TransformUploadBuffer = *GraphBuilder.AllocObject<FRDGAsyncScatterUploadBuffer>();
 
-	Builder->AddPass(GraphBuilder, TransformUploadBuffer, Context.TransformBuffer, GlobalTransformCount, sizeof(FCompressedBoneTransform), TEXT("Skinning.AnimTransforms"),
+	Context.ScatterUploadBuilder.AddPass(GraphBuilder, TransformUploadBuffer, Context.TransformBuffer, GlobalTransformCount, sizeof(FCompressedBoneTransform), TEXT("Skinning.AnimTransforms"),
 		[Indirections = Context.Indirections, Proxies = Context.Proxies, &ViewFamily = Context.ViewFamily, GlobalTransformCount] (FRDGScatterUploader& ScatterUploader)
 	{
 		TRACE_CPUPROFILER_EVENT_SCOPE(FSkinningSceneExtension::ProvideAnimRuntimeTransformsTask);
@@ -1733,8 +1797,6 @@ void FSkinningSceneExtension::ProvideAnimRuntimeTransforms(FSkinningTransformPro
 
 		INC_MEMORY_STAT_BY(STAT_SkinningSceneExtension_BoneTransformUploadSize, NumBonesTotal * sizeof(FCompressedBoneTransform));
 	});
-
-	Builder->Execute(GraphBuilder);
 }
 
 // TODO: these are prototype macros for how we might expose SceneUB for direct binding. 
@@ -1752,4 +1814,4 @@ void FSkinningSceneExtension::ProvideAnimRuntimeTransforms(FSkinningTransformPro
 	TSceneUniformBufferMemberRegistration<StructType> SceneUB::FieldName { TEXT(#FieldName), DefaultValueFactoryType }; \
 	IMPLEMENT_STATIC_UNIFORM_BUFFER_SCENE_UB(StructType, SceneUbEx##FieldName)
 
-IMPLEMENT_SCENE_UB_STRUCT_EX(FSkinningSceneParameters, Skinning, GetDefaultSkinningParameters);
\ No newline at end of file
+IMPLEMENT_SCENE_UB_STRUCT_EX(FSkinningSceneParameters, Skinning, GetDefaultSkinningParameters);
diff --git a/Engine/Source/Runtime/Renderer/Private/Skinning/SkinningSceneExtension.h b/Engine/Source/Runtime/Renderer/Private/Skinning/SkinningSceneExtension.h
--- a/Engine/Source/Runtime/Renderer/Private/Skinning/SkinningSceneExtension.h
+++ b/Engine/Source/Runtime/Renderer/Private/Skinning/SkinningSceneExtension.h
@@ -117,18 +117,24 @@ private:
 		uint32 HierarchyBufferCount              = 0;
 		uint32 TransformBufferOffset             = INDEX_NONE;
 		uint32 TransformBufferCount              = 0;
+		uint32 UniqueAnimationCount              = 1;
 		uint16 MaxTransformCount                 = 0;
 		uint16 MaxHierarchyCount                 = 0;
 		uint16 MaxObjectSpaceCount               = 0;
 		EDirtyBoneTransforms DirtyBoneTransforms = EDirtyBoneTransforms::All;
 		uint8  MaxInfluenceCount                 = 0;
-		uint8  UniqueAnimationCount              = 1;
 		uint8  bHasScale : 1                     = false;
 		uint8  bIsBatched : 1                    = false;
-		uint8  bAlwaysUpdateTransforms : 1       = false;
+		uint8  bIsInstanced : 1                  = false;
+		uint8  bIsRefPose : 1                    = false;
 		uint8  bVertexFactoryDirty : 1           = false;
 		uint8  CurrentTransformSlot : 1          = 0;
 
+		bool IsFirstUpdate() const
+		{
+			return UpdateCounter == 0;
+		}
+
 		FSkinningHeader Pack() const
 		{
 			Validate();
@@ -139,8 +145,8 @@ private:
 			Output.ObjectSpaceBufferOffset	= ObjectSpaceBufferOffset != INDEX_NONE ? ObjectSpaceBufferOffset : 0;
 			Output.MaxTransformCount		= MaxTransformCount;
 			Output.MaxInfluenceCount		= MaxInfluenceCount;
-			Output.UniqueAnimationCount		= UniqueAnimationCount;
 			Output.bHasScale				= bHasScale;
+			Output.bIsRefPose				= bIsRefPose;
 			Output.CurrentTransformSlot		= CurrentTransformSlot;
 			return Output;
 		}
@@ -199,18 +205,28 @@ private:
 	// Wait for tasks that modify HeaderData - after this the size and main fields do not change.
 	void WaitForHeaderDataUpdateTasks() const;
 
+	bool HasUpdatedThisFrame() const
+	{
+		return LastUpdateFrameNumber == GFrameNumberRenderThread;
+	}
+
 private:
+	bool GetHeaderFromMeshObject(const FSkeletalMeshObject* MeshObject, FHeaderData*& OutHeaderData, int32& OutPrimitiveIndex);
+
 	FSpanAllocator ObjectSpaceAllocator;
 	FSpanAllocator HierarchyAllocator;
 	FSpanAllocator TransformAllocator;
 	TSparseArray<FHeaderData> HeaderDatas;
-	TSet<int32> HeaderDataIndices;
+	TSet<int32> AllocatedHeaderDataIndices;
+	TSet<int32> InstancedHeaderDataIndices;
 	TMap<FSkeletonBatchKey, FHeaderData> BatchHeaderDatas;
 	TUniquePtr<FBuffers> Buffers;
 	TUniquePtr<FUploader> Uploader;
 	TStaticArray<UE::Tasks::FTask, NumTasks> TaskHandles;
 	uint32 UpdateCounter = 0;
+	uint32 LastUpdateFrameNumber = 0;
 	int32 VertexFactoryUpdateRequests = 0;
+	TArray<int32> VertexFactoryUpdateList;
 	bool bVertexFactoryFullUpload = false;
 
 	bool Tick(float DeltaTime);
diff --git a/Engine/Source/Runtime/Renderer/Private/Skinning/SkinningTransformProvider.h b/Engine/Source/Runtime/Renderer/Private/Skinning/SkinningTransformProvider.h
--- a/Engine/Source/Runtime/Renderer/Private/Skinning/SkinningTransformProvider.h
+++ b/Engine/Source/Runtime/Renderer/Private/Skinning/SkinningTransformProvider.h
@@ -58,6 +58,7 @@ enum class EDirtyBoneTransforms : uint8
 ENUM_CLASS_FLAGS(EDirtyBoneTransforms);
 
 class FSkinningSceneExtensionProxy;
+class FRDGScatterUploadBuilder;
 
 class FSkinningTransformProvider : public ISceneExtension
 {
@@ -99,6 +100,7 @@ public:
 			const TConstArrayView<FSkeletonBatch> InSkeletonBatches,
 			const FSceneViewFamily& InViewFamily,
 			FRDGBuilder& InGraphBuilder,
+			FRDGScatterUploadBuilder& InScatterUploadBuilder,
 			FRDGBufferRef InTransformBuffer,
 			FRDGBufferSRVRef InHierarchyBufferSRV
 		)
@@ -107,6 +109,7 @@ public:
 		, Indirections(InIndirections)
 		, SkeletonBatches(InSkeletonBatches)
 		, GraphBuilder(InGraphBuilder)
+		, ScatterUploadBuilder(InScatterUploadBuilder)
 		, TransformBuffer(InTransformBuffer)
 		, HierarchyBufferSRV(InHierarchyBufferSRV)
 		, ViewFamily(InViewFamily)
@@ -119,6 +122,7 @@ public:
 		TConstArrayView<FSkeletonBatch> SkeletonBatches;
 
 		FRDGBuilder& GraphBuilder;
+		FRDGScatterUploadBuilder& ScatterUploadBuilder;
 		FRDGBufferRef TransformBuffer;
 		FRDGBufferSRVRef HierarchyBufferSRV;
 
