When Lumen GI is enabled but Lumen Reflections and SSR are not, use Rough Reflections from GI instead of falling back to unshadowed skylight via ReflectionEnvironmentAndSky
Fixed incorrect ShouldPrecachePermutation logic for FDiffuseIndirectCompositePS causing an ensure
Fixed setup of BlueNoise parameters when no feature was enabled that would have loaded the engine texture

#Synced-CL 48620257

diff --git a/Engine/Source/Runtime/Renderer/Private/IndirectLightRendering.cpp b/Engine/Source/Runtime/Renderer/Private/IndirectLightRendering.cpp
--- a/Engine/Source/Runtime/Renderer/Private/IndirectLightRendering.cpp
+++ b/Engine/Source/Runtime/Renderer/Private/IndirectLightRendering.cpp
@@ -168,7 +168,7 @@ class FDiffuseIndirectCompositePS : public FGlobalShader
 		FPermutationDomain PermutationVector(Parameters.PermutationId);
 
 		// Other final gathers should upscale in their temporal filter
-		if (PermutationVector.Get<FUpscaleLumenGIDim>() != Lumen::UseIrradianceFieldGather())
+		if (PermutationVector.Get<FUpscaleLumenGIDim>() && !Lumen::UseIrradianceFieldGather())
 		{
 			return EShaderPermutationPrecacheRequest::NotUsed;
 		}
@@ -1158,11 +1158,10 @@ void FDeferredShadingSceneRenderer::RenderDiffuseIndirectAndAmbientOcclusion(
 					bRunSSRTiledComposite ? &ScreenSpaceReflectionTileClassification.TiledReflection : nullptr);
 				OutTextures.Textures[3] = SSRDenoiserInputs.Color;
 			}
-			else
+			else if (ViewPipelineState.DiffuseIndirectMethod == EDiffuseIndirectMethod::Lumen)
 			{
-				// Remove Lumen rough specular when ReflectionsMethod != EReflectionsMethod::Lumen, so that we won't get double specular when other reflection methods are used
-				OutTextures.Textures[2] = SystemTextures.Black;
-				OutTextures.Textures[3] = SystemTextures.Black;
+				// No Lumen Reflections and no SSR - use Lumen Rough Reflections only.
+				OutTextures.Textures[3] = OutTextures.Textures[2];
 			}
 
 			if (!bDoComposite)
@@ -1379,9 +1378,14 @@ void FDeferredShadingSceneRenderer::RenderDiffuseIndirectAndAmbientOcclusion(
 			LumenReflections::SetupCompositeParameters(View, ViewPipelineState.ReflectionsMethod, PSParameters->ReflectionsCompositeParameters);
 			PSParameters->ScreenBentNormalParameters = ScreenBentNormalParameters;
 			PSParameters->LumenUpsampleParameters = LumenUpsampleParameters;
-			FBlueNoise BlueNoise = GetBlueNoiseGlobalParameters();
-			TUniformBufferRef<FBlueNoise> BlueNoiseUniformBuffer = CreateUniformBufferImmediate(BlueNoise, EUniformBufferUsage::UniformBuffer_SingleDraw);
-			PSParameters->BlueNoise = BlueNoiseUniformBuffer;
+
+			if (ViewPipelineState.DiffuseIndirectMethod == EDiffuseIndirectMethod::Lumen)
+			{
+				FBlueNoise BlueNoise = GetBlueNoiseGlobalParameters();
+				TUniformBufferRef<FBlueNoise> BlueNoiseUniformBuffer = CreateUniformBufferImmediate(BlueNoise, EUniformBufferUsage::UniformBuffer_SingleDraw);
+				PSParameters->BlueNoise = BlueNoiseUniformBuffer;
+			}
+
 			PSParameters->bLumenSupportBackfaceDiffuse = ViewPipelineState.DiffuseIndirectMethod == EDiffuseIndirectMethod::Lumen && DenoiserOutputs.Textures[1] != SystemTextures.Black;
 			PSParameters->bLumenReflectionInputIsSSR = ViewPipelineState.DiffuseIndirectMethod == EDiffuseIndirectMethod::Lumen && ViewPipelineState.ReflectionsMethod == EReflectionsMethod::SSR;
 			PSParameters->LumenGIDownsampleFactor = LumenUpsampleParameters.DownsampleFactor;
@@ -2101,10 +2105,9 @@ void FDeferredShadingSceneRenderer::RenderDeferredReflectionsAndSkyLighting(
 		bool bRunSSRTiledComposite = false;
 
 		FRDGTextureRef ReflectionsColor = nullptr;
-		if ((ViewPipelineState.ReflectionsMethod == EReflectionsMethod::Lumen && ViewPipelineState.DiffuseIndirectMethod == EDiffuseIndirectMethod::Lumen)
-			|| (ViewPipelineState.DiffuseIndirectMethod == EDiffuseIndirectMethod::Lumen && ViewPipelineState.ReflectionsMethod == EReflectionsMethod::SSR))
+		if (ViewPipelineState.DiffuseIndirectMethod == EDiffuseIndirectMethod::Lumen)
 		{
-			// Specular was already comped with FDiffuseIndirectCompositePS
+			// Specular was already composited in FDiffuseIndirectCompositePS
 			continue;
 		}
 		// Standalone Lumen Reflections using HWRT and hit lighting
