Changing engine defaults for base PSO precaching to precache less permutations which by default are not used and cause a lot of extra permutations on base pass or depth rendering PSOs

#Synced-CL 48894724

diff --git a/Engine/Plugins/FX/Niagara/Source/Niagara/Private/NiagaraRendererProperties.cpp b/Engine/Plugins/FX/Niagara/Source/Niagara/Private/NiagaraRendererProperties.cpp
--- a/Engine/Plugins/FX/Niagara/Source/Niagara/Private/NiagaraRendererProperties.cpp
+++ b/Engine/Plugins/FX/Niagara/Source/Niagara/Private/NiagaraRendererProperties.cpp
@@ -33,7 +33,7 @@
 
 #define LOCTEXT_NAMESPACE "UNiagaraRendererProperties"
 
-static int GNiagaraPSOPrecacheReverseCulling = 1;
+static int GNiagaraPSOPrecacheReverseCulling = 0;
 static FAutoConsoleVariableRef CVarPSOPrecacheProjectedShadows(
 	TEXT("fx.Niagara.PSOPrecache.ReverseCulling"),
 	GNiagaraPSOPrecacheReverseCulling,
diff --git a/Engine/Source/Runtime/Renderer/Private/BasePassRendering.cpp b/Engine/Source/Runtime/Renderer/Private/BasePassRendering.cpp
--- a/Engine/Source/Runtime/Renderer/Private/BasePassRendering.cpp
+++ b/Engine/Source/Runtime/Renderer/Private/BasePassRendering.cpp
@@ -112,11 +112,18 @@ static TAutoConsoleVariable<int32> CVarPSOPrecacheTranslucencyAllPass(
 
 static TAutoConsoleVariable<int32> CVarPSOPrecacheAlphaColorChannel(
 	TEXT("r.PSOPrecache.PrecacheAlphaColorChannel"),
-	1,
+	0,
 	TEXT("Also Precache PSOs with scene color alpha channel enabled. Planar reflections and scene captures use this for compositing into a different scene later."),
 	ECVF_ReadOnly
 );
 
+static TAutoConsoleVariable<int32> CVarPSOPrecacheNotSeparateTranslucency(
+	TEXT("r.PSOPrecache.PrecacheNotSeparateTranslucency"),
+	0,
+	TEXT("TODO."),
+	ECVF_ReadOnly
+);
+
 static TAutoConsoleVariable<bool> CVarAllow128BitBPPSCompilation(
 	TEXT("r.128BitBPPSCompilation.Allow"),
 	true,
@@ -1839,18 +1846,21 @@ void FBasePassMeshProcessor::CollectPSOInitializersForLMPolicy(
 				ERenderTargetLoadAction::ELoad, FExclusiveDepthStencil::DepthRead_StencilWrite, RenderTargetsInfo);
 		}
 
-		AddGraphicsPipelineStateInitializer(
-			VertexFactoryData,
-			MaterialResource,
-			DrawRenderState,
-			RenderTargetsInfo,
-			BasePassShaders,
-			MeshFillMode,
-			MeshCullMode,
-			PrimitiveType,
-			EMeshPassFeatures::Default,
-			true /*bRequired*/,
-			PSOInitializers);
+		if (!bRenderInSeparateTranslucency || CVarPSOPrecacheNotSeparateTranslucency->GetInt() > 0)
+		{
+			AddGraphicsPipelineStateInitializer(
+				VertexFactoryData,
+				MaterialResource,
+				DrawRenderState,
+				RenderTargetsInfo,
+				BasePassShaders,
+				MeshFillMode,
+				MeshCullMode,
+				PrimitiveType,
+				EMeshPassFeatures::Default,
+				true /*bRequired*/,
+				PSOInitializers);
+		}
 
 		// Add another PSO when render in separate translucency because render target format could have changed
 		if (bRenderInSeparateTranslucency)
diff --git a/Engine/Source/Runtime/Renderer/Private/DepthRendering.cpp b/Engine/Source/Runtime/Renderer/Private/DepthRendering.cpp
--- a/Engine/Source/Runtime/Renderer/Private/DepthRendering.cpp
+++ b/Engine/Source/Runtime/Renderer/Private/DepthRendering.cpp
@@ -80,7 +80,7 @@ static TAutoConsoleVariable<int32> CVarPSOPrecacheDitheredLODFadingOutMaskPass(
 
 static TAutoConsoleVariable<int32> CVarPSOPrecacheProjectedShadows(
 	TEXT("r.PSOPrecache.ProjectedShadows"),
-	1,
+	0,
 	TEXT("Also Precache PSOs with for projected shadows.")  \
 	TEXT(" 0: No PSOs are compiled for projected shadows.\n") \
 	TEXT(" 1: PSOs are compiled for all primitives which render to depth pass when VSMs are disabled (default).\n"),
diff --git a/Engine/Source/Runtime/Renderer/Private/DistortionRendering.cpp b/Engine/Source/Runtime/Renderer/Private/DistortionRendering.cpp
--- a/Engine/Source/Runtime/Renderer/Private/DistortionRendering.cpp
+++ b/Engine/Source/Runtime/Renderer/Private/DistortionRendering.cpp
@@ -1009,7 +1009,8 @@ void FDistortionMeshProcessor::AddMeshBatch(const FMeshBatch& RESTRICT MeshBatch
 
 static bool ShouldDraw(const FMaterial& Material)
 {
-	return (IsTranslucentBlendMode(Material)
+	return (GetGlobalRefractionQuality() > 0
+		&& IsTranslucentBlendMode(Material)
 		&& ShouldIncludeDomainInMeshPass(Material.GetMaterialDomain())
 		&& Material.IsDistorted());
 }
@@ -1021,12 +1022,6 @@ void FDistortionMeshProcessor::CollectPSOInitializers(const FSceneTexturesConfig
 		return;
 	}
 
-	// If refraction effects are disabled then don't precache distortion PSOs
-	if (GetGlobalRefractionQuality() == 0)
-	{
-		return;
-	}
-
 	TMeshProcessorShaders<
 		FDistortionMeshVS,
 		FDistortionMeshPS> DistortionPassShaders;
diff --git a/Engine/Source/Runtime/Renderer/Private/ShadowDepthRendering.cpp b/Engine/Source/Runtime/Renderer/Private/ShadowDepthRendering.cpp
--- a/Engine/Source/Runtime/Renderer/Private/ShadowDepthRendering.cpp
+++ b/Engine/Source/Runtime/Renderer/Private/ShadowDepthRendering.cpp
@@ -2332,7 +2332,7 @@ void FShadowDepthPassMeshProcessor::CollectPSOInitializersForEachShadowDepthType
 
 	// Collect for non-directional non-one pass point light shadows
 	{
-		const EShadowDepthType LocalShadowDepthType = ShadowDepthVSMBit;
+		const EShadowDepthType LocalShadowDepthType = ShadowDepthVSMBit | EShadowDepthType::Spot;
 		ERasterizerCullMode FinalCullMode = SetupShadowCullMode(FeatureLevel, LocalShadowDepthType, Material, MeshCullMode, bCastShadowAsTwoSided);
 		CollectPSOInitializersForEachStreamSetup(VertexFactoryData, Material, LocalShadowDepthType, MeshFillMode, FinalCullMode, bRequired, PSOInitializers);
 	}
