VolumetricCloud - fix for initialisation/copy race condition between render and shadow work.

#tests PC editor dx12
#rnx

#Synced-CL 48943944

diff --git a/Engine/Source/Runtime/Renderer/Private/SceneRendering.cpp b/Engine/Source/Runtime/Renderer/Private/SceneRendering.cpp
--- a/Engine/Source/Runtime/Renderer/Private/SceneRendering.cpp
+++ b/Engine/Source/Runtime/Renderer/Private/SceneRendering.cpp
@@ -950,8 +950,6 @@ void FViewInfo::Init()
 	SkyAtmosphereCameraAerialPerspectiveVolumeRayOnly = nullptr;
 	SkyAtmosphereUniformShaderParameters = nullptr;
 
-	VolumetricCloudSkyAO = nullptr;
-
 	bUseDirectionalInscattering = false;
 	DirectionalInscatteringExponent = 0;
 	DirectionalInscatteringStartDistance = 0;
@@ -1035,6 +1033,17 @@ void FViewInfo::Init()
 
 	// Filled in by FDeferredShadingSceneRenderer::UpdateLumenScene
 	ViewLumenSceneData = nullptr;
+
+	VolumetricCloudSkyAO = nullptr;
+	for (int LightIndex = 0; LightIndex < NUM_ATMOSPHERE_LIGHTS; ++LightIndex)
+	{
+		VolumetricCloudShadowRenderTarget[LightIndex] = nullptr;
+		VolumetricCloudShadowExtractedRenderTarget[LightIndex] = nullptr;
+		if (ViewState != nullptr)
+		{
+			ViewState->VolumetricCloudShadowRenderTarget[LightIndex].Reset();
+		}
+	}
 }
 
 FParallelMeshDrawCommandPass* FViewInfo::CreateMeshPass(EMeshPass::Type MeshPass)
diff --git a/Engine/Source/Runtime/Renderer/Private/VolumetricCloudRendering.cpp b/Engine/Source/Runtime/Renderer/Private/VolumetricCloudRendering.cpp
--- a/Engine/Source/Runtime/Renderer/Private/VolumetricCloudRendering.cpp
+++ b/Engine/Source/Runtime/Renderer/Private/VolumetricCloudRendering.cpp
@@ -1686,13 +1686,11 @@ void FSceneRenderer::InitVolumetricCloudsForViews(FRDGBuilder& GraphBuilder, boo
 	{
 		GraphBuilder.AddPostExecuteCallback([&Views]
 		{
+			// Make sure we always clear the texture on views to make sure no dangling texture pointers are ever used
 			CleanUpCloudDataFunction(Views);
 		});
 	};
 
-	// First make sure we always clear the texture on views to make sure no dangling texture pointers are ever used
-	CleanUpCloudDataFunction(Views);
-
 	if (!bShouldRenderVolumetricCloud)
 	{
 		return;
