SkyAtmospehre - fix for race condition with light setup and gpuscene light data copy.

#tests PC editor / preview mobile
#rnx
#Synced-CL 48943944

diff --git a/Engine/Source/Runtime/Renderer/Private/DeferredShadingRenderer.cpp b/Engine/Source/Runtime/Renderer/Private/DeferredShadingRenderer.cpp
--- a/Engine/Source/Runtime/Renderer/Private/DeferredShadingRenderer.cpp
+++ b/Engine/Source/Runtime/Renderer/Private/DeferredShadingRenderer.cpp
@@ -1806,34 +1806,14 @@ void FDeferredShadingSceneRenderer::Render(FRDGBuilder& GraphBuilder, const FSce
 	{
 		if (RendererOutput == ERendererOutput::FinalSceneColor)
 		{
-			// 1. Update sky atmosphere
-			// This needs to be done prior to start Lumen scene lighting to ensure directional light color is correct, as the sun color needs atmosphere transmittance
-			{
-				const bool bPathTracedAtmosphere = ViewFamily.EngineShowFlags.PathTracing && Views.Num() > 0 && PathTracing::UsesReferenceAtmosphere(Views[0]);
-				if (ShouldRenderSkyAtmosphere(Scene, ViewFamily.EngineShowFlags) && !bPathTracedAtmosphere)
-				{
-					for (int32 LightIndex = 0; LightIndex < NUM_ATMOSPHERE_LIGHTS; ++LightIndex)
-					{
-						if (Scene->AtmosphereLights[LightIndex])
-						{
-							PrepareSunLightProxy(*Scene->GetSkyAtmosphereSceneInfo(),LightIndex, *Scene->AtmosphereLights[LightIndex]);
-						}
-					}
-				}
-				else
-				{
-					Scene->ResetAtmosphereLightsProperties();
-				}
-			}
-
-			// 2. Update global distance field view origin
+			// 1. Update global distance field view origin
 			// This needs to be done prior to start Lumen scene lighting to ensure GlobalDistanceFieldData->CameraVelocityOffset is updated
 			if (SceneUpdateInputs)
 			{
 				UpdateGlobalDistanceFieldViewOrigin(*SceneUpdateInputs);
 			}
 
-			// 3. Update lumen scene
+			// 2. Update lumen scene
 			{
 				InitViewTaskDatas.LumenFrameTemporaries = &LumenFrameTemporaries;
 	
diff --git a/Engine/Source/Runtime/Renderer/Private/MobileShadingRenderer.cpp b/Engine/Source/Runtime/Renderer/Private/MobileShadingRenderer.cpp
--- a/Engine/Source/Runtime/Renderer/Private/MobileShadingRenderer.cpp
+++ b/Engine/Source/Runtime/Renderer/Private/MobileShadingRenderer.cpp
@@ -1076,24 +1076,6 @@ void FMobileSceneRenderer::Render(FRDGBuilder& GraphBuilder, const FSceneRenderU
 	// Establish scene primitive count (must be done after UpdateAllPrimitiveSceneInfos)
 	FGPUSceneScopeBeginEndHelper GPUSceneScopeBeginEndHelper(GraphBuilder, Scene->GPUScene, GPUSceneDynamicContext);
 
-	if (bRendererOutputFinalSceneColor)
-	{
-		if (ShouldRenderSkyAtmosphere(Scene, ViewFamily.EngineShowFlags))
-		{
-			for (int32 LightIndex = 0; LightIndex < NUM_ATMOSPHERE_LIGHTS; ++LightIndex)
-			{
-				if (Scene->AtmosphereLights[LightIndex])
-				{
-					PrepareSunLightProxy(*Scene->GetSkyAtmosphereSceneInfo(), LightIndex, *Scene->AtmosphereLights[LightIndex]);
-				}
-			}
-		}
-		else
-		{
-			Scene->ResetAtmosphereLightsProperties();
-		}
-	}
-
 	CSV_SCOPED_TIMING_STAT_EXCLUSIVE(RenderOther);
 	QUICK_SCOPE_CYCLE_COUNTER(STAT_FMobileSceneRenderer_Render);
 
diff --git a/Engine/Source/Runtime/Renderer/Private/SceneRendering.cpp b/Engine/Source/Runtime/Renderer/Private/SceneRendering.cpp
--- a/Engine/Source/Runtime/Renderer/Private/SceneRendering.cpp
+++ b/Engine/Source/Runtime/Renderer/Private/SceneRendering.cpp
@@ -3959,6 +3959,8 @@ IVisibilityTaskData* FSceneRenderer::OnRenderBegin(FRDGBuilder& GraphBuilder, co
 
 	// This is called prior to scene update
 	OIT::OnRenderBegin(Scene->OITSceneData);
+	// Notify the atmosphere system to update lights that are needed for this scene and all its renders
+	SkyAtmosphereOnRenderBegin(*Scene, Views, ViewFamily);
 
 	const bool bIsMobilePlatform = IsMobilePlatform(ShaderPlatform);
 
diff --git a/Engine/Source/Runtime/Renderer/Private/SkyAtmosphereRendering.cpp b/Engine/Source/Runtime/Renderer/Private/SkyAtmosphereRendering.cpp
--- a/Engine/Source/Runtime/Renderer/Private/SkyAtmosphereRendering.cpp
+++ b/Engine/Source/Runtime/Renderer/Private/SkyAtmosphereRendering.cpp
@@ -422,6 +422,29 @@ bool ShouldRenderSkyAtmosphere(const FScene* Scene, const FEngineShowFlags& Engi
 	return false;
 }
 
+void SkyAtmosphereOnRenderBegin(FScene& Scene, TArray<FViewInfo>& Views, FViewFamilyInfo& ViewFamily)
+{
+	// Note: we are morphing the proxy according to the Skyatmosphere for the scene and all of its renderer. That is fine.
+	// But here we are testing a single view representing multiple views that could be rendered with or without the path tracer.
+	// TODO: the path tracer should handle that difference internally. Once that is done, we could move that SkyAtmosphereOnRenderBegin update in the FScene::UpdateLights task.
+	const bool bPathTracedAtmosphere = ViewFamily.EngineShowFlags.PathTracing && Views.Num() > 0 && PathTracing::UsesReferenceAtmosphere(Views[0]);
+
+	if (ShouldRenderSkyAtmosphere(&Scene, ViewFamily.EngineShowFlags) && !bPathTracedAtmosphere)
+	{
+		for (int32 LightIndex = 0; LightIndex < NUM_ATMOSPHERE_LIGHTS; ++LightIndex)
+		{
+			if (Scene.AtmosphereLights[LightIndex])
+			{
+				PrepareSunLightProxy(*Scene.GetSkyAtmosphereSceneInfo(), LightIndex, *Scene.AtmosphereLights[LightIndex]);
+			}
+		}
+	}
+	else
+	{
+		Scene.ResetAtmosphereLightsProperties();
+	}
+}
+
 static auto GetSizeAndInvSize = [](int32 Width, int32 Height)
 {
 	float FWidth = float(Width);
diff --git a/Engine/Source/Runtime/Renderer/Private/SkyAtmosphereRendering.h b/Engine/Source/Runtime/Renderer/Private/SkyAtmosphereRendering.h
--- a/Engine/Source/Runtime/Renderer/Private/SkyAtmosphereRendering.h
+++ b/Engine/Source/Runtime/Renderer/Private/SkyAtmosphereRendering.h
@@ -17,6 +17,7 @@
 
 class FScene;
 class FViewInfo;
+class FViewFamilyInfo;
 class FLightSceneInfo;
 class FVisibleLightInfo;
 class USkyAtmosphereComponent;
@@ -219,6 +220,8 @@ extern ESkyAtmospherePassLocation GetSkyAtmospherePassLocation();
 
 bool ShouldRenderSkyAtmosphere(const FScene* Scene, const FEngineShowFlags& EngineShowFlags);
 
+void SkyAtmosphereOnRenderBegin(FScene& Scene, TArray<FViewInfo>& Views, FViewFamilyInfo& ViewFamily);
+
 void InitSkyAtmosphereForScene(FRHICommandListImmediate& RHICmdList, FRDGBuilder& GraphBuilder, FScene* Scene);
 void InitSkyAtmosphereForView(FRHICommandListImmediate& RHICmdList, const FScene* Scene, FViewInfo& View);
 
