Missing TopLeftShift correction for UpperBoundViewRect when preparing the Views

#Synced-CL 48897982

diff --git a/Engine/Source/Runtime/Renderer/Private/SceneRendering.cpp b/Engine/Source/Runtime/Renderer/Private/SceneRendering.cpp
--- a/Engine/Source/Runtime/Renderer/Private/SceneRendering.cpp
+++ b/Engine/Source/Runtime/Renderer/Private/SceneRendering.cpp
@@ -3447,6 +3447,7 @@ void FSceneRenderer::PrepareViewRectsForRendering()
 		for (int32 i = 0; i < Views.Num(); i++)
 		{
 			Views[i].ViewRect -= TopLeftShift;
+			Views[i].UpperBoundViewRect -= TopLeftShift;
 		}
 	}
 
