[HWRT] Fix stale bDrawInGame flag in PrimitiveRayTracingData.
- hit this when running cooked CitySample build using ODSC.

#Synced-CL 48972954

diff --git a/Engine/Source/Runtime/Renderer/Private/RendererScene.cpp b/Engine/Source/Runtime/Renderer/Private/RendererScene.cpp
--- a/Engine/Source/Runtime/Renderer/Private/RendererScene.cpp
+++ b/Engine/Source/Runtime/Renderer/Private/RendererScene.cpp
@@ -1855,7 +1855,7 @@ void FScene::UpdatePrimitivesIsDrawn_RenderThread(TArrayView<FPrimitiveSceneProx
 			}
 				
 #if RHI_RAYTRACING
-			if (bPrimitiveIndexValid && PrimitiveSceneProxy->HasRayTracingRepresentation())
+			if (bPrimitiveIndexValid)
 			{
 				FScene::FPrimitiveRayTracingData& PrimitiveRayTracingData = PrimitiveRayTracingDatas[PrimitiveIndex];
 				PrimitiveRayTracingData.bDrawInGame = bInIsDrawn;
