[HWRT] Fix r.RayTracing.UseTextureLod not being respected.
- TAutoConsoleVariable requires querying cvar directly since the argument is just used as default value.
- also removed some other similar globals to avoid confusion.

#Synced-CL 48972954

diff --git a/Engine/Source/Runtime/Renderer/Private/DeferredShadingRenderer.cpp b/Engine/Source/Runtime/Renderer/Private/DeferredShadingRenderer.cpp
--- a/Engine/Source/Runtime/Renderer/Private/DeferredShadingRenderer.cpp
+++ b/Engine/Source/Runtime/Renderer/Private/DeferredShadingRenderer.cpp
@@ -218,36 +218,32 @@ static TAutoConsoleVariable<int32> CVarRayTracing(
 	TEXT(" 1: on"),
 	ECVF_RenderThreadSafe | ECVF_ReadOnly);
 
-int32 GRayTracingUseTextureLod = 0;
 static TAutoConsoleVariable<int32> CVarRayTracingTextureLod(
 	TEXT("r.RayTracing.UseTextureLod"),
-	GRayTracingUseTextureLod,
+	0,
 	TEXT("Enable automatic texture mip level selection in ray tracing material shaders.\n")
 	TEXT(" 0: highest resolution mip level is used for all texture (default).\n")
 	TEXT(" 1: texture LOD is approximated based on total ray length, output resolution and texel density at hit point (ray cone method)."),
 	ECVF_RenderThreadSafe | ECVF_ReadOnly);
 
-static int32 GForceAllRayTracingEffects = -1;
 static TAutoConsoleVariable<int32> CVarForceAllRayTracingEffects(
 	TEXT("r.RayTracing.ForceAllRayTracingEffects"),
-	GForceAllRayTracingEffects,
+	-1,
 	TEXT("Force all ray tracing effects ON/OFF.\n")
 	TEXT(" -1: Do not force (default) \n")
 	TEXT(" 0: All ray tracing effects disabled\n")
 	TEXT(" 1: All ray tracing effects enabled"),
 	ECVF_RenderThreadSafe);
 
-static int32 GRayTracingAllowInline = 1;
 static TAutoConsoleVariable<int32> CVarRayTracingAllowInline(
 	TEXT("r.RayTracing.AllowInline"),
-	GRayTracingAllowInline,
+	1,
 	TEXT("Allow use of Inline Ray Tracing if supported (default=1)."),	
 	ECVF_RenderThreadSafe);
 
-static int32 GRayTracingAllowPipeline = 1;
 static TAutoConsoleVariable<int32> CVarRayTracingAllowPipeline(
 	TEXT("r.RayTracing.AllowPipeline"),
-	GRayTracingAllowPipeline,
+	1,
 	TEXT("Allow use of Ray Tracing pipelines if supported (default=1)."),
 	ECVF_RenderThreadSafe);
 
diff --git a/Engine/Source/Runtime/Renderer/Private/RayTracing/RayTracingMaterialHitShaders.cpp b/Engine/Source/Runtime/Renderer/Private/RayTracing/RayTracingMaterialHitShaders.cpp
--- a/Engine/Source/Runtime/Renderer/Private/RayTracing/RayTracingMaterialHitShaders.cpp
+++ b/Engine/Source/Runtime/Renderer/Private/RayTracing/RayTracingMaterialHitShaders.cpp
@@ -60,9 +60,6 @@ static FAutoConsoleVariableRef CVarRayTracingNonBlockingPipelineCreation(
 	TEXT(" 1: on (default, non-blocking mode may sometimes use the fallback opaque black material outside of offline rendering scenarios)\n"),
 	ECVF_RenderThreadSafe);
 
-// CVar defined in DeferredShadingRenderer.cpp
-extern int32 GRayTracingUseTextureLod;
-
 static bool IsSupportedVertexFactoryType(const FVertexFactoryType* VertexFactoryType)
 {
 	return VertexFactoryType->SupportsRayTracing();
@@ -169,11 +166,14 @@ public:
 		const bool bWantAnyHitShader = CVarCompileMaterialAHS.Get(Parameters.Platform) != 0 && RTNeedsAnyHitShader(Parameters.MaterialParameters.BlendMode);
 		const bool bSupportProceduralPrimitive = Parameters.VertexFactoryType->SupportsRayTracingProceduralPrimitive() && FDataDrivenShaderPlatformInfo::GetSupportsRayTracingProceduralPrimitive(Parameters.Platform);
 
+		static const auto CVarRayTracingUseTextureLod = IConsoleManager::Get().FindTConsoleVariableDataInt(TEXT("r.RayTracing.UseTextureLod"));
+		const bool bRayTracingUseTextureLod = CVarRayTracingUseTextureLod && CVarRayTracingUseTextureLod->GetValueOnAnyThread() != 0;
+
 		return IsSupportedVertexFactoryType(Parameters.VertexFactoryType)
 			&& (bWantAnyHitShader == UseAnyHitShader)
 			&& LightMapPolicyType::ShouldCompilePermutation(Parameters)
 			&& ShouldCompileRayTracingShadersForProject(Parameters.Platform)
-			&& (bool)GRayTracingUseTextureLod == UseRayConeTextureLod
+			&& bRayTracingUseTextureLod == UseRayConeTextureLod
 			&& (UseIntersectionShader == bSupportProceduralPrimitive);
 	}
 
@@ -376,7 +376,8 @@ static bool GetRayTracingMeshProcessorShaders(
 
 	if (bMaterialsCompiled)
 	{
-		const bool bUseTextureLOD = bool(GRayTracingUseTextureLod);
+		static const auto CVarRayTracingUseTextureLod = IConsoleManager::Get().FindTConsoleVariableDataInt(TEXT("r.RayTracing.UseTextureLod"));
+		const bool bUseTextureLOD = CVarRayTracingUseTextureLod && CVarRayTracingUseTextureLod->GetValueOnAnyThread() != 0;
 
 		switch (LightMapPolicy.GetIndirectPolicy())
 		{
