Fixed assert due to list element being added to multiple lists.

#Synced-CL 48995677

diff --git a/Engine/Source/Runtime/Renderer/Private/SceneRenderBuilder.cpp b/Engine/Source/Runtime/Renderer/Private/SceneRenderBuilder.cpp
--- a/Engine/Source/Runtime/Renderer/Private/SceneRenderBuilder.cpp
+++ b/Engine/Source/Runtime/Renderer/Private/SceneRenderBuilder.cpp
@@ -587,7 +587,6 @@ public:
 		{
 			TailGroup->RenderNodes.AddTail(RenderNode);
 		}
-		RenderNodes.AddTail(RenderNode);
 
 		Ops.Emplace(FOp::Render(RenderNode));
 		Renderers.Emplace(Renderer);
@@ -679,7 +678,6 @@ private:
 	TArray<FSceneRenderer*, FConcurrentLinearArrayAllocator> ActiveRenderers;
 	TArray<FViewFamilyInfo*, FConcurrentLinearArrayAllocator> ActiveViewFamilies;
 	TArray<FViewInfo*, FConcurrentLinearArrayAllocator> ActiveViews;
-	TIntrusiveDoubleLinkedList<FRenderNode> RenderNodes;
 	TIntrusiveDoubleLinkedList<FGroupNode> GroupNodes;
 	TArray<FOp, FConcurrentLinearArrayAllocator> Ops;
 	FConcurrentLinearBulkObjectAllocator Allocator;
