Path Tracer: Adjust threshold to detect large camera moves

When moving around very small scale scenes, this threshold could cause the view to become blurry due to the accumulation not getting reset from small camera moves

#rnx
#Synced-CL 49002339

diff --git a/Engine/Source/Runtime/Renderer/Private/SceneVisibility.cpp b/Engine/Source/Runtime/Renderer/Private/SceneVisibility.cpp
--- a/Engine/Source/Runtime/Renderer/Private/SceneVisibility.cpp
+++ b/Engine/Source/Runtime/Renderer/Private/SceneVisibility.cpp
@@ -5452,11 +5452,13 @@ void FSceneRenderer::PrepareViewStateForVisibility(const FSceneTexturesConfig& S
 				// for interactive usage - any movement or scene change should restart the path tracer
 
 				// Note: 0.18 deg is the minimum angle for avoiding numerical precision issue (which would cause constant invalidation)
+				// The translation tolerance is because camera movement in the engine has a bit of falloff animation which takes a few seconds to fully settle.
+				// Allowing the path tracer to ignore this imperceptible drift helps the path tracer start accumulating frames sooner.
 				const bool bIsCameraMove = IsLargeCameraMovement(
 					View,
 					ViewState->PrevFrameViewInfo.ViewMatrices.GetViewMatrix(),
 					ViewState->PrevFrameViewInfo.ViewMatrices.GetViewOrigin(),
-					0.18f /*degree*/, 0.1f /*cm*/);
+					0.18f /*degree*/, 0.001f /*cm*/);
 				const bool bIsProjMatrixDifferent = View.ViewMatrices.GetProjectionNoAAMatrix() != View.ViewState->PrevFrameViewInfo.ViewMatrices.GetProjectionNoAAMatrix();
 
 				// For each view, we remember what the invalidation counter was set to last time we were here so we can catch all changes
