Collapse all history super sampling permutations into one as they have near-zero difference in the compiled shader code

#Synced-CL 49002339

diff --git a/Engine/Shaders/Private/VolumetricFog.usf b/Engine/Shaders/Private/VolumetricFog.usf
--- a/Engine/Shaders/Private/VolumetricFog.usf
+++ b/Engine/Shaders/Private/VolumetricFog.usf
@@ -366,8 +366,8 @@ float PhaseG;
 float InverseSquaredLightDistanceBiasScale;
 int VirtualShadowMapId;
 
-#ifndef HISTORY_MISS_SUPER_SAMPLE_COUNT
-#define HISTORY_MISS_SUPER_SAMPLE_COUNT 1
+#ifndef USE_SUPER_SAMPLING
+#define USE_SUPER_SAMPLING 0
 #endif
 
 #include "LightFunctionAtlas/LightFunctionAtlasCommon.usf"
@@ -871,7 +871,7 @@ void LightScatteringCS(
 
 	// Supersample if the history was outside the camera frustum
 	// The compute shader is dispatched with extra threads, make sure those don't supersample
-	NumSuperSamples = HistoryAlpha < .001f && all(int3(GridCoordinate) < VolumetricFog.ViewGridSizeInt) ? HISTORY_MISS_SUPER_SAMPLE_COUNT : 1;
+	NumSuperSamples = HistoryAlpha < .001f && all(int3(GridCoordinate) < VolumetricFog.ViewGridSizeInt) && USE_SUPER_SAMPLING ? HistoryMissSuperSampleCount : 1;
 #endif
 
 	for (uint SampleIndex = 0; SampleIndex < NumSuperSamples; SampleIndex++)
diff --git a/Engine/Source/Runtime/Renderer/Private/VolumetricFog.cpp b/Engine/Source/Runtime/Renderer/Private/VolumetricFog.cpp
--- a/Engine/Source/Runtime/Renderer/Private/VolumetricFog.cpp
+++ b/Engine/Source/Runtime/Renderer/Private/VolumetricFog.cpp
@@ -1030,16 +1030,15 @@ class FVolumetricFogLightScatteringCS : public FGlobalShader
 
 	class FTemporalReprojection			: SHADER_PERMUTATION_BOOL("USE_TEMPORAL_REPROJECTION");
 	class FDistanceFieldSkyOcclusion	: SHADER_PERMUTATION_BOOL("DISTANCE_FIELD_SKY_OCCLUSION");
-	class FSuperSampleCount				: SHADER_PERMUTATION_SPARSE_INT("HISTORY_MISS_SUPER_SAMPLE_COUNT", 1, 4, 8, 16);
 	class FLumenGI						: SHADER_PERMUTATION_BOOL("LUMEN_GI");
 	class FVirtualShadowMap				: SHADER_PERMUTATION_BOOL("VIRTUAL_SHADOW_MAP");
 	class FRaytracedShadowsVolume		: SHADER_PERMUTATION_BOOL("USE_RAYTRACED_SHADOWS_VOLUME");
 	class FSampleLightFunctionAtlas		: SHADER_PERMUTATION_BOOL("USE_LIGHT_FUNCTION_ATLAS");
 	class FMegaLights					: SHADER_PERMUTATION_BOOL("USE_MEGA_LIGHTS");
 	class FLightSoftFading				: SHADER_PERMUTATION_BOOL("USE_LIGHT_SOFT_FADING");
+	class FSuperSampling				: SHADER_PERMUTATION_BOOL("USE_SUPER_SAMPLING");
 
 	using FPermutationDomain = TShaderPermutationDomain<
-		FSuperSampleCount,
 		FTemporalReprojection,
 		FDistanceFieldSkyOcclusion,
 		FLumenGI,
@@ -1047,7 +1046,8 @@ class FVolumetricFogLightScatteringCS : public FGlobalShader
 		FRaytracedShadowsVolume,
 		FSampleLightFunctionAtlas,
 		FMegaLights,
-		FLightSoftFading>;
+		FLightSoftFading,
+		FSuperSampling>;
 
 	BEGIN_SHADER_PARAMETER_STRUCT(FParameters, )
 		SHADER_PARAMETER_STRUCT_REF(FViewUniformShaderParameters, View)
@@ -1190,7 +1190,7 @@ class FVolumetricFogLightScatteringCS : public FGlobalShader
 		}
 
 		const int32 SuperSampleCount = FVolumetricFogLightScatteringCS::GetSuperSampleCount(GVolumetricFogHistoryMissSupersampleCount);
-		if (PermutationVector.Get<FSuperSampleCount>() != SuperSampleCount)
+		if (PermutationVector.Get<FSuperSampling>() != bool(SuperSampleCount > 1))
 		{
 			return EShaderPermutationPrecacheRequest::NotUsed;
 		}
@@ -1833,11 +1833,12 @@ void FSceneRenderer::ComputeVolumetricFog(FRDGBuilder& GraphBuilder,
 				&& !IsMobileVolumetricFog(View.GetShaderPlatform());
 
 			const int32 SuperSampleCount = FVolumetricFogLightScatteringCS::GetSuperSampleCount(GVolumetricFogHistoryMissSupersampleCount);
+			check(PassParameters->VolumetricFogParameters.HistoryMissSuperSampleCount == SuperSampleCount)
 
 			FVolumetricFogLightScatteringCS::FPermutationDomain PermutationVector;
 			PermutationVector.Set< FVolumetricFogLightScatteringCS::FTemporalReprojection >(bUseTemporalReprojection);
 			PermutationVector.Set< FVolumetricFogLightScatteringCS::FDistanceFieldSkyOcclusion >(bUseDistanceFieldSkyOcclusion);
-			PermutationVector.Set< FVolumetricFogLightScatteringCS::FSuperSampleCount >(SuperSampleCount);
+			PermutationVector.Set< FVolumetricFogLightScatteringCS::FSuperSampling>(SuperSampleCount > 1);
 			PermutationVector.Set< FVolumetricFogLightScatteringCS::FLumenGI >(bUseLumenGI);
 			PermutationVector.Set< FVolumetricFogLightScatteringCS::FVirtualShadowMap >(VirtualShadowMapArray.IsAllocated() );
 			PermutationVector.Set< FVolumetricFogLightScatteringCS::FRaytracedShadowsVolume >(bUseRaytracedShadowsVolume);
