[HWRT] Refactored FPrimitiveSceneInfo::CacheRayTracingPrimitives(...) and related functions.
- moved allocation of RayTracingLODData and DeferredData out of CacheRayTracingMeshCommands(...).

#Synced-CL 49007722

diff --git a/Engine/Source/Runtime/Renderer/Private/PrimitiveSceneInfo.cpp b/Engine/Source/Runtime/Renderer/Private/PrimitiveSceneInfo.cpp
--- a/Engine/Source/Runtime/Renderer/Private/PrimitiveSceneInfo.cpp
+++ b/Engine/Source/Runtime/Renderer/Private/PrimitiveSceneInfo.cpp
@@ -1110,7 +1110,7 @@ struct FDeferredPrimitiveRayTracingData
 	FPrimitiveSceneInfo* SceneInfo = nullptr;
 	TArray<TArray<int32, TInlineAllocator<2>>, TInlineAllocator<1>> CommandIndicesPerLOD;
 	FRayTracingGeometryInstance Instance;
-	uint32 SegmentCount;
+	uint32 SegmentCount = 0;
 	bool bCacheInstance = false;
 };
 
@@ -1132,8 +1132,8 @@ public:
 template<class T>
 void CacheRayTracingMeshBatch(
 	const FMeshBatch& MeshBatch,
-	FPrimitiveSceneInfo* SceneInfo,
-	TArray<FPrimitiveSceneInfo::FRayTracingLODData>& RayTracingLODData,
+	const FPrimitiveSceneProxy* PrimitiveSceneProxy,
+	TArrayView<FPrimitiveSceneInfo::FRayTracingLODData> RayTracingLODData,
 	T& Commands,
 	FCachedRayTracingMeshCommandContext<T>& CommandContext,
 	FRayTracingMeshProcessor& RayTracingMeshProcessor,
@@ -1142,7 +1142,7 @@ void CacheRayTracingMeshBatch(
 	// Why do we pass a full mask here when the dynamic case only uses a mask of 1?
 	// Also note that the code below assumes only a single command was generated per batch (see SupportsCachingMeshDrawCommands(...))
 	const uint64 BatchElementMask = ~0ull;
-	RayTracingMeshProcessor.AddMeshBatch(MeshBatch, BatchElementMask, SceneInfo->Proxy);
+	RayTracingMeshProcessor.AddMeshBatch(MeshBatch, BatchElementMask, PrimitiveSceneProxy);
 
 	if (CommandContext.CommandIndex >= 0)
 	{
@@ -1176,58 +1176,25 @@ void CacheRayTracingMeshBatch(
 }
 
 template<class T>
-int32 CacheRayTracingMeshCommands(
-	FPrimitiveSceneInfo* SceneInfo,
+void CacheRayTracingMeshCommands(
 	TConstArrayView<FMeshBatch> MeshBatches,
-	ERayTracingPrimitiveFlags Flags,
+	FPrimitiveSceneProxy* PrimitiveSceneProxy,
+	TArrayView<FPrimitiveSceneInfo::FRayTracingLODData> RayTracingLODData,
+	int32 LODCount,
 	T& Commands,
 	FCachedRayTracingMeshCommandContext<T>& CommandContext,
 	FRayTracingMeshProcessor& RayTracingMeshProcessor,
-	TArray<FDeferredPrimitiveRayTracingData>& DeferredDatas)
+	FDeferredPrimitiveRayTracingData& DeferredData)
 {
-	int32 DeferredDataIndex = INDEX_NONE;
+	ensureMsgf(MeshBatches.Num() > 0, TEXT("Caching ray tracing mesh commands but instance doesn't have any."));
 
-	// the following flags cause ray tracing mesh command caching to be disabled
-	static const ERayTracingPrimitiveFlags DisableCacheMeshCommandsFlags = ERayTracingPrimitiveFlags::Dynamic
-		| ERayTracingPrimitiveFlags::Exclude
-		| ERayTracingPrimitiveFlags::Skip
-		| ERayTracingPrimitiveFlags::UnsupportedProxyType;
-
-	if (!EnumHasAnyFlags(Flags, DisableCacheMeshCommandsFlags))
-	{
-		// Cache ray tracing mesh commands in FPrimitiveSceneInfo
-
-		ensureMsgf(MeshBatches.Num() > 0, TEXT("Caching ray tracing mesh commands but instance doesn't have any."));
-
-		const int32 LODCount = EnumHasAnyFlags(Flags, ERayTracingPrimitiveFlags::CacheInstances) ? 1 : SceneInfo->GetStaticRayTracingGeometryNum();
-
-		check(SceneInfo->GetRayTracingLODDataNum() == 0);
-				
-		TArray<FPrimitiveSceneInfo::FRayTracingLODData> RayTracingLODData;
-		RayTracingLODData.Empty(LODCount);
-		RayTracingLODData.AddDefaulted(LODCount);
-
-		DeferredDataIndex = DeferredDatas.AddZeroed();
-
-		FDeferredPrimitiveRayTracingData& DeferredData = DeferredDatas[DeferredDataIndex];
-		DeferredData.SceneInfo = SceneInfo;
-		DeferredData.CommandIndicesPerLOD.SetNum(LODCount);
+	DeferredData.CommandIndicesPerLOD.SetNum(LODCount);
 		
-		for (const FMeshBatch& Mesh : MeshBatches)
-		{
-			check(Mesh.LODIndex >= 0 && Mesh.LODIndex < LODCount);
-			CacheRayTracingMeshBatch(Mesh, SceneInfo, RayTracingLODData, Commands, CommandContext, RayTracingMeshProcessor, DeferredData);
-		}
-
-		// Store in the Scene info
-		SceneInfo->SetRayTracingLODData(MoveTemp(RayTracingLODData));
-	}
-	else
+	for (const FMeshBatch& Mesh : MeshBatches)
 	{
-		checkf(!EnumHasAnyFlags(Flags, ERayTracingPrimitiveFlags::CacheInstances), TEXT("Can't cache ray tracing instance when mesh commands are not cached."));
+		check(Mesh.LODIndex >= 0 && Mesh.LODIndex < LODCount);
+		CacheRayTracingMeshBatch(Mesh, PrimitiveSceneProxy, RayTracingLODData, Commands, CommandContext, RayTracingMeshProcessor, DeferredData);
 	}
-
-	return DeferredDataIndex;
 }
 
 void FPrimitiveSceneInfo::CacheRayTracingPrimitives(FScene* Scene, const TArrayView<FPrimitiveSceneInfo*>& SceneInfos)
@@ -1260,6 +1227,7 @@ void FPrimitiveSceneInfo::CacheRayTracingPrimitives(FScene* Scene, const TArrayV
 
 #if DO_CHECK
 					{
+						check(SceneInfo->GetRayTracingLODDataNum() == 0);
 
 						Experimental::FHashElementId SceneRayTracingGroupId;
 						const int32 RayTracingGroupId = SceneProxy->GetRayTracingGroupId();
@@ -1280,13 +1248,40 @@ void FPrimitiveSceneInfo::CacheRayTracingPrimitives(FScene* Scene, const TArrayV
 					// Write flags
 					Flags = SceneProxy->GetCachedRayTracingInstance(RayTracingInstance);
 
-					const int32 DeferredDataIndex = CacheRayTracingMeshCommands(SceneInfo, RayTracingInstance.Materials, Flags, Context.Commands, Context.CommandContext, Context.RayTracingMeshProcessor, Context.DeferredDatas);
+					// the following flags cause static ray tracing caching to be disabled
+					static const ERayTracingPrimitiveFlags DisableCachingFlags = 
+						ERayTracingPrimitiveFlags::Dynamic
+						| ERayTracingPrimitiveFlags::Exclude
+						| ERayTracingPrimitiveFlags::Skip
+						| ERayTracingPrimitiveFlags::UnsupportedProxyType;
 
-					if (EnumHasAnyFlags(Flags, ERayTracingPrimitiveFlags::CacheInstances))
+					if (!EnumHasAnyFlags(Flags, DisableCachingFlags))
 					{
-						FDeferredPrimitiveRayTracingData& DeferredData = Context.DeferredDatas[DeferredDataIndex];
-						FPrimitiveSceneInfo::SetupCachedRayTracingInstance(SceneInfo, RayTracingInstance, DeferredData.Instance, DeferredData.SegmentCount);
-						DeferredData.bCacheInstance = true;
+						FDeferredPrimitiveRayTracingData& DeferredData = Context.DeferredDatas.AddDefaulted_GetRef();
+						DeferredData.SceneInfo = SceneInfo;
+
+						const int32 LODCount = EnumHasAnyFlags(Flags, ERayTracingPrimitiveFlags::CacheInstances) ? 1 : SceneInfo->GetStaticRayTracingGeometryNum();
+
+						TArray<FPrimitiveSceneInfo::FRayTracingLODData> RayTracingLODData;
+						RayTracingLODData.Empty(LODCount);
+						RayTracingLODData.AddDefaulted(LODCount);
+
+						CacheRayTracingMeshCommands(RayTracingInstance.Materials, SceneProxy, RayTracingLODData, LODCount, Context.Commands, Context.CommandContext, Context.RayTracingMeshProcessor, DeferredData);
+
+						if (EnumHasAnyFlags(Flags, ERayTracingPrimitiveFlags::CacheInstances))
+						{
+							SceneInfo->CachedRayTracingGeometry = RayTracingInstance.Geometry;
+
+							SetupCachedRayTracingInstance(SceneInfo, RayTracingInstance, DeferredData.Instance, DeferredData.SegmentCount);
+							DeferredData.bCacheInstance = true;
+						}
+
+						// Store LOD data in the FPrimitiveSceneInfo
+						SceneInfo->SetRayTracingLODData(MoveTemp(RayTracingLODData));
+					}
+					else
+					{
+						checkf(!EnumHasAnyFlags(Flags, ERayTracingPrimitiveFlags::CacheInstances), TEXT("Can't cache ray tracing instance when caching is disabled due to other flags."));
 					}
 
 					PrimitiveRayTracingData.bCachedRaytracingDataDirty = false;
@@ -1334,7 +1329,6 @@ void FPrimitiveSceneInfo::CacheRayTracingPrimitives(FScene* Scene, const TArrayV
 							Entry.SceneInfo->AllocateRayTracingSBT(nullptr, 0);
 						}
 
-
 						Entry.SceneInfo->CacheRayTracingShaderBindingData(Entry.Instance.GeometryRHI);
 
 						if (Entry.bCacheInstance)
@@ -1355,8 +1349,6 @@ void FPrimitiveSceneInfo::SetupCachedRayTracingInstance(FPrimitiveSceneInfo* Sce
 
 	FPrimitiveSceneProxy* SceneProxy = SceneInfo->Proxy;
 
-	SceneInfo->CachedRayTracingGeometry = RayTracingInstance.Geometry;
-
 	OutCachedRayTracingInstance.NumTransforms = RayTracingInstance.NumTransforms;
 	OutCachedRayTracingInstance.BaseInstanceSceneDataOffset = SceneInfo->GetInstanceSceneDataOffset();
 
