Path Tracer: Move the logic for directional light atmosphere attenuation to the path tracer

This avoids a long range interaction with the raster logic to update lights.

#rnx
#Synced-CL 49038315

diff --git a/Engine/Source/Runtime/Renderer/Private/PathTracing.cpp b/Engine/Source/Runtime/Renderer/Private/PathTracing.cpp
--- a/Engine/Source/Runtime/Renderer/Private/PathTracing.cpp
+++ b/Engine/Source/Runtime/Renderer/Private/PathTracing.cpp
@@ -3099,7 +3099,18 @@ void SetLightParameters(
 			}
 
 			// these mean roughly the same thing across all light types
-			DestLight.Color = FVector3f(LightParameters.Color) * LightParameters.GetLightExposureScale(View.GetLastEyeAdaptationExposure());
+			if (bUseAtmosphere)
+			{
+				// we want the sky color to come in _without_ the sky transmittance baked in, so need to go back to the proxy
+				DestLight.Color = FVector3f(Light.LightSceneInfo->Proxy->GetOuterSpaceIlluminance());
+			}
+			else
+			{
+				// this color has the sky transmittance already factored in
+				DestLight.Color = FVector3f(LightParameters.Color);
+			}
+
+			DestLight.Color *= LightParameters.GetLightExposureScale(View.GetLastEyeAdaptationExposure());
 			DestLight.TranslatedWorldPosition = FVector3f(LightParameters.WorldPosition + View.ViewMatrices.GetPreViewTranslation());
 			DestLight.Normal = -LightParameters.Direction;
 			DestLight.Tangent = LightParameters.Tangent;
diff --git a/Engine/Source/Runtime/Renderer/Private/SceneRendering.cpp b/Engine/Source/Runtime/Renderer/Private/SceneRendering.cpp
--- a/Engine/Source/Runtime/Renderer/Private/SceneRendering.cpp
+++ b/Engine/Source/Runtime/Renderer/Private/SceneRendering.cpp
@@ -3961,7 +3961,7 @@ IVisibilityTaskData* FSceneRenderer::OnRenderBegin(FRDGBuilder& GraphBuilder, co
 	// This is called prior to scene update
 	OIT::OnRenderBegin(Scene->OITSceneData);
 	// Notify the atmosphere system to update lights that are needed for this scene and all its renders
-	SkyAtmosphereOnRenderBegin(*Scene, Views, ViewFamily);
+	SkyAtmosphereOnRenderBegin(*Scene, ViewFamily.EngineShowFlags);
 
 	const bool bIsMobilePlatform = IsMobilePlatform(ShaderPlatform);
 
diff --git a/Engine/Source/Runtime/Renderer/Private/SkyAtmosphereRendering.cpp b/Engine/Source/Runtime/Renderer/Private/SkyAtmosphereRendering.cpp
--- a/Engine/Source/Runtime/Renderer/Private/SkyAtmosphereRendering.cpp
+++ b/Engine/Source/Runtime/Renderer/Private/SkyAtmosphereRendering.cpp
@@ -422,14 +422,9 @@ bool ShouldRenderSkyAtmosphere(const FScene* Scene, const FEngineShowFlags& Engi
 	return false;
 }
 
-void SkyAtmosphereOnRenderBegin(FScene& Scene, TArray<FViewInfo>& Views, FViewFamilyInfo& ViewFamily)
+void SkyAtmosphereOnRenderBegin(FScene& Scene, const FEngineShowFlags& EngineShowFlags)
 {
-	// Note: we are morphing the proxy according to the Skyatmosphere for the scene and all of its renderer. That is fine.
-	// But here we are testing a single view representing multiple views that could be rendered with or without the path tracer.
-	// TODO: the path tracer should handle that difference internally. Once that is done, we could move that SkyAtmosphereOnRenderBegin update in the FScene::UpdateLights task.
-	const bool bPathTracedAtmosphere = ViewFamily.EngineShowFlags.PathTracing && Views.Num() > 0 && PathTracing::UsesReferenceAtmosphere(Views[0]);
-
-	if (ShouldRenderSkyAtmosphere(&Scene, ViewFamily.EngineShowFlags) && !bPathTracedAtmosphere)
+	if (ShouldRenderSkyAtmosphere(&Scene, EngineShowFlags))
 	{
 		for (int32 LightIndex = 0; LightIndex < NUM_ATMOSPHERE_LIGHTS; ++LightIndex)
 		{
diff --git a/Engine/Source/Runtime/Renderer/Private/SkyAtmosphereRendering.h b/Engine/Source/Runtime/Renderer/Private/SkyAtmosphereRendering.h
--- a/Engine/Source/Runtime/Renderer/Private/SkyAtmosphereRendering.h
+++ b/Engine/Source/Runtime/Renderer/Private/SkyAtmosphereRendering.h
@@ -220,7 +220,7 @@ extern ESkyAtmospherePassLocation GetSkyAtmospherePassLocation();
 
 bool ShouldRenderSkyAtmosphere(const FScene* Scene, const FEngineShowFlags& EngineShowFlags);
 
-void SkyAtmosphereOnRenderBegin(FScene& Scene, TArray<FViewInfo>& Views, FViewFamilyInfo& ViewFamily);
+void SkyAtmosphereOnRenderBegin(FScene& Scene, const FEngineShowFlags& EngineShowFlags);
 
 void InitSkyAtmosphereForScene(FRHICommandListImmediate& RHICmdList, FRDGBuilder& GraphBuilder, FScene* Scene);
 void InitSkyAtmosphereForView(FRHICommandListImmediate& RHICmdList, const FScene* Scene, FViewInfo& View);
