[HWRT] Store pointer to FRayTracingGeometry in FRayTracingLODData.
- replacing StaticRayTracingGeometries and CachedRayTracingGeometry.
- this is a step towards unifying static ray tracing instance logic.

#Synced-CL 49007722

diff --git a/Engine/Source/Runtime/Engine/Private/PrimitiveSceneProxy.cpp b/Engine/Source/Runtime/Engine/Private/PrimitiveSceneProxy.cpp
--- a/Engine/Source/Runtime/Engine/Private/PrimitiveSceneProxy.cpp
+++ b/Engine/Source/Runtime/Engine/Private/PrimitiveSceneProxy.cpp
@@ -1910,12 +1910,6 @@ ERayTracingPrimitiveFlags FPrimitiveSceneProxy::GetCachedRayTracingInstance(FRay
 
 	if (IsRayTracingStaticRelevant())
 	{
-		if (PrimitiveSceneInfo->GetStaticRayTracingGeometryNum() == 0
-			|| PrimitiveSceneInfo->StaticMeshes.IsEmpty())
-		{
-			return ERayTracingPrimitiveFlags::Exclude;
-		}
-
 		ResultFlags |= ERayTracingPrimitiveFlags::ComputeLOD;
 	}
 	else
diff --git a/Engine/Source/Runtime/Renderer/Private/PrimitiveSceneInfo.cpp b/Engine/Source/Runtime/Renderer/Private/PrimitiveSceneInfo.cpp
--- a/Engine/Source/Runtime/Renderer/Private/PrimitiveSceneInfo.cpp
+++ b/Engine/Source/Runtime/Renderer/Private/PrimitiveSceneInfo.cpp
@@ -352,12 +352,6 @@ FPrimitiveSceneInfo::FPrimitiveSceneInfo(const FPrimitiveSceneInfoAdapter& InAda
 
 	FMemory::Memzero(CachedReflectionCaptureProxies);
 
-#if RHI_RAYTRACING
-	// Cache static ray tracing geometries in SceneInfo to avoid having to access SceneProxy later
-	StaticRayTracingGeometries = InAdapter.SceneProxy->GetStaticRayTracingGeometries();
-	CachedRayTracingGeometry = nullptr;
-#endif
-
 	if (FInstanceCullingContext::IsGPUCullingEnabled())
 	{
 		GpuLodInstanceRadius = InAdapter.SceneProxy->GetGpuLodInstanceRadius();
@@ -384,8 +378,12 @@ FPrimitiveSceneInfo::~FPrimitiveSceneInfo()
 }
 
 #if RHI_RAYTRACING
-bool FPrimitiveSceneInfo::IsCachedRayTracingGeometryValid() const
+bool FPrimitiveSceneInfo::IsCachedRayTracingGeometryValid(int8 LODIndex) const
 {
+	check(LODIndex >= 0 && LODIndex < RayTracingLODData.Num());
+
+	const FRayTracingGeometry* CachedRayTracingGeometry = RayTracingLODData[LODIndex].RayTracingGeometry;
+
 	if (CachedRayTracingGeometry)
 	{
 		// TODO: Doesn't take Nanite Ray Tracing into account
@@ -418,7 +416,13 @@ void FPrimitiveSceneInfo::AllocateRayTracingSBT(const FRHIRayTracingGeometry* Ca
 		const FRHIRayTracingGeometry* RayTracingGeometry = nullptr;
 		uint32 SegmentCount = 0;
 
-		if (CachedRayTracingInstanceGeometryRHI)
+		if (const FRayTracingGeometry* StaticRayTracingGeometry = GetStaticRayTracingGeometry(LODIndex))
+		{
+			// If there is a valid FRayTracingGeometry, retrieve the RHI object and segment count from this object (RenderThread timeline valid)
+			 RayTracingGeometry = StaticRayTracingGeometry->GetRHI();
+			 SegmentCount = StaticRayTracingGeometry->Initializer.Segments.Num();
+		}
+		else if(CachedRayTracingInstanceGeometryRHI)
 		{
 			checkf(RayTracingLODData.Num() == 1, TEXT("We expect to have a single LOD when using CachedRayTracingInstanceGeometryRHI."));
 
@@ -430,12 +434,6 @@ void FPrimitiveSceneInfo::AllocateRayTracingSBT(const FRHIRayTracingGeometry* Ca
 			// Need to fix UE-309710 first
 			checkf(LODData.NumCachedMeshCommands <= (int32)SegmentCount, TEXT("Number of cached ray tracing mesh commands (%d), should be lower or equal to number of segments (%d)."), LODData.NumCachedMeshCommands, (int32)SegmentCount);
 		}
-		else if (FRayTracingGeometry* StaticRayTracingGeometry = GetStaticRayTracingGeometry(LODIndex))
-		{
-			// If there is a valid FRayTracingGeometry, retrieve the RHI object and segment count from this object (RenderThread timeline valid)
-			 RayTracingGeometry = StaticRayTracingGeometry->GetRHI();
-			 SegmentCount = StaticRayTracingGeometry->Initializer.Segments.Num();
-		}
 
 		if (RayTracingGeometry && SegmentCount > 0)
 		{
@@ -477,13 +475,10 @@ void FPrimitiveSceneInfo::CacheRayTracingShaderBindingData(const FRHIRayTracingG
 		const FRHIRayTracingGeometry* RayTracingGeometry = nullptr;
 		bool bDynamicRayTracingGeometry = false;
 
-		if (CachedRayTracingGeometry)
+		if (const FRayTracingGeometry* StaticRayTracingGeometry = GetStaticRayTracingGeometry(LODIndex))
 		{
-			checkf(RayTracingLODData.Num() == 1, TEXT("When using CachedRayTracingGeometry we only expect to have a single LOD"));
-
-			// If we have a valid cached raytracing instance geometry then use this one
-			RayTracingGeometry = CachedRayTracingGeometry->GetRHI();
-			bDynamicRayTracingGeometry = CachedRayTracingGeometry->Initializer.bAllowUpdate;
+			RayTracingGeometry = StaticRayTracingGeometry->GetRHI();
+			bDynamicRayTracingGeometry = StaticRayTracingGeometry->Initializer.bAllowUpdate;
 		}
 		else if (CachedRayTracingInstanceGeometryRHI)
 		{
@@ -492,11 +487,6 @@ void FPrimitiveSceneInfo::CacheRayTracingShaderBindingData(const FRHIRayTracingG
 			// this path is only used by Nanite Ray Tracing so bDynamicRayTracingGeometry = false
 			RayTracingGeometry = CachedRayTracingInstanceGeometryRHI;
 		}
-		else if (FRayTracingGeometry* StaticRayTracingGeometry = GetStaticRayTracingGeometry(LODIndex))
-		{
-			RayTracingGeometry = StaticRayTracingGeometry->GetRHI();
-			bDynamicRayTracingGeometry = StaticRayTracingGeometry->Initializer.bAllowUpdate;
-		}
 
 		if (LODData.SBTAllocation != nullptr)
 		{
@@ -531,8 +521,7 @@ void FPrimitiveSceneInfo::CacheRayTracingInstance(FRayTracingGeometryInstance In
 
 	const FPrimitiveSceneInfo::FRayTracingLODData& LODData = RayTracingLODData[0];
 
-	const FRayTracingGeometry* RayTracingGeometry = GetCachedRayTracingGeometry();
-	const int32 GeometryHandle = RayTracingGeometry ? RayTracingGeometry->GetGeometryHandle() : INDEX_NONE;
+	const int32 GeometryHandle = LODData.RayTracingGeometry ? LODData.RayTracingGeometry->GetGeometryHandle() : INDEX_NONE;
 
 	const bool bIsFarFieldInstance = Proxy->IsRayTracingFarField();
 
@@ -567,11 +556,13 @@ void FPrimitiveSceneInfo::CacheRayTracingInstance(FRayTracingGeometryInstance In
 	bIsCachedRayTracingInstanceValid = Instance.GeometryRHI != nullptr;
 }
 
-FRayTracingGeometry* FPrimitiveSceneInfo::GetStaticRayTracingGeometry(int8 LODIndex) const
+const FRayTracingGeometry* FPrimitiveSceneInfo::GetStaticRayTracingGeometry(int8 LODIndex) const
 {
-	if (LODIndex < StaticRayTracingGeometries.Num())
+	check(LODIndex >= 0);
+
+	if (LODIndex < RayTracingLODData.Num())
 	{
-		return StaticRayTracingGeometries[LODIndex];
+		return RayTracingLODData[LODIndex].RayTracingGeometry;
 	}
 	else
 	{
@@ -579,19 +570,21 @@ FRayTracingGeometry* FPrimitiveSceneInfo::GetStaticRayTracingGeometry(int8 LODIn
 	}
 }
 
-FRayTracingGeometry* FPrimitiveSceneInfo::GetValidStaticRayTracingGeometry(int8& InOutLODIndex) const
+const FRayTracingGeometry* FPrimitiveSceneInfo::GetValidStaticRayTracingGeometry(int8& InOutLODIndex) const
 {
 	// TODO: Move HasPendingBuildRequest() / BoostBuildPriority() out of this function
 
-	for (; InOutLODIndex < StaticRayTracingGeometries.Num(); ++InOutLODIndex)
+	for (; InOutLODIndex < RayTracingLODData.Num(); ++InOutLODIndex)
 	{
-		if (StaticRayTracingGeometries[InOutLODIndex]->HasPendingBuildRequest())
+		const FRayTracingGeometry* RayTracingGeometry = RayTracingLODData[InOutLODIndex].RayTracingGeometry;
+
+		if (RayTracingGeometry->HasPendingBuildRequest())
 		{
-			StaticRayTracingGeometries[InOutLODIndex]->BoostBuildPriority();
+			RayTracingGeometry->BoostBuildPriority();
 		}
-		else if (StaticRayTracingGeometries[InOutLODIndex]->IsValid() && !StaticRayTracingGeometries[InOutLODIndex]->IsEvicted())
+		else if (RayTracingGeometry->IsValid() && !RayTracingGeometry->IsEvicted())
 		{
-			return StaticRayTracingGeometries[InOutLODIndex];
+			return RayTracingGeometry;
 		}
 	}
 
@@ -1260,7 +1253,9 @@ void FPrimitiveSceneInfo::CacheRayTracingPrimitives(FScene* Scene, const TArrayV
 						FDeferredPrimitiveRayTracingData& DeferredData = Context.DeferredDatas.AddDefaulted_GetRef();
 						DeferredData.SceneInfo = SceneInfo;
 
-						const int32 LODCount = EnumHasAnyFlags(Flags, ERayTracingPrimitiveFlags::CacheInstances) ? 1 : SceneInfo->GetStaticRayTracingGeometryNum();
+						TArray<FRayTracingGeometry*> StaticRayTracingGeometries = SceneProxy->GetStaticRayTracingGeometries();
+
+						const int32 LODCount = EnumHasAnyFlags(Flags, ERayTracingPrimitiveFlags::CacheInstances) ? 1 : StaticRayTracingGeometries.Num();
 
 						TArray<FPrimitiveSceneInfo::FRayTracingLODData> RayTracingLODData;
 						RayTracingLODData.Empty(LODCount);
@@ -1268,10 +1263,20 @@ void FPrimitiveSceneInfo::CacheRayTracingPrimitives(FScene* Scene, const TArrayV
 
 						CacheRayTracingMeshCommands(RayTracingInstance.Materials, SceneProxy, RayTracingLODData, LODCount, Context.Commands, Context.CommandContext, Context.RayTracingMeshProcessor, DeferredData);
 
-						if (EnumHasAnyFlags(Flags, ERayTracingPrimitiveFlags::CacheInstances))
+						if (RayTracingInstance.Geometry)
 						{
-							SceneInfo->CachedRayTracingGeometry = RayTracingInstance.Geometry;
+							RayTracingLODData[0].RayTracingGeometry = RayTracingInstance.Geometry;
+						}
+						else if(!StaticRayTracingGeometries.IsEmpty())
+						{
+							for (int32 LODIndex = 0; LODIndex < LODCount; ++LODIndex)
+							{
+								RayTracingLODData[LODIndex].RayTracingGeometry = StaticRayTracingGeometries[LODIndex];
+							}
+						}
 
+						if (EnumHasAnyFlags(Flags, ERayTracingPrimitiveFlags::CacheInstances))
+						{
 							SetupCachedRayTracingInstance(SceneInfo, RayTracingInstance, DeferredData.Instance, DeferredData.SegmentCount);
 							DeferredData.bCacheInstance = true;
 						}
@@ -1420,7 +1425,8 @@ void FPrimitiveSceneInfo::SetCachedRayTracingInstanceGeometryRHI(FRHIRayTracingG
 	{
 		return;
 	}
-
+	
+	// TODO: check Nanite Ray Tracing is enabled
 	check(RayTracingLODData.Num() == 1);
 
 	FPrimitiveSceneInfo::FRayTracingLODData& LODData = RayTracingLODData[0];
diff --git a/Engine/Source/Runtime/Renderer/Private/RayTracing/RayTracing.cpp b/Engine/Source/Runtime/Renderer/Private/RayTracing/RayTracing.cpp
--- a/Engine/Source/Runtime/Renderer/Private/RayTracing/RayTracing.cpp
+++ b/Engine/Source/Runtime/Renderer/Private/RayTracing/RayTracing.cpp
@@ -1434,6 +1434,11 @@ namespace RayTracing
 
 				if (EnumHasAnyFlags(Flags, ERayTracingPrimitiveFlags::CacheInstances))
 				{
+					// For primitives with ERayTracingPrimitiveFlags::CacheInstances flag we only cache the instance/mesh commands of the current LOD
+					// (see FPrimitiveSceneInfo::UpdateCachedRayTracingInstance(...) and CacheRayTracingMeshCommands(...))
+					check(!EnumHasAnyFlags(Flags, ERayTracingPrimitiveFlags::ComputeLOD));
+					LODIndex = 0;
+
 					if (bUsingNaniteRayTracing)
 					{
 						if (!SceneInfo->bIsCachedRayTracingInstanceValid)
@@ -1445,16 +1450,11 @@ namespace RayTracing
 					else
 					{
 						// Currently IsCachedRayTracingGeometryValid() can only be called for non-nanite geometries
-						checkf(SceneInfo->IsCachedRayTracingGeometryValid(), TEXT("Cached ray tracing instance is expected to be valid. Was mesh LOD streamed but cached data was not invalidated?"));
+						checkf(SceneInfo->IsCachedRayTracingGeometryValid(LODIndex), TEXT("Cached ray tracing instance is expected to be valid. Was mesh LOD streamed but cached data was not invalidated?"));
 					}
 
 					checkf(SceneInfo->bIsCachedRayTracingInstanceValid, TEXT("Cached ray tracing instance must be valid."));
 
-					// For primitives with ERayTracingPrimitiveFlags::CacheInstances flag we only cache the instance/mesh commands of the current LOD
-					// (see FPrimitiveSceneInfo::UpdateCachedRayTracingInstance(...) and CacheRayTracingMeshCommands(...))
-					check(!EnumHasAnyFlags(Flags, ERayTracingPrimitiveFlags::ComputeLOD));
-					LODIndex = 0;
-
 					const FPrimitiveSceneInfo::FRayTracingLODData& RTLODData = SceneInfo->GetRayTracingLODData(LODIndex);
 
 					if (bExcludeDecals && RTLODData.CachedMeshCommandFlags.bAllSegmentsDecal)
@@ -1506,7 +1506,7 @@ namespace RayTracing
 #if DO_CHECK
 					if (bUsingReferenceBasedResidency)
 					{
-						FRayTracingGeometry* TargetRayTracingGeometry = SceneInfo->GetStaticRayTracingGeometry(LODIndex);
+						const FRayTracingGeometry* TargetRayTracingGeometry = SceneInfo->GetStaticRayTracingGeometry(LODIndex);
 
 						// TODO: Should have an assert here but disabled it due to UE-112448
 						if (TargetRayTracingGeometry != nullptr)
@@ -1518,7 +1518,7 @@ namespace RayTracing
 					}
 #endif
 
-					FRayTracingGeometry* RayTracingGeometry = SceneInfo->GetValidStaticRayTracingGeometry(LODIndex);
+					const FRayTracingGeometry* RayTracingGeometry = SceneInfo->GetValidStaticRayTracingGeometry(LODIndex);
 
 					if (LODIndex < 0)
 					{
diff --git a/Engine/Source/Runtime/Renderer/Public/PrimitiveSceneInfo.h b/Engine/Source/Runtime/Renderer/Public/PrimitiveSceneInfo.h
--- a/Engine/Source/Runtime/Renderer/Public/PrimitiveSceneInfo.h
+++ b/Engine/Source/Runtime/Renderer/Public/PrimitiveSceneInfo.h
@@ -526,14 +526,15 @@ public:
 #if RHI_RAYTRACING
 	static void UpdateCachedRaytracingData(FScene* Scene, const TArrayView<FPrimitiveSceneInfo*>& SceneInfos);
 
-	bool IsCachedRayTracingGeometryValid() const;
+	bool IsCachedRayTracingGeometryValid(int8 LODIndex) const;
 
-	RENDERER_API FRayTracingGeometry* GetStaticRayTracingGeometry(int8 LODIndex) const;
-	RENDERER_API FRayTracingGeometry* GetValidStaticRayTracingGeometry(int8& InOutLODIndex) const;
+	RENDERER_API const FRayTracingGeometry* GetStaticRayTracingGeometry(int8 LODIndex) const;
+	RENDERER_API const FRayTracingGeometry* GetValidStaticRayTracingGeometry(int8& InOutLODIndex) const;
 
-	int GetStaticRayTracingGeometryNum() const { return StaticRayTracingGeometries.Num(); }
+	int32 GetStaticRayTracingGeometryNum() const { return RayTracingLODData.Num(); }
 
-	const FRayTracingGeometry* GetCachedRayTracingGeometry() const { return CachedRayTracingGeometry; }
+	UE_DEPRECATED(5.7, "Use GetRayTracingLODData() instead.")
+	const FRayTracingGeometry* GetCachedRayTracingGeometry() const { return nullptr; }
 #endif
 
 	/** Return primitive fullname (for debugging only). */
@@ -668,6 +669,8 @@ public:
 
 	struct FRayTracingLODData
 	{
+		const FRayTracingGeometry* RayTracingGeometry = nullptr;
+
 		FRayTracingSBTAllocation* SBTAllocation = nullptr;
 		FRayTracingCachedMeshCommandFlags CachedMeshCommandFlags;
 		TArray<FRayTracingShaderBindingData> CachedShaderBindingDataBase;
@@ -783,12 +786,6 @@ private:
 	void RemoveCachedNaniteMaterialBins();
 
 #if RHI_RAYTRACING
-	TArray<FRayTracingGeometry*> StaticRayTracingGeometries;
-
-	// Cache pointer to FRayTracingGeometry used by cached ray tracing instance
-	// since primitives using ERayTracingPrimitiveFlags::CacheInstances don't fill the StaticRayTracingGeometries array above
-	const FRayTracingGeometry* CachedRayTracingGeometry;
-
 	/** Creates cached ray tracing representations for all meshes. */
 	static void CacheRayTracingPrimitives(FScene* Scene, const TArrayView<FPrimitiveSceneInfo*>& SceneInfos);
 
diff --git a/Engine/Source/Runtime/Renderer/Public/RayTracingInstance.h b/Engine/Source/Runtime/Renderer/Public/RayTracingInstance.h
--- a/Engine/Source/Runtime/Renderer/Public/RayTracingInstance.h
+++ b/Engine/Source/Runtime/Renderer/Public/RayTracingInstance.h
@@ -16,7 +16,7 @@ enum class ERayTracingInstanceLayer : uint8
 struct FRayTracingInstance
 {
 	/** The underlying geometry of this instance specification. */
-	const FRayTracingGeometry* Geometry;
+	const FRayTracingGeometry* Geometry = nullptr;
 
 	/**
 	 * Materials for each segment, in the form of mesh batches. We will check whether every segment of the geometry has been assigned a material.
