Various fixes to anim bank.
- Fixed dangling pointer with sequence / mapping data due to unstable TMap.
- Fixed refpose transform provider stomping over other transform allocations.
- Fixed lastupdateframecount being assigned in the async deletion task, which is not synchronized with the render thread.

#Synced-CL 49025386

diff --git a/Engine/Source/Runtime/Engine/Private/Animation/AnimBank.cpp b/Engine/Source/Runtime/Engine/Private/Animation/AnimBank.cpp
--- a/Engine/Source/Runtime/Engine/Private/Animation/AnimBank.cpp
+++ b/Engine/Source/Runtime/Engine/Private/Animation/AnimBank.cpp
@@ -1484,7 +1484,7 @@ bool UAnimBankData::IsAnimBankValid(FInstancedSkinningSceneExtensionProxy* Exten
 {
 	for (const FAnimBankItem& Item : AnimBankItems)
 	{
-		if (Item.BankAsset == nullptr || Item.BankAsset->Asset != ExtensionProxy->GetSkinnedAsset())
+		if (Item.BankAsset == nullptr /* || Item.BankAsset->Asset != ExtensionProxy->GetSkinnedAsset() */)
 		{
 			return false;
 		}
diff --git a/Engine/Source/Runtime/Renderer/Private/Skinning/AnimBankTransformProvider.cpp b/Engine/Source/Runtime/Renderer/Private/Skinning/AnimBankTransformProvider.cpp
--- a/Engine/Source/Runtime/Renderer/Private/Skinning/AnimBankTransformProvider.cpp
+++ b/Engine/Source/Runtime/Renderer/Private/Skinning/AnimBankTransformProvider.cpp
@@ -137,40 +137,46 @@ void FAnimBankTransformProvider::RegisterAnimBank(FAnimBankTransformProviderProx
 		const FAnimBankData& AnimBankData = Item.BankAsset->GetData();
 		const FSequenceRegistryKey SequenceRegistryKey(Item);
 
-		FMapping& Mapping = MappingRegistry.FindOrAdd(SequenceRegistryKey.BankAsset);
-		if (Mapping.RefCount == 0)
+		TUniquePtr<FMapping>& Mapping = MappingRegistry.FindOrAdd(SequenceRegistryKey.BankAsset);
+		if (!Mapping)
 		{
+			Mapping = MakeUnique<FMapping>();
+
 			const FSkinnedAssetMapping& AnimBankMapping = AnimBankData.Mapping;
 
-			Mapping.RegistryKey   = SequenceRegistryKey.BankAsset;
-			Mapping.Keys.Position = AnimBankMapping.PositionKeys;
-			Mapping.Keys.Rotation = AnimBankMapping.RotationKeys;
-			Mapping.BoneCount     = AnimBankMapping.BoneCount;
+			Mapping->RegistryKey   = SequenceRegistryKey.BankAsset;
+			Mapping->Keys.Position = AnimBankMapping.PositionKeys;
+			Mapping->Keys.Rotation = AnimBankMapping.RotationKeys;
+			Mapping->BoneCount     = AnimBankMapping.BoneCount;
 
-			MappingsToUpload.Add(&Mapping);
+			check(!MappingsToUpload.Contains(Mapping.Get()));
+			MappingsToUpload.AddUnique(Mapping.Get());
 		}
-		++Mapping.RefCount;
+		++Mapping->RefCount;
 
-		FSequence& Sequence = SequenceRegistry.FindOrAdd(SequenceRegistryKey);
-		if (Sequence.RefCount == 0)
+		TUniquePtr<FSequence>& Sequence = SequenceRegistry.FindOrAdd(SequenceRegistryKey);
+		if (!Sequence)
 		{
+			Sequence = MakeUnique<FSequence>();
+
 			const FAnimBankEntry& AnimBankEntry = AnimBankData.Entries[SequenceRegistryKey.SequenceIndex];
 
-			Sequence.RegistryKey   = SequenceRegistryKey;
-			Sequence.Keys.Position = AnimBankEntry.PositionKeys;
-			Sequence.Keys.Rotation = AnimBankEntry.RotationKeys;
-			Sequence.PlayLength    = AnimBankEntry.GetPlayLength();
-			Sequence.FrameCount    = AnimBankEntry.FrameCount;
-			Sequence.KeyCount      = AnimBankEntry.KeyCount;
-			Sequence.BoneCount     = Mapping.BoneCount;
-			Sequence.Mapping       = &Mapping;
-			Sequence.bLooping      = AnimBankEntry.IsLooping();
-
-			SequencesToUpload.AddUnique(&Sequence);
+			Sequence->RegistryKey   = SequenceRegistryKey;
+			Sequence->Keys.Position = AnimBankEntry.PositionKeys;
+			Sequence->Keys.Rotation = AnimBankEntry.RotationKeys;
+			Sequence->PlayLength    = AnimBankEntry.GetPlayLength();
+			Sequence->FrameCount    = AnimBankEntry.FrameCount;
+			Sequence->KeyCount      = AnimBankEntry.KeyCount;
+			Sequence->BoneCount     = Mapping->BoneCount;
+			Sequence->Mapping       = Mapping.Get();
+			Sequence->bLooping      = AnimBankEntry.IsLooping();
+
+			check(!SequencesToUpload.Contains(Sequence.Get()));
+			SequencesToUpload.AddUnique(Sequence.Get());
 		}
-		++Sequence.RefCount;
+		++Sequence->RefCount;
 
-		ProxyData.Sequences.Emplace(SequenceRegistryKey, Mapping, Sequence);
+		ProxyData.Sequences.Emplace(SequenceRegistryKey, *Mapping, *Sequence);
 	}
 }
 
@@ -184,36 +190,36 @@ void FAnimBankTransformProvider::UnregisterAnimBank(FAnimBankTransformProviderPr
 	for (const FSequenceView& SequenceView : ProxyData.Sequences)
 	{
 		const uint32 BankAssetHash = GetTypeHash(SequenceView.Key.BankAsset);
-		FMapping& Mapping = MappingRegistry.FindByHashChecked(BankAssetHash, SequenceView.Key.BankAsset);
-		check(Mapping.RefCount > 0);
-		--Mapping.RefCount;
+		TUniquePtr<FMapping>& Mapping = MappingRegistry.FindByHashChecked(BankAssetHash, SequenceView.Key.BankAsset);
+		check(Mapping && Mapping->RefCount > 0);
+		--Mapping->RefCount;
 
-		if (Mapping.RefCount == 0)
+		if (Mapping->RefCount == 0)
 		{
-			if (Mapping.BufferAllocation.IsValid())
+			if (Mapping->BufferAllocation.IsValid())
 			{
-				SequenceTransformAllocator.Free(Mapping.BufferAllocation.Offset, Mapping.BufferAllocation.Count);
-				Mapping.BufferAllocation = {};
+				SequenceTransformAllocator.Free(Mapping->BufferAllocation.Offset, Mapping->BufferAllocation.Count);
+				Mapping->BufferAllocation = {};
 			}
 
-			MappingsToUpload.Remove(&Mapping);
+			MappingsToUpload.RemoveSingleSwap(Mapping.Get(), EAllowShrinking::No);
 			MappingRegistry.RemoveByHash(BankAssetHash, SequenceView.Key.BankAsset);
 		}
 
 		const uint32 SequenceRegistryKeyHash = GetTypeHash(SequenceView.Key);
-		FSequence& Sequence = SequenceRegistry.FindByHashChecked(SequenceRegistryKeyHash, SequenceView.Key);
-		check(Sequence.RefCount > 0);
-		--Sequence.RefCount;
+		TUniquePtr<FSequence>& Sequence = SequenceRegistry.FindByHashChecked(SequenceRegistryKeyHash, SequenceView.Key);
+		check(Sequence && Sequence->RefCount > 0);
+		--Sequence->RefCount;
 
-		if (Sequence.RefCount == 0)
+		if (Sequence->RefCount == 0)
 		{
-			if (Sequence.BufferAllocation.IsValid())
+			if (Sequence->BufferAllocation.IsValid())
 			{
-				SequenceTransformAllocator.Free(Sequence.BufferAllocation.Offset, Mapping.BufferAllocation.Count);
-				Sequence.BufferAllocation = {};
+				SequenceTransformAllocator.Free(Sequence->BufferAllocation.Offset, Mapping->BufferAllocation.Count);
+				Sequence->BufferAllocation = {};
 			}
 
-			SequencesToUpload.Remove(&Sequence);
+			SequencesToUpload.RemoveSingleSwap(Sequence.Get(), EAllowShrinking::No);
 			SequenceRegistry.RemoveByHash(SequenceRegistryKeyHash, SequenceView.Key);
 		}
 	}
diff --git a/Engine/Source/Runtime/Renderer/Private/Skinning/AnimBankTransformProvider.h b/Engine/Source/Runtime/Renderer/Private/Skinning/AnimBankTransformProvider.h
--- a/Engine/Source/Runtime/Renderer/Private/Skinning/AnimBankTransformProvider.h
+++ b/Engine/Source/Runtime/Renderer/Private/Skinning/AnimBankTransformProvider.h
@@ -176,8 +176,8 @@ private:
 		bool bHasEmptyTracks = false;
 	};
 
-	TMap<FMappingRegistryKey, FMapping> MappingRegistry;
-	TMap<FSequenceRegistryKey, FSequence> SequenceRegistry;
+	TMap<FMappingRegistryKey, TUniquePtr<FMapping>> MappingRegistry;
+	TMap<FSequenceRegistryKey, TUniquePtr<FSequence>> SequenceRegistry;
 	TMap<const FSkinningSceneExtensionProxy*, FProxyData> ProxyDataMap;
 
 	TArray<FSequence*> SequencesToUpload;
diff --git a/Engine/Source/Runtime/Renderer/Private/Skinning/SkinningSceneExtension.cpp b/Engine/Source/Runtime/Renderer/Private/Skinning/SkinningSceneExtension.cpp
--- a/Engine/Source/Runtime/Renderer/Private/Skinning/SkinningSceneExtension.cpp
+++ b/Engine/Source/Runtime/Renderer/Private/Skinning/SkinningSceneExtension.cpp
@@ -756,8 +756,6 @@ void FSkinningSceneExtension::FUpdater::End()
 	// Ensure these tasks finish before we fall out of scope.
 	// NOTE: This should be unnecessary if the updater shares the graph builder's lifetime but we don't enforce that
 	SceneData->SyncAllTasks();
-
-	SceneData->LastUpdateFrameNumber = GFrameNumberRenderThread;
 }
 
 void FSkinningSceneExtension::FUpdater::PreSceneUpdate(FRDGBuilder& GraphBuilder, const FScenePreUpdateChangeSet& ChangeSet, FSceneUniformBuffer& SceneUniforms)
@@ -1419,6 +1417,11 @@ void FSkinningSceneExtension::FUpdater::PostMeshUpdate(FRDGBuilder& GraphBuilder
 		UE::Tasks::ETaskPriority::Normal,
 		bEnableAsync
 	);
+
+	GraphBuilder.AddPostExecuteCallback([SceneData = SceneData]
+	{
+		SceneData->LastUpdateFrameNumber = GFrameNumberRenderThread;
+	});
 }
 
 class FNaniteSkinningUpdateViewDataCS : public FGlobalShader
@@ -1625,8 +1628,8 @@ void FSkinningSceneExtension::ProvideRefPoseTransforms(FSkinningTransformProvide
 		const uint32 TransformCount = Proxy->GetMaxBoneTransformCount();
 		const uint32 AnimationCount = Proxy->GetUniqueAnimationCount();
 		const uint32 TotalTransformCount = TransformCount * AnimationCount;
-
-		uint32 TransformWrite = Indirection.CurrentTransformOffset;
+	
+		uint32 TransformWrite =  Indirection.CurrentTransformOffset < Indirection.PreviousTransformOffset ? Indirection.CurrentTransformOffset : Indirection.PreviousTransformOffset;
 
 		const uint32 FullBlockCount = TotalTransformCount / TransformsPerGroup;
 		for (uint32 BlockIndex = 0; BlockIndex < FullBlockCount; ++BlockIndex)
