PSO precaching reduction for mobile:
- skip collection if depth if no depth pass on mobile (also for default material)
- skip collection for rvt if not enabled on platform
- skip stencil state for decal passes which don't need it

#Synced-CL 49063813

diff --git a/Engine/Source/Runtime/Renderer/Private/CompositionLighting/PostProcessDeferredDecals.cpp b/Engine/Source/Runtime/Renderer/Private/CompositionLighting/PostProcessDeferredDecals.cpp
--- a/Engine/Source/Runtime/Renderer/Private/CompositionLighting/PostProcessDeferredDecals.cpp
+++ b/Engine/Source/Runtime/Renderer/Private/CompositionLighting/PostProcessDeferredDecals.cpp
@@ -635,8 +635,12 @@ void CollectDeferredDecalPassPSOInitializers(
 		AddDeferredDecalPSOReverseHanded(true /*bReverseCulling*/, bDecalUsesStencil);
 	};
 
-	AddDeferredDecalPSOReverseCulling(false /*bDecalUsesStencil*/);
-	AddDeferredDecalPSOReverseCulling(true /*bDecalUsesStencil*/);
+	AddDeferredDecalPSOReverseCulling(false /*bDecalUsesStencil*/);	
+	const bool bStencilThisDecal = IsStencilOptimizationAvailable(DecalRenderStage);
+	if (bStencilThisDecal)
+	{
+		AddDeferredDecalPSOReverseCulling(true /*bDecalUsesStencil*/);
+	}
 }
 
 void AddDeferredDecalPass(
diff --git a/Engine/Source/Runtime/Renderer/Private/DepthRendering.cpp b/Engine/Source/Runtime/Renderer/Private/DepthRendering.cpp
--- a/Engine/Source/Runtime/Renderer/Private/DepthRendering.cpp
+++ b/Engine/Source/Runtime/Renderer/Private/DepthRendering.cpp
@@ -1095,6 +1095,12 @@ void FDepthPassMeshProcessor::AddMeshBatch(const FMeshBatch& RESTRICT MeshBatch,
 
 void FDepthPassMeshProcessor::CollectPSOInitializers(const FSceneTexturesConfig& SceneTexturesConfig, const FMaterial& Material, const FPSOPrecacheVertexFactoryData& VertexFactoryData, const FPSOPrecacheParams& PreCacheParams, TArray<FPSOPrecacheData>& PSOInitializers)
 {		
+	if (FeatureLevel == ERHIFeatureLevel::ES3_1 && EarlyZPassMode == DDM_None)
+	{
+		// Do not cache MDC and do not pre-cache PSOs for a depth pass if it's never going to be used on mobile platforms
+		return;
+	}
+
 	// Are we currently collecting PSO's for the default material
 	if (PreCacheParams.bDefaultMaterial)
 	{
diff --git a/Engine/Source/Runtime/Renderer/Private/VT/RuntimeVirtualTextureRender.cpp b/Engine/Source/Runtime/Renderer/Private/VT/RuntimeVirtualTextureRender.cpp
--- a/Engine/Source/Runtime/Renderer/Private/VT/RuntimeVirtualTextureRender.cpp
+++ b/Engine/Source/Runtime/Renderer/Private/VT/RuntimeVirtualTextureRender.cpp
@@ -1533,6 +1533,12 @@ namespace RuntimeVirtualTexture
 			const FPSOPrecacheParams& PreCacheParams, 
 			TArray<FPSOPrecacheData>& PSOInitializers) override final
 		{
+			EShaderPlatform ShaderPlatform = GetFeatureLevelShaderPlatform(FeatureLevel);
+			if (!UseVirtualTexturing(ShaderPlatform))
+			{
+				return;
+			}
+
 			const uint8 OutputAttributeMask = Material.IsDefaultMaterial() ? 0xff : Material.GetRuntimeVirtualTextureOutputAttibuteMask_GameThread();
 
 			if (OutputAttributeMask != 0)
