Change Bloom Mobile passes to render into fixed size RTs and with a dynamic resolution tied viewport instead, moving sample clamping to the PS.
This prevent RT reallocations during dynamic resolution variations.

#Synced-CL 49027503

diff --git a/Engine/Shaders/Private/PostProcessMobile.usf b/Engine/Shaders/Private/PostProcessMobile.usf
--- a/Engine/Shaders/Private/PostProcessMobile.usf
+++ b/Engine/Shaders/Private/PostProcessMobile.usf
@@ -64,7 +64,6 @@ float4 BufferSizeAndInvSize;
 float4 DofBlurSizeAndInvSize;
 float4 BufferASizeAndInvSize;
 float4 BufferBSizeAndInvSize;
-float4 BloomUpSizeAndInvSize;
 float4 BufferViewRect;
 
 // Point on circle.
@@ -485,75 +484,70 @@ void BloomDownVS_Mobile(
 	in float4 InPosition : ATTRIBUTE0,
 	in float2 InTexCoord : ATTRIBUTE1,
 	in FStereoVSInput StereoInput,
-	out float4 OutTexCoords[8] : TEXCOORD0,
+	out float2 OutTexCoord : TEXCOORD0,
 	out FStereoVSOutput StereoOutput,
 	out float4 OutPosition : SV_POSITION
 	)
 {
 	StereoSetupVS(StereoInput, StereoOutput);
+	DrawRectangle(InPosition, InTexCoord, OutPosition, OutTexCoord);
+}
 
-	DrawRectangle(InPosition, InTexCoord, OutPosition, InTexCoord);
 
-	float Start = 2.0/14.0;
-	float Scale = BloomDownScale;
-
-	OutTexCoords[0].xy = InTexCoord.xy;
-	OutTexCoords[0].zw = InTexCoord.xy + Circle(Start, 14.0, 0.0) * Scale * BufferSizeAndInvSize.zw;
-	OutTexCoords[1].xy = InTexCoord.xy + Circle(Start, 14.0, 1.0) * Scale * BufferSizeAndInvSize.zw;
-	OutTexCoords[1].zw = InTexCoord.xy + Circle(Start, 14.0, 2.0) * Scale * BufferSizeAndInvSize.zw;
-	OutTexCoords[2].xy = InTexCoord.xy + Circle(Start, 14.0, 3.0) * Scale * BufferSizeAndInvSize.zw;
-	OutTexCoords[2].zw = InTexCoord.xy + Circle(Start, 14.0, 4.0) * Scale * BufferSizeAndInvSize.zw;
-	OutTexCoords[3].xy = InTexCoord.xy + Circle(Start, 14.0, 5.0) * Scale * BufferSizeAndInvSize.zw;
-	OutTexCoords[3].zw = InTexCoord.xy + Circle(Start, 14.0, 6.0) * Scale * BufferSizeAndInvSize.zw;
-	OutTexCoords[4].xy = InTexCoord.xy + Circle(Start, 14.0, 7.0) * Scale * BufferSizeAndInvSize.zw;
-	OutTexCoords[4].zw = InTexCoord.xy + Circle(Start, 14.0, 8.0) * Scale * BufferSizeAndInvSize.zw;
-	OutTexCoords[5].xy = InTexCoord.xy + Circle(Start, 14.0, 9.0) * Scale * BufferSizeAndInvSize.zw;
-	OutTexCoords[5].zw = InTexCoord.xy + Circle(Start, 14.0, 10.0) * Scale * BufferSizeAndInvSize.zw;
-	OutTexCoords[6].xy = InTexCoord.xy + Circle(Start, 14.0, 11.0) * Scale * BufferSizeAndInvSize.zw;
-	OutTexCoords[6].zw = InTexCoord.xy + Circle(Start, 14.0, 12.0) * Scale * BufferSizeAndInvSize.zw;
-	OutTexCoords[7].xy = InTexCoord.xy + Circle(Start, 14.0, 13.0) * Scale * BufferSizeAndInvSize.zw;
-	OutTexCoords[7].zw = float2(0.0, 0.0);
-}
+float4 InputViewportRect;     // xy = Min, zw = Size
+float4 TextureSizeAndInvSize; // xy = TextureSize, zw = 1/TextureSize
 
 void BloomDownPS_Mobile(
-	float4 InUVs[8] : TEXCOORD0,
+	float2 InUV : TEXCOORD0,
 	in FStereoPSInput StereoInput,
 	out HALF4_TYPE OutColor : SV_Target0
 	)
 {
 	const uint EyeIndex = GetEyeIndex(StereoInput);
-	half4 N0 = SampleBloomDown(InUVs[0].xy, EyeIndex).rgba;
-	half4 N1 = SampleBloomDown(InUVs[0].zw, EyeIndex).rgba;
-	half4 N2 = SampleBloomDown(InUVs[1].xy, EyeIndex).rgba;
-	half4 N3 = SampleBloomDown(InUVs[1].zw, EyeIndex).rgba;
-	half4 N4 = SampleBloomDown(InUVs[2].xy, EyeIndex).rgba;
-	half4 N5 = SampleBloomDown(InUVs[2].zw, EyeIndex).rgba;
-	half4 N6 = SampleBloomDown(InUVs[3].xy, EyeIndex).rgba;
-	half4 N7 = SampleBloomDown(InUVs[3].zw, EyeIndex).rgba;
-	half4 N8 = SampleBloomDown(InUVs[4].xy, EyeIndex).rgba;
-	half4 N9 = SampleBloomDown(InUVs[4].zw, EyeIndex).rgba;
-	half4 N10 = SampleBloomDown(InUVs[5].xy, EyeIndex).rgba;
-	half4 N11 = SampleBloomDown(InUVs[5].zw, EyeIndex).rgba;
-	half4 N12 = SampleBloomDown(InUVs[6].xy, EyeIndex).rgba;
-	half4 N13 = SampleBloomDown(InUVs[6].zw, EyeIndex).rgba;
-	half4 N14 = SampleBloomDown(InUVs[7].xy, EyeIndex).rgba;
-	float W = 1.0/15.0;
-	OutColor.rgb = 
-		(N0 * W) +
-		(N1 * W) +
-		(N2 * W) +
-		(N3 * W) +
-		(N4 * W) +
-		(N5 * W) +
-		(N6 * W) +
-		(N7 * W) +
-		(N8 * W) +
-		(N9 * W) +
-		(N10 * W) +
-		(N11 * W) +
-		(N12 * W) +
-		(N13 * W) +
-		(N14 * W);
+
+	const float2 UVMax = (InputViewportRect.xy + InputViewportRect.zw - 0.5f) * TextureSizeAndInvSize.zw;
+
+	const float Start = 2.0 / 14.0;
+	const float Scale = BloomDownScale;
+
+	// Circle offsets in UV space
+	float2 Offset0  = Circle(Start, 14.0, 0.0)  * Scale * TextureSizeAndInvSize.zw;
+	float2 Offset1  = Circle(Start, 14.0, 1.0)  * Scale * TextureSizeAndInvSize.zw;
+	float2 Offset2  = Circle(Start, 14.0, 2.0)  * Scale * TextureSizeAndInvSize.zw;
+	float2 Offset3  = Circle(Start, 14.0, 3.0)  * Scale * TextureSizeAndInvSize.zw;
+	float2 Offset4  = Circle(Start, 14.0, 4.0)  * Scale * TextureSizeAndInvSize.zw;
+	float2 Offset5  = Circle(Start, 14.0, 5.0)  * Scale * TextureSizeAndInvSize.zw;
+	float2 Offset6  = Circle(Start, 14.0, 6.0)  * Scale * TextureSizeAndInvSize.zw;
+	float2 Offset7  = Circle(Start, 14.0, 7.0)  * Scale * TextureSizeAndInvSize.zw;
+	float2 Offset8  = Circle(Start, 14.0, 8.0)  * Scale * TextureSizeAndInvSize.zw;
+	float2 Offset9  = Circle(Start, 14.0, 9.0)  * Scale * TextureSizeAndInvSize.zw;
+	float2 Offset10 = Circle(Start, 14.0, 10.0) * Scale * TextureSizeAndInvSize.zw;
+	float2 Offset11 = Circle(Start, 14.0, 11.0) * Scale * TextureSizeAndInvSize.zw;
+	float2 Offset12 = Circle(Start, 14.0, 12.0) * Scale * TextureSizeAndInvSize.zw;
+	float2 Offset13 = Circle(Start, 14.0, 13.0) * Scale * TextureSizeAndInvSize.zw;
+
+	half4 N0 = SampleBloomDown(min(InUV,            UVMax), EyeIndex);
+	half4 N1 = SampleBloomDown(min(InUV + Offset0,  UVMax), EyeIndex);
+	half4 N2 = SampleBloomDown(min(InUV + Offset1,  UVMax), EyeIndex);
+	half4 N3 = SampleBloomDown(min(InUV + Offset2,  UVMax), EyeIndex);
+	half4 N4 = SampleBloomDown(min(InUV + Offset3,  UVMax), EyeIndex);
+	half4 N5 = SampleBloomDown(min(InUV + Offset4,  UVMax), EyeIndex);
+	half4 N6 = SampleBloomDown(min(InUV + Offset5,  UVMax), EyeIndex);
+	half4 N7 = SampleBloomDown(min(InUV + Offset6,  UVMax), EyeIndex);
+	half4 N8 = SampleBloomDown(min(InUV + Offset7,  UVMax), EyeIndex);
+	half4 N9 = SampleBloomDown(min(InUV + Offset8,  UVMax), EyeIndex);
+	half4 N10 = SampleBloomDown(min(InUV + Offset9,  UVMax), EyeIndex);
+	half4 N11 = SampleBloomDown(min(InUV + Offset10, UVMax), EyeIndex);
+	half4 N12 = SampleBloomDown(min(InUV + Offset11, UVMax), EyeIndex);
+	half4 N13 = SampleBloomDown(min(InUV + Offset12, UVMax), EyeIndex);
+	half4 N14 = SampleBloomDown(min(InUV + Offset13, UVMax), EyeIndex);
+
+	float W = 1.0 / 15.0;
+	OutColor.rgb =
+		(N0 * W) + (N1 * W) + (N2 * W) + (N3 * W) + 
+		(N4 * W) + (N5 * W) + (N6 * W) + (N7 * W) + 
+		(N8 * W) + (N9 * W) + (N10 * W) + (N11 * W) + 
+		(N12 * W) + (N13 * W) + (N14 * W);
 	OutColor.a = 0;
 }
 
@@ -568,95 +562,89 @@ void BloomUpVS_Mobile(
 	in float4 InPosition : ATTRIBUTE0,
 	in float2 InTexCoord : ATTRIBUTE1,
 	in FStereoVSInput StereoInput,
-	out float4 OutTexCoords[8] : TEXCOORD0,
+	out float2 OutTexCoord : TEXCOORD0,
 	out FStereoVSOutput StereoOutput,
 	out float4 OutPosition : SV_POSITION
 	)
 {
 	StereoSetupVS(StereoInput, StereoOutput);
-	DrawRectangle(InPosition, InTexCoord, OutPosition, InTexCoord);
-
-	float Start;
-	float Scale;
-
-	Start = 2.0/7.0;
-	Scale = BloomUpScales.x;
-
-
-	OutTexCoords[0].xy = InTexCoord.xy + Circle(Start, 7.0, 0.0) * Scale * BufferASizeAndInvSize.zw;
-	OutTexCoords[0].zw = InTexCoord.xy + Circle(Start, 7.0, 1.0) * Scale * BufferASizeAndInvSize.zw;
-	OutTexCoords[1].xy = InTexCoord.xy + Circle(Start, 7.0, 2.0) * Scale * BufferASizeAndInvSize.zw;
-	OutTexCoords[1].zw = InTexCoord.xy + Circle(Start, 7.0, 3.0) * Scale * BufferASizeAndInvSize.zw;
-	OutTexCoords[2].xy = InTexCoord.xy + Circle(Start, 7.0, 4.0) * Scale * BufferASizeAndInvSize.zw;
-	OutTexCoords[2].zw = InTexCoord.xy + Circle(Start, 7.0, 5.0) * Scale * BufferASizeAndInvSize.zw;
-	OutTexCoords[3].xy = InTexCoord.xy + Circle(Start, 7.0, 6.0) * Scale * BufferASizeAndInvSize.zw;
-
-	OutTexCoords[3].zw = InTexCoord.xy;
-
-	Start = 2.0/7.0;
-	Scale = BloomUpScales.y;
-
-	OutTexCoords[4].xy = InTexCoord.xy + Circle(Start, 7.0, 0.0) * Scale * BufferBSizeAndInvSize.zw;
-	OutTexCoords[4].zw = InTexCoord.xy + Circle(Start, 7.0, 1.0) * Scale * BufferBSizeAndInvSize.zw;
-	OutTexCoords[5].xy = InTexCoord.xy + Circle(Start, 7.0, 2.0) * Scale * BufferBSizeAndInvSize.zw;
-	OutTexCoords[5].zw = InTexCoord.xy + Circle(Start, 7.0, 3.0) * Scale * BufferBSizeAndInvSize.zw;
-	OutTexCoords[6].xy = InTexCoord.xy + Circle(Start, 7.0, 4.0) * Scale * BufferBSizeAndInvSize.zw;
-	OutTexCoords[6].zw = InTexCoord.xy + Circle(Start, 7.0, 5.0) * Scale * BufferBSizeAndInvSize.zw;
-	OutTexCoords[7].xy = InTexCoord.xy + Circle(Start, 7.0, 6.0) * Scale * BufferBSizeAndInvSize.zw;
-	OutTexCoords[7].zw = float2(0.0, 0.0);
+	DrawRectangle(InPosition, InTexCoord, OutPosition, OutTexCoord);
 }
 
+
 float4 BloomTintA;
 float4 BloomTintB;
 
+float4 ViewportARect;           // xy = Min, zw = Size)
+float4 ViewportBRect;           // xy = Min, zw = Size)
+float4 TextureASizeAndInvSize;  // xy = TextureSize, zw = 1/TextureSize
+float4 TextureBSizeAndInvSize;  // xy = TextureSize, zw = 1/TextureSize
+
 void BloomUpPS_Mobile(
-	float4 InUVs[8] : TEXCOORD0,
+	float2 InUV : TEXCOORD0,
 	in FStereoPSInput StereoInput,
 	out HALF4_TYPE OutColor : SV_Target0
 	)
 {
 	const uint EyeIndex = GetEyeIndex(StereoInput);
 
-	half3 A0 = SampleBloomUpA(InUVs[0].xy, EyeIndex).rgb;
-	half3 A1 = SampleBloomUpA(InUVs[0].zw, EyeIndex).rgb;
-	half3 A2 = SampleBloomUpA(InUVs[1].xy, EyeIndex).rgb;
-	half3 A3 = SampleBloomUpA(InUVs[1].zw, EyeIndex).rgb;
-	half3 A4 = SampleBloomUpA(InUVs[2].xy, EyeIndex).rgb;
-	half3 A5 = SampleBloomUpA(InUVs[2].zw, EyeIndex).rgb;
-	half3 A6 = SampleBloomUpA(InUVs[3].xy, EyeIndex).rgb;
-	half3 A7 = SampleBloomUpA(InUVs[3].zw, EyeIndex).rgb;
-										   
-	half3 B0 = SampleBloomUpB(InUVs[3].zw, EyeIndex).rgb;
-	half3 B1 = SampleBloomUpB(InUVs[4].xy, EyeIndex).rgb;
-	half3 B2 = SampleBloomUpB(InUVs[4].zw, EyeIndex).rgb;
-	half3 B3 = SampleBloomUpB(InUVs[5].xy, EyeIndex).rgb;
-	half3 B4 = SampleBloomUpB(InUVs[5].zw, EyeIndex).rgb;
-	half3 B5 = SampleBloomUpB(InUVs[6].xy, EyeIndex).rgb;
-	half3 B6 = SampleBloomUpB(InUVs[6].zw, EyeIndex).rgb;
-	half3 B7 = SampleBloomUpB(InUVs[7].xy, EyeIndex).rgb;
+	const float2 ViewportAUVScale = ViewportARect.zw * TextureASizeAndInvSize.zw;
+	const float2 ViewportBUVScale = ViewportBRect.zw * TextureBSizeAndInvSize.zw;
+
+	const float2 BaseUV_A = ViewportARect.xy * TextureASizeAndInvSize.zw + InUV * ViewportAUVScale;
+	const float2 BaseUV_B = ViewportBRect.xy * TextureBSizeAndInvSize.zw + InUV * ViewportBUVScale;
+
+	const float2 UVMinA = ViewportARect.xy * TextureASizeAndInvSize.zw;
+	const float2 UVMaxA = (ViewportARect.xy + ViewportARect.zw - 0.5f) * TextureASizeAndInvSize.zw;
+	const float2 UVMinB = ViewportBRect.xy * TextureBSizeAndInvSize.zw;
+	const float2 UVMaxB = (ViewportBRect.xy + ViewportBRect.zw - 0.5f) * TextureBSizeAndInvSize.zw;
+
+	const float Start = 2.0 / 7.0;
+	const float ScaleA = BloomUpScales.x;
+	const float ScaleB = BloomUpScales.y;
+
+	const float2 OffsetA0 = Circle(Start, 7.0, 0.0) * ScaleA * TextureASizeAndInvSize.zw;
+	const float2 OffsetA1 = Circle(Start, 7.0, 1.0) * ScaleA * TextureASizeAndInvSize.zw;
+	const float2 OffsetA2 = Circle(Start, 7.0, 2.0) * ScaleA * TextureASizeAndInvSize.zw;
+	const float2 OffsetA3 = Circle(Start, 7.0, 3.0) * ScaleA * TextureASizeAndInvSize.zw;
+	const float2 OffsetA4 = Circle(Start, 7.0, 4.0) * ScaleA * TextureASizeAndInvSize.zw;
+	const float2 OffsetA5 = Circle(Start, 7.0, 5.0) * ScaleA * TextureASizeAndInvSize.zw;
+	const float2 OffsetA6 = Circle(Start, 7.0, 6.0) * ScaleA * TextureASizeAndInvSize.zw;
+
+	const float2 OffsetB0 = Circle(Start, 7.0, 0.0) * ScaleB * TextureBSizeAndInvSize.zw;
+	const float2 OffsetB1 = Circle(Start, 7.0, 1.0) * ScaleB * TextureBSizeAndInvSize.zw;
+	const float2 OffsetB2 = Circle(Start, 7.0, 2.0) * ScaleB * TextureBSizeAndInvSize.zw;
+	const float2 OffsetB3 = Circle(Start, 7.0, 3.0) * ScaleB * TextureBSizeAndInvSize.zw;
+	const float2 OffsetB4 = Circle(Start, 7.0, 4.0) * ScaleB * TextureBSizeAndInvSize.zw;
+	const float2 OffsetB5 = Circle(Start, 7.0, 5.0) * ScaleB * TextureBSizeAndInvSize.zw;
+	const float2 OffsetB6 = Circle(Start, 7.0, 6.0) * ScaleB * TextureBSizeAndInvSize.zw;
+
+	const half3 A0 = SampleBloomUpA(clamp(BaseUV_A + OffsetA0, UVMinA, UVMaxA), EyeIndex).rgb;
+	const half3 A1 = SampleBloomUpA(clamp(BaseUV_A + OffsetA1, UVMinA, UVMaxA), EyeIndex).rgb;
+	const half3 A2 = SampleBloomUpA(clamp(BaseUV_A + OffsetA2, UVMinA, UVMaxA), EyeIndex).rgb;
+	const half3 A3 = SampleBloomUpA(clamp(BaseUV_A + OffsetA3, UVMinA, UVMaxA), EyeIndex).rgb;
+	const half3 A4 = SampleBloomUpA(clamp(BaseUV_A + OffsetA4, UVMinA, UVMaxA), EyeIndex).rgb;
+	const half3 A5 = SampleBloomUpA(clamp(BaseUV_A + OffsetA5, UVMinA, UVMaxA), EyeIndex).rgb;
+	const half3 A6 = SampleBloomUpA(clamp(BaseUV_A + OffsetA6, UVMinA, UVMaxA), EyeIndex).rgb;
+	const half3 A7 = SampleBloomUpA(clamp(BaseUV_A,            UVMinA, UVMaxA), EyeIndex).rgb;
+
+	const half3 B0 = SampleBloomUpB(clamp(BaseUV_B,            UVMinB, UVMaxB), EyeIndex).rgb;
+	const half3 B1 = SampleBloomUpB(clamp(BaseUV_B + OffsetB0, UVMinB, UVMaxB), EyeIndex).rgb;
+	const half3 B2 = SampleBloomUpB(clamp(BaseUV_B + OffsetB1, UVMinB, UVMaxB), EyeIndex).rgb;
+	const half3 B3 = SampleBloomUpB(clamp(BaseUV_B + OffsetB2, UVMinB, UVMaxB), EyeIndex).rgb;
+	const half3 B4 = SampleBloomUpB(clamp(BaseUV_B + OffsetB3, UVMinB, UVMaxB), EyeIndex).rgb;
+	const half3 B5 = SampleBloomUpB(clamp(BaseUV_B + OffsetB4, UVMinB, UVMaxB), EyeIndex).rgb;
+	const half3 B6 = SampleBloomUpB(clamp(BaseUV_B + OffsetB5, UVMinB, UVMaxB), EyeIndex).rgb;
+	const half3 B7 = SampleBloomUpB(clamp(BaseUV_B + OffsetB6, UVMinB, UVMaxB), EyeIndex).rgb;
 
 	// A is the same size source.
-	half3 WA = BloomTintA.rgb;
+	const half3 WA = BloomTintA.rgb;
 	// B is the upsampled source.
-	half3 WB = BloomTintB.rgb;
-
-	OutColor.rgb = 
-		A0 * WA + 
-		A1 * WA + 
-		A2 * WA + 
-		A3 * WA + 
-		A4 * WA + 
-		A5 * WA + 
-		A6 * WA + 
-		A7 * WA +
-		B0 * WB + 
-		B1 * WB + 
-		B2 * WB + 
-		B3 * WB + 
-		B4 * WB + 
-		B5 * WB + 
-		B6 * WB + 
-		B7 * WB;
+	const half3 WB = BloomTintB.rgb;
+
+	OutColor.rgb =
+		A0 * WA + A1 * WA + A2 * WA + A3 * WA + A4 * WA + A5 * WA + A6 * WA + A7 * WA +
+		B0 * WB + B1 * WB + B2 * WB + B3 * WB + B4 * WB + B5 * WB + B6 * WB + B7 * WB;
 	OutColor.a = 0;
 }
 
@@ -1153,44 +1141,33 @@ void SunBlurPS_Mobile(
 
 #define SUN_M 0.25
 
+float4 BloomUpTextureSizeAndInvSize; 	//xy=Size,zw=1/TexSize
+
+
 void SunMergeVS_Mobile(
 	in float4 InPosition : ATTRIBUTE0,
 	in float2 InTexCoord : ATTRIBUTE1,
 	in FStereoVSInput StereoInput,
-	out float4 OutTexCoordVignette : TEXCOORD0,
-	out float4 OutTexCoords[7] : TEXCOORD1,
+	out float2 OutTexCoord : TEXCOORD0,
+	out float4 OutSunShaftUVs[4] : TEXCOORD1,
+	out float2 OutVignetteSpace : TEXCOORD5,
 	out FStereoVSOutput StereoOutput,
 	out float4 OutPosition : SV_POSITION
 	)
 {
 	StereoSetupVS(StereoInput, StereoOutput);
-	DrawRectangle(InPosition, InTexCoord, OutPosition, InTexCoord);
-
-	OutTexCoordVignette.xy = InTexCoord.xy;
-	OutTexCoordVignette.zw =  VignetteSpace(OutPosition.xy);
-
-	float Start;
-	float Scale;
-
-	Start = 2.0/6.0;
-	Scale = 0.66/2.0;
-
-	OutTexCoords[0].xy = InTexCoord.xy + Circle(Start, 6.0, 0.0) * Scale * BloomUpSizeAndInvSize.zw;
-	OutTexCoords[1].xy = InTexCoord.xy + Circle(Start, 6.0, 1.0) * Scale * BloomUpSizeAndInvSize.zw;
-	OutTexCoords[2].xy = InTexCoord.xy + Circle(Start, 6.0, 2.0) * Scale * BloomUpSizeAndInvSize.zw;
-	OutTexCoords[3].xy = InTexCoord.xy + Circle(Start, 6.0, 3.0) * Scale * BloomUpSizeAndInvSize.zw;
-	OutTexCoords[4].xy = InTexCoord.xy + Circle(Start, 6.0, 4.0) * Scale * BloomUpSizeAndInvSize.zw;
-	OutTexCoords[5].xy = InTexCoord.xy + Circle(Start, 6.0, 5.0) * Scale * BloomUpSizeAndInvSize.zw;
-
-    OutTexCoords[0].zw = SunShaftRect(OutPosition.xy, 1.0 - SUN_P0 * SUN_M);
-    OutTexCoords[1].zw = SunShaftRect(OutPosition.xy, 1.0 - SUN_P1 * SUN_M);
-    OutTexCoords[2].zw = SunShaftRect(OutPosition.xy, 1.0 - SUN_P2 * SUN_M);
-    OutTexCoords[3].zw = SunShaftRect(OutPosition.xy, 1.0 - SUN_P3 * SUN_M);
-    OutTexCoords[4].zw = SunShaftRect(OutPosition.xy, 1.0 - SUN_P4 * SUN_M);
-    OutTexCoords[5].zw = SunShaftRect(OutPosition.xy, 1.0 - SUN_P5 * SUN_M);
-    OutTexCoords[6].xy = SunShaftRect(OutPosition.xy, 1.0 - SUN_P6 * SUN_M);
-	OutTexCoords[6].zw = float2(0.0, 0.0);
-
+	DrawRectangle(InPosition, InTexCoord, OutPosition, OutTexCoord);
+
+	OutSunShaftUVs[0].xy = SunShaftRect(OutPosition.xy, 1.0 - SUN_P0 * SUN_M);
+	OutSunShaftUVs[0].zw = SunShaftRect(OutPosition.xy, 1.0 - SUN_P1 * SUN_M);
+	OutSunShaftUVs[1].xy = SunShaftRect(OutPosition.xy, 1.0 - SUN_P2 * SUN_M);
+	OutSunShaftUVs[1].zw = SunShaftRect(OutPosition.xy, 1.0 - SUN_P3 * SUN_M);
+	OutSunShaftUVs[2].xy = SunShaftRect(OutPosition.xy, 1.0 - SUN_P4 * SUN_M);
+	OutSunShaftUVs[2].zw = SunShaftRect(OutPosition.xy, 1.0 - SUN_P5 * SUN_M);
+	OutSunShaftUVs[3].xy = SunShaftRect(OutPosition.xy, 1.0 - SUN_P6 * SUN_M);
+	OutSunShaftUVs[3].zw = float2(0.0, 0.0);
+
+    OutVignetteSpace = VignetteSpace(OutPosition.xy);
 }
 
 #undef SUN_M
@@ -1202,33 +1179,61 @@ Texture2D BloomDirtMaskTexture;
 SamplerState BloomDirtMaskSampler;
 float4 BloomDirtMaskTint;
 
+float4 SunMergeBloomSetupViewportRect;
+float4 SunMergeBloomSetupTextureSizeAndInvSize;
+float4 SunMergeBloomUpViewportRect;
+float4 SunMergeBloomUpTextureSizeAndInvSize;
+
 void SunMergePS_Mobile(
-	float4 InUVVignette : TEXCOORD0,
-	float4 InUVs[7] : TEXCOORD1,
+	float2 InUV : TEXCOORD0,
+	float4 InSunShaftUVs[4] : TEXCOORD1,
+	float2 InVignetteSpace : TEXCOORD5,
 	in FStereoPSInput StereoInput,
 	out HALF4_TYPE OutColor : SV_Target0
 	)
 {
 	const uint EyeIndex = GetEyeIndex(StereoInput);
 
-	#if MOBILE_USEBLOOM
+	// UV scaling
+	const float2 BloomSetupUVScale = SunMergeBloomSetupViewportRect.zw * SunMergeBloomSetupTextureSizeAndInvSize.zw;
+	const float2 BloomUpUVScale = SunMergeBloomUpViewportRect.zw * SunMergeBloomUpTextureSizeAndInvSize.zw;
+	const float2 BloomSetupUV = SunMergeBloomSetupViewportRect.xy * SunMergeBloomSetupTextureSizeAndInvSize.zw + InUV * BloomSetupUVScale;
+	const float2 BloomUpUV = SunMergeBloomUpViewportRect.xy * SunMergeBloomUpTextureSizeAndInvSize.zw + InUV * BloomUpUVScale;
+
+	// UV clamping
+	const float2 BloomSetupUVMin = SunMergeBloomSetupViewportRect.xy * SunMergeBloomSetupTextureSizeAndInvSize.zw;
+	const float2 BloomSetupUVMax = (SunMergeBloomSetupViewportRect.xy + SunMergeBloomSetupViewportRect.zw - 0.5f) * SunMergeBloomSetupTextureSizeAndInvSize.zw;
+	const float2 BloomUpUVMin = SunMergeBloomUpViewportRect.xy * SunMergeBloomUpTextureSizeAndInvSize.zw;
+	const float2 BloomUpUVMax = (SunMergeBloomUpViewportRect.xy + SunMergeBloomUpViewportRect.zw - 0.5f) * SunMergeBloomUpTextureSizeAndInvSize.zw;
+
+	// Circle sampling in PS
+	const float Start = 2.0 / 6.0;
+	const float Scale = 0.66 / 2.0;
+
+	const float2 Offset0 = Circle(Start, 6.0, 0.0) * Scale * SunMergeBloomSetupTextureSizeAndInvSize.zw;
+	const float2 Offset1 = Circle(Start, 6.0, 1.0) * Scale * SunMergeBloomSetupTextureSizeAndInvSize.zw;
+	const float2 Offset2 = Circle(Start, 6.0, 2.0) * Scale * SunMergeBloomSetupTextureSizeAndInvSize.zw;
+	const float2 Offset3 = Circle(Start, 6.0, 3.0) * Scale * SunMergeBloomSetupTextureSizeAndInvSize.zw;
+	const float2 Offset4 = Circle(Start, 6.0, 4.0) * Scale * SunMergeBloomSetupTextureSizeAndInvSize.zw;
+	const float2 Offset5 = Circle(Start, 6.0, 5.0) * Scale * SunMergeBloomSetupTextureSizeAndInvSize.zw;
 
-		float Scale1 = 1.0/7.0;
+	#if MOBILE_USEBLOOM
 
-		float Scale2 = 1.0/7.0;
+		const float Scale1 = 1.0 / 7.0;
+		const float Scale2 = 1.0 / 7.0;
 
-		half3 Bloom2 = (
-			SampleBloomSetup(InUVVignette.xy, EyeIndex).rgba *Scale1 +
-			SampleBloomSetup(InUVs[0].xy, EyeIndex).rgba * Scale2 +
-			SampleBloomSetup(InUVs[1].xy, EyeIndex).rgba * Scale2 +
-			SampleBloomSetup(InUVs[2].xy, EyeIndex).rgba * Scale2 +
-			SampleBloomSetup(InUVs[3].xy, EyeIndex).rgba * Scale2 +
-			SampleBloomSetup(InUVs[4].xy, EyeIndex).rgba * Scale2 +
-			SampleBloomSetup(InUVs[5].xy, EyeIndex).rgba * Scale2) * rcp(Scale1 * 1.0 + Scale2 * 6.0);
+		const half3 Bloom2 = (
+			SampleBloomSetup(clamp(BloomSetupUV,           BloomSetupUVMin, BloomSetupUVMax), EyeIndex).rgb * Scale1 +
+			SampleBloomSetup(clamp(BloomSetupUV + Offset0, BloomSetupUVMin, BloomSetupUVMax), EyeIndex).rgb * Scale2 +
+			SampleBloomSetup(clamp(BloomSetupUV + Offset1, BloomSetupUVMin, BloomSetupUVMax), EyeIndex).rgb * Scale2 +
+			SampleBloomSetup(clamp(BloomSetupUV + Offset2, BloomSetupUVMin, BloomSetupUVMax), EyeIndex).rgb * Scale2 +
+			SampleBloomSetup(clamp(BloomSetupUV + Offset3, BloomSetupUVMin, BloomSetupUVMax), EyeIndex).rgb * Scale2 +
+			SampleBloomSetup(clamp(BloomSetupUV + Offset4, BloomSetupUVMin, BloomSetupUVMax), EyeIndex).rgb * Scale2 +
+			SampleBloomSetup(clamp(BloomSetupUV + Offset5, BloomSetupUVMin, BloomSetupUVMax), EyeIndex).rgb * Scale2) * rcp(Scale1 * 1.0 + Scale2 * 6.0);
 
-		OutColor.rgb = SampleBloomUp(InUVVignette.xy, EyeIndex);
+		OutColor.rgb = SampleBloomUp(clamp(BloomUpUV, BloomUpUVMin, BloomUpUVMax), EyeIndex);
 
-		half3 BloomDirtMaskColor = BloomDirtMaskTexture.Sample(BloomDirtMaskSampler, InUVVignette.xy).rgb * BloomDirtMaskTint.rgb;
+		const half3 BloomDirtMaskColor = BloomDirtMaskTexture.Sample(BloomDirtMaskSampler, clamp(BloomSetupUV, BloomSetupUVMin, BloomSetupUVMax)).rgb * BloomDirtMaskTint.rgb;
 
 		// Have 5 layers on mobile.
 		half Scale3 = 1.0/5.0;	
@@ -1245,14 +1250,14 @@ void SunMergePS_Mobile(
 
 	#if MOBILE_USESUN
 		half Sun = 
-			SunBlurTexture.Sample(SunBlurSampler, InUVs[0].zw).r * 0.125 +
-			SunBlurTexture.Sample(SunBlurSampler, InUVs[1].zw).r * 0.125 +
-			SunBlurTexture.Sample(SunBlurSampler, InUVs[2].zw).r * 0.125 +
-			SunBlurTexture.Sample(SunBlurSampler, InUVs[3].zw).r * 0.125 +
-			SunBlurTexture.Sample(SunBlurSampler, InUVs[4].zw).r * 0.125 +
-			SunBlurTexture.Sample(SunBlurSampler, InUVs[5].zw).r * 0.125 +
-			SunBlurTexture.Sample(SunBlurSampler, InUVs[6].xy).r * 0.125 +
-			SunBlurTexture.Sample(SunBlurSampler, InUVVignette.xy).r * 0.125;
+			SunBlurTexture.Sample(SunBlurSampler, InSunShaftUVs[0].xy).r * 0.125 +
+			SunBlurTexture.Sample(SunBlurSampler, InSunShaftUVs[0].zw).r * 0.125 +
+			SunBlurTexture.Sample(SunBlurSampler, InSunShaftUVs[1].xy).r * 0.125 +
+			SunBlurTexture.Sample(SunBlurSampler, InSunShaftUVs[1].zw).r * 0.125 +
+			SunBlurTexture.Sample(SunBlurSampler, InSunShaftUVs[2].xy).r * 0.125 +
+			SunBlurTexture.Sample(SunBlurSampler, InSunShaftUVs[2].zw).r * 0.125 +
+			SunBlurTexture.Sample(SunBlurSampler, InSunShaftUVs[3].xy).r * 0.125 +
+			SunBlurTexture.Sample(SunBlurSampler, clamp(BloomSetupUV, BloomSetupUVMin, BloomSetupUVMax)).r * 0.125;
 		Sun = HighlightDecompression(Sun);
 		OutColor.rgb += SunColorVignetteIntensity.rgb * Sun;
 	#endif
diff --git a/Engine/Source/Runtime/Renderer/Private/PostProcess/PostProcessMobile.cpp b/Engine/Source/Runtime/Renderer/Private/PostProcess/PostProcessMobile.cpp
--- a/Engine/Source/Runtime/Renderer/Private/PostProcess/PostProcessMobile.cpp
+++ b/Engine/Source/Runtime/Renderer/Private/PostProcess/PostProcessMobile.cpp
@@ -277,7 +277,8 @@ IMPLEMENT_GLOBAL_SHADER(FMobileBloomSetupPS, "/Engine/Private/PostProcessMobile.
 
 FMobileBloomSetupOutputs AddMobileBloomSetupPass(FRDGBuilder& GraphBuilder, const FViewInfo& View, const FEyeAdaptationParameters& EyeAdaptationParameters, const FMobileBloomSetupInputs& Inputs)
 {
-	FIntPoint OutputSize = FIntPoint::DivideAndRoundUp(Inputs.SceneColor.ViewRect.Size(), (GetBloomQuality() <= EBloomQuality::Q2) ? 4 : 2);
+	const FIntPoint OutputSize = FIntPoint::DivideAndRoundUp(Inputs.SceneColor.ViewRect.Size(), (GetBloomQuality() <= EBloomQuality::Q2) ? 4 : 2);
+	const FIntPoint OutputFixedTextureSize = Inputs.OutputFixedTextureSize;
 
 	const FIntPoint& BufferSize = Inputs.SceneColor.Texture->Desc.Extent;
 
@@ -287,7 +288,7 @@ FMobileBloomSetupOutputs AddMobileBloomSetupPass(FRDGBuilder& GraphBuilder, cons
 
 	bool bUseTextureArrays = View.bIsMobileMultiViewEnabled;
 
-	const auto GetBloomSetupTarget = [&GraphBuilder, bIsValidVariation, &Inputs, &OutputSize, &BloomSetupOutputType, bUseTextureArrays](int32 OutputIndex)
+	const auto GetBloomSetupTarget = [&GraphBuilder, bIsValidVariation, &Inputs, &OutputFixedTextureSize, &BloomSetupOutputType, bUseTextureArrays](int32 OutputIndex)
 	{
 		checkSlow(OutputIndex < BloomSetupOutputType.Num());
 
@@ -308,8 +309,8 @@ FMobileBloomSetupOutputs AddMobileBloomSetupPass(FRDGBuilder& GraphBuilder, cons
 		}
 
 		FRDGTextureDesc BloomSetupDesc = bUseTextureArrays ?
-			FRDGTextureDesc::Create2DArray(OutputSize, Format, FClearValueBinding::Black, TargetableFlags, 2) :
-			FRDGTextureDesc::Create2D(OutputSize, Format, FClearValueBinding::Black, TargetableFlags);
+			FRDGTextureDesc::Create2DArray(OutputFixedTextureSize, Format, FClearValueBinding::Black, TargetableFlags, 2) :
+			FRDGTextureDesc::Create2D(OutputFixedTextureSize, Format, FClearValueBinding::Black, TargetableFlags);
 
 		return FScreenPassRenderTarget(GraphBuilder.CreateTexture(BloomSetupDesc, GetBloomSetupOutputTypeName(BloomSetupOutputType[OutputIndex])), ERenderTargetLoadAction::ENoAction);
 	};
@@ -358,10 +359,10 @@ FMobileBloomSetupOutputs AddMobileBloomSetupPass(FRDGBuilder& GraphBuilder, cons
 		(Inputs.SceneColor.ViewRect.Max.Y - 0.5f) / BufferSize.Y);
 
 	const FScreenPassTextureViewport InputViewport(Inputs.SceneColor);
-	const FScreenPassTextureViewport OutputViewport(DestRenderTargets[0]);
+	const FScreenPassTextureViewport OutputViewport(OutputSize);
 
 	GraphBuilder.AddPass(
-		RDG_EVENT_NAME("BloomSetup %dx%d (PS)", OutputViewport.Extent.X, OutputViewport.Extent.Y),
+		RDG_EVENT_NAME("BloomSetup %dx%d in %dx%d (PS) ", OutputViewport.Extent.X, OutputViewport.Extent.Y, OutputFixedTextureSize.X, OutputFixedTextureSize.Y),
 		PSShaderParameters,
 		ERDGPassFlags::Raster,
 		[VertexShader, VSShaderParameters, PixelShader, PSShaderParameters, InputViewport, OutputViewport, BloomSetupOutputType, &View](FRHICommandList& RHICmdList)
@@ -404,15 +405,15 @@ FMobileBloomSetupOutputs AddMobileBloomSetupPass(FRDGBuilder& GraphBuilder, cons
 	{
 		if (BloomSetupOutputType[i] == EBloomSetupOutputType::Bloom)
 		{
-			Outputs.Bloom = DestRenderTargets[i];
+			Outputs.Bloom = FScreenPassTexture(DestRenderTargets[i].Texture, OutputViewport.Rect);
 		}
 		else if(BloomSetupOutputType[i] == EBloomSetupOutputType::SunShaftAndDof)
 		{
-			Outputs.SunShaftAndDof = DestRenderTargets[i];
+			Outputs.SunShaftAndDof = FScreenPassTexture(DestRenderTargets[i].Texture, OutputViewport.Rect);
 		}
 		else if (BloomSetupOutputType[i] == EBloomSetupOutputType::EyeAdaptation)
 		{
-			Outputs.EyeAdaptation = DestRenderTargets[i];
+			Outputs.EyeAdaptation = FScreenPassTexture(DestRenderTargets[i].Texture, OutputViewport.Rect);
 		}
 	}
 
@@ -434,6 +435,9 @@ public:
 		SHADER_PARAMETER_STRUCT_INCLUDE(FViewShaderParameters, View)
 		SHADER_PARAMETER_RDG_TEXTURE(Texture2D, BloomDownSourceTexture)
 		SHADER_PARAMETER_SAMPLER(SamplerState, BloomDownSourceSampler)
+		SHADER_PARAMETER(FVector4f, InputViewportRect)
+		SHADER_PARAMETER(FVector4f, TextureSizeAndInvSize)
+		SHADER_PARAMETER(float, BloomDownScale)
 		RENDER_TARGET_BINDING_SLOTS()
 	END_SHADER_PARAMETER_STRUCT()
 
@@ -467,16 +471,14 @@ IMPLEMENT_GLOBAL_SHADER(FMobileBloomDownVS, "/Engine/Private/PostProcessMobile.u
 
 FScreenPassTexture AddMobileBloomDownPass(FRDGBuilder& GraphBuilder, const FViewInfo& View, const FMobileBloomDownInputs& Inputs)
 {
-	FIntPoint OutputSize = FIntPoint::DivideAndRoundUp(Inputs.BloomDownSource.ViewRect.Size(), 2);
-
-	const FIntPoint& BufferSize = Inputs.BloomDownSource.Texture->Desc.Extent;
+	const FIntPoint& BufferSize = FIntPoint(Inputs.BloomDownSource.ViewRect.Max.X, Inputs.BloomDownSource.ViewRect.Max.Y);
 
 	bool bUseTextureArrays = View.bIsMobileMultiViewEnabled;
 	FRDGTextureDesc BloomDownDesc = bUseTextureArrays ?
-		FRDGTextureDesc::Create2DArray(OutputSize, PF_FloatR11G11B10, FClearValueBinding::Black, TexCreate_RenderTargetable | TexCreate_ShaderResource, 2) :
-		FRDGTextureDesc::Create2D(OutputSize, PF_FloatR11G11B10, FClearValueBinding::Black, TexCreate_RenderTargetable | TexCreate_ShaderResource);
+		FRDGTextureDesc::Create2DArray(Inputs.BloomDownOutputFixedTextureExtent, PF_FloatR11G11B10, FClearValueBinding::Black, TexCreate_RenderTargetable | TexCreate_ShaderResource, 2) :
+		FRDGTextureDesc::Create2D(Inputs.BloomDownOutputFixedTextureExtent, PF_FloatR11G11B10, FClearValueBinding::Black, TexCreate_RenderTargetable | TexCreate_ShaderResource);
 
-	FScreenPassRenderTarget BloomDownOutput = FScreenPassRenderTarget(GraphBuilder.CreateTexture(BloomDownDesc, TEXT("BloomDown")), ERenderTargetLoadAction::EClear);
+	FScreenPassRenderTarget BloomDownOutput = FScreenPassRenderTarget(GraphBuilder.CreateTexture(BloomDownDesc, TEXT("BloomDown")), FIntRect(0,0,Inputs.BloomDownOutputViewExtent.X,Inputs.BloomDownOutputViewExtent.Y) , ERenderTargetLoadAction::EClear);
 
 	TShaderMapRef<FMobileBloomDownVS> VertexShader(View.ShaderMap);
 
@@ -496,14 +498,29 @@ FScreenPassTexture AddMobileBloomDownPass(FRDGBuilder& GraphBuilder, const FView
 	PSShaderParameters->BloomDownSourceTexture = Inputs.BloomDownSource.Texture;
 	PSShaderParameters->BloomDownSourceSampler = TStaticSamplerState<SF_Bilinear, AM_Clamp, AM_Clamp, AM_Clamp>::GetRHI();
 
-	const FScreenPassTextureViewport InputViewport(Inputs.BloomDownSource);
-	const FScreenPassTextureViewport OutputViewport(BloomDownOutput);
+	PSShaderParameters->InputViewportRect = FVector4f(
+		Inputs.BloomDownSource.ViewRect.Min.X,
+		Inputs.BloomDownSource.ViewRect.Min.Y,
+		Inputs.BloomDownSource.ViewRect.Width(),
+		Inputs.BloomDownSource.ViewRect.Height()
+	);
+	PSShaderParameters->TextureSizeAndInvSize = FVector4f(
+		Inputs.BloomDownSource.Texture->Desc.Extent.X,
+		Inputs.BloomDownSource.Texture->Desc.Extent.Y,
+		1.0f / Inputs.BloomDownSource.Texture->Desc.Extent.X,
+		1.0f / Inputs.BloomDownSource.Texture->Desc.Extent.Y
+	);
+	PSShaderParameters->BloomDownScale = Inputs.BloomDownScale;
+
+	const FScreenPassTextureViewport InputViewport(Inputs.BloomDownSource.ViewRect);
+	const FIntPoint InputTextureSize = Inputs.BloomDownSource.Texture->Desc.Extent;
+	const FScreenPassTextureViewport OutputViewport(Inputs.BloomDownOutputViewExtent);
 
 	GraphBuilder.AddPass(
-		RDG_EVENT_NAME("BloomDown %dx%d (PS)", OutputViewport.Extent.X, OutputViewport.Extent.Y),
+		RDG_EVENT_NAME("BloomDown %dx%d in %dx%d RT (PS)", OutputViewport.Extent.X, OutputViewport.Extent.Y, Inputs.BloomDownOutputFixedTextureExtent.X, Inputs.BloomDownOutputFixedTextureExtent.Y),
 		PSShaderParameters,
 		ERDGPassFlags::Raster,
-		[VertexShader, VSShaderParameters, PixelShader, PSShaderParameters, InputViewport, OutputViewport, &View](FRHICommandList& RHICmdList)
+		[VertexShader, VSShaderParameters, PixelShader, PSShaderParameters, InputViewport, InputTextureSize, OutputViewport, &View](FRHICommandList& RHICmdList)
 	{
 		RHICmdList.SetViewport(OutputViewport.Rect.Min.X, OutputViewport.Rect.Min.Y, 0.0f, OutputViewport.Rect.Max.X, OutputViewport.Rect.Max.Y, 1.0f);
 
@@ -531,7 +548,7 @@ FScreenPassTexture AddMobileBloomDownPass(FRDGBuilder& GraphBuilder, const FView
 			0, 0,
 			InputViewport.Rect.Width(), InputViewport.Rect.Height(),
 			OutputViewport.Extent,
-			InputViewport.Extent,
+			InputTextureSize,
 			VertexShader,
 			EDRF_UseTriangleOptimization,
 			View.GetStereoPassInstanceFactor());
@@ -558,6 +575,12 @@ public:
 		SHADER_PARAMETER_SAMPLER(SamplerState, BloomUpSourceASampler)
 		SHADER_PARAMETER_RDG_TEXTURE(Texture2D, BloomUpSourceBTexture)
 		SHADER_PARAMETER_SAMPLER(SamplerState, BloomUpSourceBSampler)
+		SHADER_PARAMETER(FVector2f, BloomUpScales)
+		SHADER_PARAMETER(FVector4f, ViewportARect)
+		SHADER_PARAMETER(FVector4f, ViewportBRect)
+		SHADER_PARAMETER(FVector4f, TextureASizeAndInvSize)
+		SHADER_PARAMETER(FVector4f, TextureBSizeAndInvSize)
+
 		RENDER_TARGET_BINDING_SLOTS()
 	END_SHADER_PARAMETER_STRUCT()
 
@@ -577,8 +600,8 @@ public:
 
 	BEGIN_SHADER_PARAMETER_STRUCT(FParameters, )
 		SHADER_PARAMETER_STRUCT_INCLUDE(FViewShaderParameters, View)
-		SHADER_PARAMETER(FVector4f, BufferASizeAndInvSize)
-		SHADER_PARAMETER(FVector4f, BufferBSizeAndInvSize)
+		SHADER_PARAMETER(FVector4f, BufferASizeAndInvSize)  // Texture A
+		SHADER_PARAMETER(FVector4f, BufferBSizeAndInvSize)  // Texture B
 		SHADER_PARAMETER(FVector2f, BloomUpScales)
 	END_SHADER_PARAMETER_STRUCT()
 
@@ -592,27 +615,42 @@ IMPLEMENT_GLOBAL_SHADER(FMobileBloomUpVS, "/Engine/Private/PostProcessMobile.usf
 
 FScreenPassTexture AddMobileBloomUpPass(FRDGBuilder& GraphBuilder, const FViewInfo& View, const FMobileBloomUpInputs& Inputs)
 {
-	FIntPoint OutputSize = Inputs.BloomUpSourceA.ViewRect.Size();
-
-	const FIntPoint& BufferSizeA = Inputs.BloomUpSourceA.Texture->Desc.Extent;
+	// Use actual texture sizes for UV calculations
+	const FIntPoint& TextureSizeA = Inputs.BloomUpSourceA.Texture->Desc.Extent;
+	const FIntPoint& TextureSizeB = Inputs.BloomUpSourceB.Texture->Desc.Extent;
 
-	const FIntPoint& BufferSizeB = Inputs.BloomUpSourceB.Texture->Desc.Extent;
+	// Get viewport sizes
+	const FIntRect& ViewportRectA = Inputs.BloomUpSourceA.ViewRect;
+	const FIntRect& ViewportRectB = Inputs.BloomUpSourceB.ViewRect;
 
 	const bool bUseTextureArrays = View.bIsMobileMultiViewEnabled;
-	FRDGTextureDesc BloomUpDesc = bUseTextureArrays ? 
-		FRDGTextureDesc::Create2DArray(OutputSize, PF_FloatR11G11B10, FClearValueBinding::Black, TexCreate_RenderTargetable | TexCreate_ShaderResource, 2) :
-		FRDGTextureDesc::Create2D(OutputSize, PF_FloatR11G11B10, FClearValueBinding::Black, TexCreate_RenderTargetable | TexCreate_ShaderResource);
+	FRDGTextureDesc BloomUpDesc = bUseTextureArrays ?
+		FRDGTextureDesc::Create2DArray(Inputs.BloomUpOutputFixedTextureExtent, PF_FloatR11G11B10, FClearValueBinding::Black, TexCreate_RenderTargetable | TexCreate_ShaderResource, 2) :
+		FRDGTextureDesc::Create2D(Inputs.BloomUpOutputFixedTextureExtent, PF_FloatR11G11B10, FClearValueBinding::Black, TexCreate_RenderTargetable | TexCreate_ShaderResource);
 
-	FScreenPassRenderTarget BloomUpOutput = FScreenPassRenderTarget(GraphBuilder.CreateTexture(BloomUpDesc, TEXT("BloomUp")), ERenderTargetLoadAction::EClear);
+	FScreenPassRenderTarget BloomUpOutput = FScreenPassRenderTarget(
+		GraphBuilder.CreateTexture(BloomUpDesc, TEXT("BloomUp")),
+		FIntRect(0, 0, Inputs.BloomUpOutputViewExtent.X, Inputs.BloomUpOutputViewExtent.Y),
+		ERenderTargetLoadAction::EClear
+	);
 
 	TShaderMapRef<FMobileBloomUpVS> VertexShader(View.ShaderMap);
 
 	FMobileBloomUpVS::FParameters VSShaderParameters;
 
 	VSShaderParameters.View = View.GetShaderParameters();
-	VSShaderParameters.BufferASizeAndInvSize = FVector4f(BufferSizeA.X, BufferSizeA.Y, 1.0f / BufferSizeA.X, 1.0f / BufferSizeA.Y);
-	VSShaderParameters.BufferBSizeAndInvSize = FVector4f(BufferSizeB.X, BufferSizeB.Y, 1.0f / BufferSizeB.X, 1.0f / BufferSizeB.Y);
-	VSShaderParameters.BloomUpScales = FVector2f(Inputs.ScaleAB);	// LWC_TODO: Precision loss
+
+	// Pass TEXTURE sizes for correct inverse calculations
+	VSShaderParameters.BufferASizeAndInvSize = FVector4f(
+		TextureSizeA.X, TextureSizeA.Y,
+		1.0f / TextureSizeA.X, 1.0f / TextureSizeA.Y
+	);
+	VSShaderParameters.BufferBSizeAndInvSize = FVector4f(
+		TextureSizeB.X, TextureSizeB.Y,
+		1.0f / TextureSizeB.X, 1.0f / TextureSizeB.Y
+	);
+
+	VSShaderParameters.BloomUpScales = FVector2f(Inputs.ScaleAB);
 
 	TShaderMapRef<FMobileBloomUpPS> PixelShader(View.ShaderMap);
 
@@ -629,16 +667,43 @@ FScreenPassTexture AddMobileBloomUpPass(FRDGBuilder& GraphBuilder, const FViewIn
 	PSShaderParameters->BloomTintA = Inputs.TintA * (1.0f / 8.0f);
 	PSShaderParameters->BloomTintB = Inputs.TintB * (1.0f / 8.0f);
 
-	const FScreenPassTextureViewport InputViewport(Inputs.BloomUpSourceA);
+	PSShaderParameters->ViewportARect = FVector4f(
+		ViewportRectA.Min.X,
+		ViewportRectA.Min.Y,
+		ViewportRectA.Width(),
+		ViewportRectA.Height()
+	);
+	PSShaderParameters->ViewportBRect = FVector4f(
+		ViewportRectB.Min.X,
+		ViewportRectB.Min.Y,
+		ViewportRectB.Width(),
+		ViewportRectB.Height()
+	);
+	PSShaderParameters->TextureASizeAndInvSize = FVector4f(
+		TextureSizeA.X,
+		TextureSizeA.Y,
+		1.0f / TextureSizeA.X,
+		1.0f / TextureSizeA.Y
+	);
+	PSShaderParameters->TextureBSizeAndInvSize = FVector4f(
+		TextureSizeB.X,
+		TextureSizeB.Y,
+		1.0f / TextureSizeB.X,
+		1.0f / TextureSizeB.Y
+	);
+
+	PSShaderParameters->BloomUpScales = FVector2f(Inputs.ScaleAB);
+
 	const FScreenPassTextureViewport OutputViewport(BloomUpOutput);
+	const FIntPoint OutputViewportSize = Inputs.BloomUpOutputViewExtent;
 
 	GraphBuilder.AddPass(
-		RDG_EVENT_NAME("BloomUp %dx%d (PS)", OutputViewport.Extent.X, OutputViewport.Extent.Y),
+		RDG_EVENT_NAME("BloomUp %dx%d in %dx%d RT (PS)", OutputViewportSize.X, OutputViewportSize.Y, Inputs.BloomUpOutputFixedTextureExtent.X, Inputs.BloomUpOutputFixedTextureExtent.Y),
 		PSShaderParameters,
 		ERDGPassFlags::Raster,
-		[VertexShader, VSShaderParameters, PixelShader, PSShaderParameters, InputViewport, OutputViewport, &View](FRHICommandList& RHICmdList)
-	{
-		RHICmdList.SetViewport(OutputViewport.Rect.Min.X, OutputViewport.Rect.Min.Y, 0.0f, OutputViewport.Rect.Max.X, OutputViewport.Rect.Max.Y, 1.0f);
+		[VertexShader, VSShaderParameters, PixelShader, PSShaderParameters, OutputViewport, OutputViewportSize, ViewportRectA, &View](FRHICommandList& RHICmdList)
+		{
+			RHICmdList.SetViewport(0, 0, 0.0f, OutputViewportSize.X, OutputViewportSize.Y, 1.0f);
 
 		FGraphicsPipelineStateInitializer GraphicsPSOInit;
 		RHICmdList.ApplyCachedRenderTargets(GraphicsPSOInit);
@@ -657,18 +722,18 @@ FScreenPassTexture AddMobileBloomUpPass(FRDGBuilder& GraphBuilder, const FViewIn
 		SetShaderParameters(RHICmdList, VertexShader, VertexShader.GetVertexShader(), VSShaderParameters);
 		SetShaderParameters(RHICmdList, PixelShader, PixelShader.GetPixelShader(), *PSShaderParameters);
 
-		DrawRectangle(
-			RHICmdList,
-			0, 0,
-			OutputViewport.Extent.X, OutputViewport.Extent.Y,
-			0, 0,
-			InputViewport.Rect.Width(), InputViewport.Rect.Height(),
-			OutputViewport.Extent,
-			InputViewport.Extent,
-			VertexShader,
-			EDRF_UseTriangleOptimization,
-			View.GetStereoPassInstanceFactor());
-	});
+			DrawRectangle(
+				RHICmdList,
+				0, 0,
+				OutputViewportSize.X, OutputViewportSize.Y,
+				ViewportRectA.Min.X, ViewportRectA.Min.Y,
+				ViewportRectA.Width(), ViewportRectA.Height(),
+				OutputViewportSize,
+				FIntPoint(ViewportRectA.Width(), ViewportRectA.Height()),
+				VertexShader,
+				EDRF_UseTriangleOptimization,
+				View.GetStereoPassInstanceFactor());
+		});
 
 	return MoveTemp(BloomUpOutput);
 }
@@ -898,13 +963,14 @@ FScreenPassTexture AddMobileSunAlphaPass(FRDGBuilder& GraphBuilder, const FViewI
 	PSShaderParameters->SunShaftAndDofTexture = Inputs.BloomSetup_SunShaftAndDof.Texture;
 	PSShaderParameters->SunShaftAndDofSampler = TStaticSamplerState<SF_Bilinear, AM_Clamp, AM_Clamp, AM_Clamp>::GetRHI();
 
+	const FScreenPassTextureViewport InputViewport(Inputs.BloomSetup_SunShaftAndDof);
 	const FScreenPassTextureViewport OutputViewport(SunAlphaOutput);
 
 	GraphBuilder.AddPass(
 		RDG_EVENT_NAME("SunAlpha %dx%d (PS)", OutputViewport.Extent.X, OutputViewport.Extent.Y),
 		PSShaderParameters,
 		ERDGPassFlags::Raster,
-		[VertexShader, VSShaderParameters, PixelShader, PSShaderParameters, OutputViewport, &View](FRHICommandList& RHICmdList)
+		[VertexShader, VSShaderParameters, PixelShader, PSShaderParameters, InputViewport, OutputViewport, &View](FRHICommandList& RHICmdList)
 	{
 		RHICmdList.SetViewport(OutputViewport.Rect.Min.X, OutputViewport.Rect.Min.Y, 0.0f, OutputViewport.Rect.Max.X, OutputViewport.Rect.Max.Y, 1.0f);
 
@@ -929,10 +995,10 @@ FScreenPassTexture AddMobileSunAlphaPass(FRDGBuilder& GraphBuilder, const FViewI
 			RHICmdList,
 			0, 0,
 			OutputViewport.Extent.X, OutputViewport.Extent.Y,
-			0, 0,
-			OutputViewport.Extent.X, OutputViewport.Extent.Y,
-			OutputViewport.Extent,
+			InputViewport.Rect.Min.X, InputViewport.Rect.Min.Y,      
+			InputViewport.Rect.Width(), InputViewport.Rect.Height(), 
 			OutputViewport.Extent,
+			InputViewport.Extent,                                    
 			VertexShader,
 			EDRF_UseTriangleOptimization,
 			View.GetStereoPassInstanceFactor());
@@ -1081,6 +1147,12 @@ public:
 		SHADER_PARAMETER_SAMPLER(SamplerState, BloomUpSampler)
 		SHADER_PARAMETER_TEXTURE(Texture2D, BloomDirtMaskTexture)
 		SHADER_PARAMETER_SAMPLER(SamplerState, BloomDirtMaskSampler)
+
+		SHADER_PARAMETER(FVector4f, SunMergeBloomSetupViewportRect)
+		SHADER_PARAMETER(FVector4f, SunMergeBloomSetupTextureSizeAndInvSize)
+		SHADER_PARAMETER(FVector4f, SunMergeBloomUpViewportRect)
+		SHADER_PARAMETER(FVector4f, SunMergeBloomUpTextureSizeAndInvSize)
+
 		RENDER_TARGET_BINDING_SLOTS()
 	END_SHADER_PARAMETER_STRUCT()
 
@@ -1109,7 +1181,7 @@ public:
 	BEGIN_SHADER_PARAMETER_STRUCT(FParameters, )
 		SHADER_PARAMETER_STRUCT_INCLUDE(FViewShaderParameters, View)
 		SHADER_PARAMETER(FVector2f, LightShaftCenter)
-		SHADER_PARAMETER(FVector4f, BloomUpSizeAndInvSize)
+ 		SHADER_PARAMETER(FVector4f, BloomUpTextureSizeAndInvSize)     // xy=TexSize, zw=1/TexSize
 		SHADER_PARAMETER(FVector4f, ViewportSize)
 	END_SHADER_PARAMETER_STRUCT()
 
@@ -1144,18 +1216,30 @@ FScreenPassTexture AddMobileSunMergePass(FRDGBuilder& GraphBuilder, const FViewI
 	}
 	VSShaderParameters.LightShaftCenter = FVector2f(MobileLightShaft.Center);	// LWC_TODO: Precision loss
 
-	if (Inputs.BloomUp.IsValid())
-	{
-		const FIntPoint& BloomUpSize = Inputs.BloomUp.Texture->Desc.Extent;
-		VSShaderParameters.BloomUpSizeAndInvSize = FVector4f(BloomUpSize.X, BloomUpSize.Y, 1.0f / BloomUpSize.X, 1.0f / BloomUpSize.Y);
-	}
-	else
-	{
-		VSShaderParameters.BloomUpSizeAndInvSize = FVector4f(1.0f, 1.0f, 1.0f, 1.0f);
-	}
+	// Setup BloomSetup_Bloom parameters
+	const FIntRect& BloomSetupViewportRect = Inputs.BloomSetup_Bloom.ViewRect;
+	const FIntPoint BloomSetupTextureSize = Inputs.BloomSetup_Bloom.Texture->Desc.Extent;
+
+	// Setup BloomUp parameters
+	const FIntRect& BloomUpViewportRect = Inputs.BloomUp.ViewRect;
+	const FIntPoint BloomUpTextureSize = Inputs.BloomUp.Texture->Desc.Extent;
+
+ 	if (Inputs.BloomUp.IsValid())
+ 	{
+ 		VSShaderParameters.BloomUpTextureSizeAndInvSize = FVector4f(
+ 			BloomUpTextureSize.X,
+ 			BloomUpTextureSize.Y,
+ 			1.0f / BloomUpTextureSize.X,
+ 			1.0f / BloomUpTextureSize.Y
+ 		);
+ 	}
+ 	else
+ 	{
+ 		VSShaderParameters.BloomUpTextureSizeAndInvSize = FVector4f(1.0f, 1.0f, 1.0f, 1.0f);
+ 	}
 
 	VSShaderParameters.ViewportSize = FVector4f(OutputSize.X, OutputSize.Y, 1.0f / OutputSize.X, 1.0f / OutputSize.Y);
-
+	
 	auto ShaderPermutationVector = FMobileSunMergePS::BuildPermutationVector(Inputs.bUseBloom, Inputs.bUseSun);
 
 	TShaderMapRef<FMobileSunMergePS> PixelShader(View.ShaderMap, ShaderPermutationVector);
@@ -1193,6 +1277,31 @@ FScreenPassTexture AddMobileSunMergePass(FRDGBuilder& GraphBuilder, const FViewI
 	PSShaderParameters->BloomDirtMaskTexture = BloomDirtMaskTexture;
 	PSShaderParameters->BloomDirtMaskSampler = TStaticSamplerState<SF_Bilinear, AM_Clamp, AM_Clamp, AM_Clamp>::GetRHI();
 
+	PSShaderParameters->SunMergeBloomSetupViewportRect = FVector4f(
+		BloomSetupViewportRect.Min.X,
+		BloomSetupViewportRect.Min.Y,
+		BloomSetupViewportRect.Width(),
+		BloomSetupViewportRect.Height()
+	);
+	PSShaderParameters->SunMergeBloomSetupTextureSizeAndInvSize = FVector4f(
+		BloomSetupTextureSize.X,
+		BloomSetupTextureSize.Y,
+		1.0f / BloomSetupTextureSize.X,
+		1.0f / BloomSetupTextureSize.Y
+	);
+	PSShaderParameters->SunMergeBloomUpViewportRect = FVector4f(
+		BloomUpViewportRect.Min.X,
+		BloomUpViewportRect.Min.Y,
+		BloomUpViewportRect.Width(),
+		BloomUpViewportRect.Height()
+	);
+	PSShaderParameters->SunMergeBloomUpTextureSizeAndInvSize = FVector4f(
+		BloomUpTextureSize.X,
+		BloomUpTextureSize.Y,
+		1.0f / BloomUpTextureSize.X,
+		1.0f / BloomUpTextureSize.Y
+	);
+
 	const FScreenPassTextureViewport OutputViewport(SunMergeOutput);
 
 	GraphBuilder.AddPass(
@@ -1464,13 +1573,14 @@ FMobileDofNearOutputs AddMobileDofNearPass(FRDGBuilder& GraphBuilder, const FVie
 	PSShaderParameters->SunShaftAndDofTexture = Inputs.BloomSetup_SunShaftAndDof.Texture;
 	PSShaderParameters->SunShaftAndDofSampler = TStaticSamplerState<SF_Point, AM_Clamp, AM_Clamp, AM_Clamp>::GetRHI();
 
+	const FScreenPassTextureViewport InputViewport(Inputs.BloomSetup_SunShaftAndDof);
 	const FScreenPassTextureViewport OutputViewport(DofNearOutput);
 
 	GraphBuilder.AddPass(
 		RDG_EVENT_NAME("DofNear %dx%d (PS)", OutputViewport.Extent.X, OutputViewport.Extent.Y),
 		PSShaderParameters,
 		ERDGPassFlags::Raster,
-		[VertexShader, VSShaderParameters, PixelShader, PSShaderParameters, OutputViewport, &View](FRHICommandList& RHICmdList)
+		[VertexShader, VSShaderParameters, PixelShader, PSShaderParameters, InputViewport, OutputViewport, &View](FRHICommandList& RHICmdList)
 	{
 		RHICmdList.SetViewport(OutputViewport.Rect.Min.X, OutputViewport.Rect.Min.Y, 0.0f, OutputViewport.Rect.Max.X, OutputViewport.Rect.Max.Y, 1.0f);
 
@@ -1495,10 +1605,10 @@ FMobileDofNearOutputs AddMobileDofNearPass(FRDGBuilder& GraphBuilder, const FVie
 			RHICmdList,
 			0, 0,
 			OutputViewport.Extent.X, OutputViewport.Extent.Y,
-			0, 0,
-			OutputViewport.Extent.X, OutputViewport.Extent.Y,
-			OutputViewport.Extent,
+			InputViewport.Rect.Min.X, InputViewport.Rect.Min.Y,      
+			InputViewport.Rect.Width(), InputViewport.Rect.Height(), 
 			OutputViewport.Extent,
+			InputViewport.Extent,                                    
 			VertexShader,
 			EDRF_UseTriangleOptimization,
 			View.GetStereoPassInstanceFactor());
@@ -1893,13 +2003,13 @@ FMobileEyeAdaptationSetupOutputs AddMobileEyeAdaptationSetupPass(FRDGBuilder& Gr
 	FRDGBufferUAVRef EyeAdaptationSetupBufferUAV = GraphBuilder.CreateUAV(EyeAdaptationSetupBuffer, PF_R32_UINT);
 	AddClearUAVPass(GraphBuilder, EyeAdaptationSetupBufferUAV, 0);
 
-	const FIntPoint& BufferSize = Inputs.BloomSetup_EyeAdaptation.Texture->Desc.Extent;
+	const FIntPoint& ViewSize = Inputs.BloomSetup_EyeAdaptation.ViewRect.Size();
 
 	if (Inputs.bUseBasicEyeAdaptation)
 	{
 		FMobileAverageLuminanceCS::FParameters* PassParameters = GraphBuilder.AllocParameters<FMobileAverageLuminanceCS::FParameters>();
 
-		PassParameters->SourceSizeAndInvSize = FVector4f(BufferSize.X, BufferSize.Y, 1.0f / BufferSize.X, 1.0f / BufferSize.Y);
+		PassParameters->SourceSizeAndInvSize = FVector4f(ViewSize.X, ViewSize.Y, 1.0f / ViewSize.X, 1.0f / ViewSize.Y);
 		PassParameters->InputSampler = TStaticSamplerState<SF_Bilinear, AM_Clamp, AM_Clamp, AM_Clamp>::GetRHI();
 		PassParameters->InputTexture = Inputs.BloomSetup_EyeAdaptation.Texture;
 		PassParameters->EyeAdaptation = EyeAdaptationParameters;
@@ -1912,14 +2022,14 @@ FMobileEyeAdaptationSetupOutputs AddMobileEyeAdaptationSetupPass(FRDGBuilder& Gr
 			RDG_EVENT_NAME("EyeAdaptation_AverageLuminance (CS)"),
 			ComputeShader,
 			PassParameters,
-			FComputeShaderUtils::GetGroupCount(BufferSize, FMobileAverageLuminanceCS::TexelsPerThreadGroup));
+			FComputeShaderUtils::GetGroupCount(ViewSize, FMobileAverageLuminanceCS::TexelsPerThreadGroup));
 	}
 	else if (Inputs.bUseHistogramEyeAdaptation)
 	{
 		FMobileHistogram::FParameters* PassParameters = GraphBuilder.AllocParameters<FMobileHistogram::FParameters>();
 
 		PassParameters->View = View.GetShaderParameters();
-		PassParameters->SourceSizeAndInvSize = FVector4f(BufferSize.X, BufferSize.Y, 1.0f / BufferSize.X, 1.0f / BufferSize.Y);
+		PassParameters->SourceSizeAndInvSize = FVector4f(ViewSize.X, ViewSize.Y, 1.0f / ViewSize.X, 1.0f / ViewSize.Y);
 		PassParameters->InputTexture = Inputs.BloomSetup_EyeAdaptation.Texture;
 		PassParameters->InputSampler = TStaticSamplerState<SF_Bilinear, AM_Clamp, AM_Clamp, AM_Clamp>::GetRHI();
 		PassParameters->RWHistogramBuffer = EyeAdaptationSetupBufferUAV;
@@ -1936,7 +2046,7 @@ FMobileEyeAdaptationSetupOutputs AddMobileEyeAdaptationSetupPass(FRDGBuilder& Gr
 				RDG_EVENT_NAME("EyeAdaptation_Histogram (CS)"),
 				ComputeShader,
 				PassParameters,
-				FComputeShaderUtils::GetGroupCount(BufferSize, TMobileHistogramCS<true>::TexelsPerThreadGroup));
+				FComputeShaderUtils::GetGroupCount(ViewSize, TMobileHistogramCS<true>::TexelsPerThreadGroup));
 		}
 		else
 		{
@@ -1947,7 +2057,7 @@ FMobileEyeAdaptationSetupOutputs AddMobileEyeAdaptationSetupPass(FRDGBuilder& Gr
 				RDG_EVENT_NAME("EyeAdaptation_Histogram (CS)"),
 				ComputeShader,
 				PassParameters,
-				FComputeShaderUtils::GetGroupCount(BufferSize, TMobileHistogramCS<false>::TexelsPerThreadGroup));
+				FComputeShaderUtils::GetGroupCount(ViewSize, TMobileHistogramCS<false>::TexelsPerThreadGroup));
 		}
 	}
 
diff --git a/Engine/Source/Runtime/Renderer/Private/PostProcess/PostProcessMobile.h b/Engine/Source/Runtime/Renderer/Private/PostProcess/PostProcessMobile.h
--- a/Engine/Source/Runtime/Renderer/Private/PostProcess/PostProcessMobile.h
+++ b/Engine/Source/Runtime/Renderer/Private/PostProcess/PostProcessMobile.h
@@ -20,6 +20,7 @@ struct FMobileBloomSetupInputs
 {
 	FScreenPassTexture SceneColor;
 	FScreenPassTexture SunShaftAndDof;
+	FIntPoint OutputFixedTextureSize;
 
 	bool bUseBloom = false;
 	bool bUseSun = false;
@@ -91,7 +92,9 @@ FScreenPassTexture AddMobileIntegrateDofPass(FRDGBuilder& GraphBuilder, const FV
 
 struct FMobileBloomDownInputs
 {
-	FScreenPassTexture BloomDownSource;
+	FIntPoint			BloomDownOutputFixedTextureExtent;
+	FIntPoint			BloomDownOutputViewExtent;
+	FScreenPassTexture	BloomDownSource;
 
 	float BloomDownScale = 0.66f * 4.0f;
 };
@@ -103,6 +106,8 @@ struct FMobileBloomUpInputs
 	FScreenPassTexture BloomUpSourceA;
 	FScreenPassTexture BloomUpSourceB;
 
+	FIntPoint BloomUpOutputFixedTextureExtent;
+	FIntPoint BloomUpOutputViewExtent;
 	FVector2D ScaleAB;
 	FVector4f TintA;
 	FVector4f TintB;
diff --git a/Engine/Source/Runtime/Renderer/Private/PostProcess/PostProcessing.cpp b/Engine/Source/Runtime/Renderer/Private/PostProcess/PostProcessing.cpp
--- a/Engine/Source/Runtime/Renderer/Private/PostProcess/PostProcessing.cpp
+++ b/Engine/Source/Runtime/Renderer/Private/PostProcess/PostProcessing.cpp
@@ -2930,6 +2930,12 @@ void AddMobilePostProcessingPasses(FRDGBuilder& GraphBuilder, FScene* Scene, con
 		}
 
 		FMobileBloomSetupOutputs BloomSetupOutputs;
+
+		const int32 BloomQualityDivider = GetBloomQuality() <= EBloomQuality::Q2 ? 4 : 2;
+		const FIntPoint BloomOutputTextureSize = FIntPoint::DivideAndRoundUp(SceneColor.ViewRect.Size(), BloomQualityDivider);
+		const FIntPoint BloomOutputFixedTextureSize = FIntPoint::DivideAndRoundUp(View.GetSceneTexturesConfig().Extent, BloomQualityDivider);
+		
+		
 		if (PassSequence.IsEnabled(EPass::BloomSetup))
 		{
 			PassSequence.AcceptPass(EPass::BloomSetup);
@@ -2943,6 +2949,7 @@ void AddMobilePostProcessingPasses(FRDGBuilder& GraphBuilder, FScene* Scene, con
 			BloomSetupInputs.bUseSun = bUseSun;
 			BloomSetupInputs.SceneColor = SceneColor;
 			BloomSetupInputs.SunShaftAndDof = PostProcessSunShaftAndDof;
+			BloomSetupInputs.OutputFixedTextureSize = BloomOutputFixedTextureSize;
 
 			BloomSetupOutputs = AddMobileBloomSetupPass(GraphBuilder, View, EyeAdaptationParameters, BloomSetupInputs);
 		}
@@ -3059,11 +3066,13 @@ void AddMobilePostProcessingPasses(FRDGBuilder& GraphBuilder, FScene* Scene, con
 		if (PassSequence.IsEnabled(EPass::Bloom))
 		{
 			PassSequence.AcceptPass(EPass::Bloom);
-			auto AddBloomDownPass = [&GraphBuilder, &View](const FScreenPassTexture& BloomDownSource, float BloomDownScale)
+			auto AddBloomDownPass = [&GraphBuilder, &View](const FScreenPassTexture& BloomDownSource, float BloomDownScale, const FIntPoint& OutputTextureFixedExtent, const FIntPoint& OutputViewExtent)
 			{
 				FMobileBloomDownInputs BloomDownInputs;
 				BloomDownInputs.BloomDownScale = BloomDownScale;
 				BloomDownInputs.BloomDownSource = BloomDownSource;
+				BloomDownInputs.BloomDownOutputFixedTextureExtent = OutputTextureFixedExtent;
+				BloomDownInputs.BloomDownOutputViewExtent = OutputViewExtent;
 
 				return AddMobileBloomDownPass(GraphBuilder, View, BloomDownInputs);
 			};
@@ -3074,14 +3083,19 @@ void AddMobilePostProcessingPasses(FRDGBuilder& GraphBuilder, FScene* Scene, con
 			const int NumDownsamplePasses = BloomQuality == EBloomQuality::Q1 ? 4 : BloomQuality == EBloomQuality::Q2 ? 5 : 6;
 			FScreenPassTexture PostProcessDownsample_Bloom[MaxPasses];
 
+			int32 Divider = BloomQualityDivider;
 			for (int32 i = 0; i < NumDownsamplePasses; ++i)
 			{
-				PostProcessDownsample_Bloom[i] = AddBloomDownPass(i == 0 ? BloomSetupOutputs.Bloom : PostProcessDownsample_Bloom[i - 1], BloomDownScale);
+				FIntPoint OutputTextureFixedExtent = FIntPoint::DivideAndRoundUp(BloomOutputFixedTextureSize, Divider);
+				FIntPoint OutputViewExtent = FIntPoint::DivideAndRoundUp(BloomOutputTextureSize, Divider);
+				Divider *= BloomQualityDivider;
+
+				PostProcessDownsample_Bloom[i] = AddBloomDownPass(i == 0 ? BloomSetupOutputs.Bloom : PostProcessDownsample_Bloom[i - 1], BloomDownScale, OutputTextureFixedExtent, OutputViewExtent );
 			}
 
 			const FFinalPostProcessSettings& Settings = View.FinalPostProcessSettings;
 
-			auto AddBloomUpPass = [&GraphBuilder, &View](FScreenPassTexture& BloomUpSourceA, FScreenPassTexture& BloomUpSourceB, float BloomSourceScale, const FVector4f& TintA, const FVector4f& TintB)
+			auto AddBloomUpPass = [&GraphBuilder, &View](FScreenPassTexture& BloomUpSourceA, FScreenPassTexture& BloomUpSourceB, float BloomSourceScale, const FVector4f& TintA, const FVector4f& TintB, const FIntPoint& OutputTextureFixedExtent)
 			{
 				FMobileBloomUpInputs BloomUpInputs;
 				BloomUpInputs.BloomUpSourceA = BloomUpSourceA;
@@ -3090,9 +3104,13 @@ void AddMobilePostProcessingPasses(FRDGBuilder& GraphBuilder, FScene* Scene, con
 				BloomUpInputs.TintA = TintA;
 				BloomUpInputs.TintB = TintB;
 
+				BloomUpInputs.BloomUpOutputFixedTextureExtent = OutputTextureFixedExtent;
+				BloomUpInputs.BloomUpOutputViewExtent = FIntPoint(BloomUpSourceA.ViewRect.Max.X, BloomUpSourceA.ViewRect.Max.Y);
+
 				return AddMobileBloomUpPass(GraphBuilder, View, BloomUpInputs);
 			};
 
+			Divider /= BloomQualityDivider;
 			const float BloomUpScale = 0.66f * 2.0f;
 			const float BloomGaussianIntensity = Settings.BloomIntensity * Settings.BloomGaussianIntensity;
 
@@ -3107,7 +3125,8 @@ void AddMobilePostProcessingPasses(FRDGBuilder& GraphBuilder, FScene* Scene, con
 				TintA *= BloomGaussianIntensity;
 				TintB *= BloomGaussianIntensity;
 
-				BloomUpOutputs = AddBloomUpPass(PostProcessDownsample_Bloom[IndexA], PostProcessDownsample_Bloom[IndexB], BloomUpScale, TintA, TintB);
+				Divider /= BloomQualityDivider;
+				BloomUpOutputs = AddBloomUpPass(PostProcessDownsample_Bloom[IndexA], PostProcessDownsample_Bloom[IndexB], BloomUpScale, TintA, TintB, FIntPoint::DivideAndRoundUp(BloomOutputFixedTextureSize, Divider));
 			}
 
 			for (int Index = NumDownsamplePasses - 3; Index > 0; --Index)
@@ -3117,7 +3136,8 @@ void AddMobilePostProcessingPasses(FRDGBuilder& GraphBuilder, FScene* Scene, con
 				TintA *= BloomGaussianIntensity;
 				FVector4f TintB = FVector4f(1.0f, 1.0f, 1.0f, 0.0f);
 
-				BloomUpOutputs = AddBloomUpPass(PostProcessDownsample_Bloom[Index], BloomUpOutputs, BloomUpScale, TintA, TintB);
+				Divider /= BloomQualityDivider;
+				BloomUpOutputs = AddBloomUpPass(PostProcessDownsample_Bloom[Index], BloomUpOutputs, BloomUpScale, TintA, TintB, FIntPoint::DivideAndRoundUp(BloomOutputFixedTextureSize, Divider));
 			}
 
 			// Upsample by 2
@@ -3128,7 +3148,8 @@ void AddMobilePostProcessingPasses(FRDGBuilder& GraphBuilder, FScene* Scene, con
 				TintA *= 0.5;
 				FVector4f TintB = FVector4f(1.0f, 1.0f, 1.0f, 0.0f);
 
-				BloomUpOutputs = AddBloomUpPass(PostProcessDownsample_Bloom[0], BloomUpOutputs, BloomUpScale, TintA, TintB);
+				Divider /= BloomQualityDivider;
+				BloomUpOutputs = AddBloomUpPass(PostProcessDownsample_Bloom[0], BloomUpOutputs, BloomUpScale, TintA, TintB, FIntPoint::DivideAndRoundUp(BloomOutputFixedTextureSize, Divider));
 			}
 
 			if (IsLensFlaresEnabled(View))
@@ -3148,7 +3169,7 @@ void AddMobilePostProcessingPasses(FRDGBuilder& GraphBuilder, FScene* Scene, con
 			
 			EyeAdaptationSetupInputs.bUseBasicEyeAdaptation = bUseBasicEyeAdaptation;
 			EyeAdaptationSetupInputs.bUseHistogramEyeAdaptation = bUseHistogramEyeAdaptation;
-			EyeAdaptationSetupInputs.BloomSetup_EyeAdaptation = FScreenPassTexture(TryRegisterExternalTexture(GraphBuilder, View.PrevViewInfo.MobileBloomSetup_EyeAdaptation));
+			EyeAdaptationSetupInputs.BloomSetup_EyeAdaptation = FScreenPassTexture(TryRegisterExternalTexture(GraphBuilder, View.PrevViewInfo.MobileBloomSetup_EyeAdaptation), View.PrevViewInfo.MobileBloomSetup_EyeAdaptationViewSize);
 			if (!EyeAdaptationSetupInputs.BloomSetup_EyeAdaptation.IsValid())
 			{
 				EyeAdaptationSetupInputs.BloomSetup_EyeAdaptation = BloomSetupOutputs.EyeAdaptation;
@@ -3167,6 +3188,7 @@ void AddMobilePostProcessingPasses(FRDGBuilder& GraphBuilder, FScene* Scene, con
 			if ((bUseBasicEyeAdaptation || bUseHistogramEyeAdaptation) && View.ViewState && !View.bStatePrevViewInfoIsReadOnly)
 			{
 				GraphBuilder.QueueTextureExtraction(BloomSetupOutputs.EyeAdaptation.Texture, &View.ViewState->PrevFrameViewInfo.MobileBloomSetup_EyeAdaptation);
+				View.ViewState->PrevFrameViewInfo.MobileBloomSetup_EyeAdaptationViewSize = BloomSetupOutputs.EyeAdaptation.ViewRect;
 			}
 		}
 
diff --git a/Engine/Source/Runtime/Renderer/Private/SceneRendering.h b/Engine/Source/Runtime/Renderer/Private/SceneRendering.h
--- a/Engine/Source/Runtime/Renderer/Private/SceneRendering.h
+++ b/Engine/Source/Runtime/Renderer/Private/SceneRendering.h
@@ -1089,6 +1089,9 @@ struct FPreviousViewInfo
 	// Mobile bloom setup eye adaptation surface.
 	TRefCountPtr<IPooledRenderTarget> MobileBloomSetup_EyeAdaptation;
 
+	// View size of the EyeAdaptation texture
+	FIntRect MobileBloomSetup_EyeAdaptationViewSize;
+
 	// Scene color used for reprojecting next frame to verify the motion vector reprojects correctly.
 	TRefCountPtr<IPooledRenderTarget> VisualizeMotionVectors;
 	FIntRect VisualizeMotionVectorsRect;
