[HWRT] Fix out-of-bounds read in FRayTracingScene::FinishStats(...).
- keep track of which views where active on the frame that the readback was enqueued  and then skip any that was removed in the meantime.

#Synced-CL 49056655

diff --git a/Engine/Source/Runtime/Renderer/Private/RayTracing/RayTracingScene.cpp b/Engine/Source/Runtime/Renderer/Private/RayTracing/RayTracingScene.cpp
--- a/Engine/Source/Runtime/Renderer/Private/RayTracing/RayTracingScene.cpp
+++ b/Engine/Source/Runtime/Renderer/Private/RayTracing/RayTracingScene.cpp
@@ -145,7 +145,7 @@ FRayTracingScene::FViewHandle FRayTracingScene::AddView(uint32 ViewKey)
 	if (ViewIndexMap.Contains(ViewKey))
 	{
 		const int32 ViewIndex = ViewIndexMap[ViewKey];
-		check(ActiveViews.IsValidIndex(ViewIndex) && ActiveViews[ViewIndex] == ViewIndex);
+		check(ActiveViews.IsValidIndex(ViewIndex) && ActiveViews[ViewIndex] == ViewIndex && ViewParameters[ViewIndex].ViewKey == ViewKey);
 		return FViewHandle(ViewIndex);
 	}
 
@@ -167,6 +167,8 @@ FRayTracingScene::FViewHandle FRayTracingScene::AddView(uint32 ViewKey)
 		ViewParameters.SetNum(ViewIndex + 1);
 	}
 
+	ViewParameters[ViewIndex].ViewKey = ViewKey;
+
 	const uint8 NumLayers = uint8(ERayTracingSceneLayer::NUM);
 	for (uint32 LayerIndex = 0; LayerIndex < NumLayers; ++LayerIndex)
 	{
@@ -858,8 +860,15 @@ void FRayTracingScene::FinishStats(FRDGBuilder& GraphBuilder, ERDGPassFlags Comp
 			{
 				FLayer& Layer = Layers[LayerIndex];
 
-				for (int32 ViewIndex : ActiveViews)
+				for (uint32 ViewKey : ReadbackData.ViewKeys)
 				{
+					if (!ViewIndexMap.Contains(ViewKey))
+					{
+						continue;
+					}
+
+					const int32 ViewIndex = ViewIndexMap[ViewKey];
+
 					FLayerView& LayerView = Layer.Views[ViewIndex];
 
 					const uint32 LayerViewNumNativeInstances = LayerView.InstanceBufferBuilder.GetMaxNumInstances();
@@ -893,6 +902,13 @@ void FRayTracingScene::FinishStats(FRDGBuilder& GraphBuilder, ERDGPassFlags Comp
 
 		StatsReadback[StatsReadbackWriteIndex].MaxNumViews = ActiveViews.GetMaxIndex();
 
+		StatsReadback[StatsReadbackWriteIndex].ViewKeys.Reset(ActiveViews.GetMaxIndex());
+
+		for (int32 ViewIndex : ActiveViews)
+		{
+			StatsReadback[StatsReadbackWriteIndex].ViewKeys.Add(ViewParameters[ViewIndex].ViewKey);
+		}
+
 		StatsReadbackWriteIndex = (StatsReadbackWriteIndex + 1u) % MaxReadbackBuffers;
 		StatsReadbackNumPending = FMath::Min(StatsReadbackNumPending + 1u, MaxReadbackBuffers);
 	}
diff --git a/Engine/Source/Runtime/Renderer/Private/RayTracing/RayTracingScene.h b/Engine/Source/Runtime/Renderer/Private/RayTracing/RayTracingScene.h
--- a/Engine/Source/Runtime/Renderer/Private/RayTracing/RayTracingScene.h
+++ b/Engine/Source/Runtime/Renderer/Private/RayTracing/RayTracingScene.h
@@ -294,6 +294,8 @@ private:
 
 		// Used for transforming to translated world space in which TLAS was built.
 		FVector PreViewTranslation;
+
+		uint32 ViewKey = 0;
 	};
 
 	TArray<FViewParameters> ViewParameters;
@@ -329,6 +331,7 @@ private:
 	struct FStatsReadbackData
 	{
 		FRHIGPUBufferReadback* ReadbackBuffer = nullptr;
+		TArray<uint32> ViewKeys;
 		uint32 MaxNumViews = 0;
 	};
 
