Allow TAA history to remain in the transient allocator heap.

#Synced-CL 49075818

diff --git a/Engine/Source/Runtime/Renderer/Private/PostProcess/TemporalAA.cpp b/Engine/Source/Runtime/Renderer/Private/PostProcess/TemporalAA.cpp
--- a/Engine/Source/Runtime/Renderer/Private/PostProcess/TemporalAA.cpp
+++ b/Engine/Source/Runtime/Renderer/Private/PostProcess/TemporalAA.cpp
@@ -985,7 +985,7 @@ FTAAOutputs AddTemporalAAPass(
 		{
 			if (bUseHistoryTexture[i])
 			{
-				GraphBuilder.QueueTextureExtraction(NewHistoryTexture[i], &OutputHistory->RT[i]);
+				GraphBuilder.QueueTextureExtraction(NewHistoryTexture[i], &OutputHistory->RT[i], ERDGResourceExtractionFlags::AllowTransient);
 			}
 		}
 
