Brute force removal of Immutable sampler state that is no longer used.
Removal of some internal Vulkan RHI support code to follow in a second step.
#rnx
#Synced-CL 49100268

diff --git a/Engine/Source/Runtime/Apple/MetalRHI/Private/MetalDynamicRHI.cpp b/Engine/Source/Runtime/Apple/MetalRHI/Private/MetalDynamicRHI.cpp
--- a/Engine/Source/Runtime/Apple/MetalRHI/Private/MetalDynamicRHI.cpp
+++ b/Engine/Source/Runtime/Apple/MetalRHI/Private/MetalDynamicRHI.cpp
@@ -216,8 +216,7 @@ uint64 FMetalDynamicRHI::RHIComputePrecachePSOHash(const FGraphicsPipelineStateI
 bool FMetalDynamicRHI::RHIMatchPrecachePSOInitializers(const FGraphicsPipelineStateInitializer& LHS, const FGraphicsPipelineStateInitializer& RHS)
 {
 	// first check non pointer objects
-	if (LHS.ImmutableSamplerState != RHS.ImmutableSamplerState ||
-		LHS.PrimitiveType != RHS.PrimitiveType ||
+	if (LHS.PrimitiveType != RHS.PrimitiveType ||
 		LHS.bDepthBounds != RHS.bDepthBounds ||
 		LHS.MultiViewCount != RHS.MultiViewCount ||
 		LHS.ShadingRate != RHS.ShadingRate ||
diff --git a/Engine/Source/Runtime/D3D12RHI/Private/D3D12State.cpp b/Engine/Source/Runtime/D3D12RHI/Private/D3D12State.cpp
--- a/Engine/Source/Runtime/D3D12RHI/Private/D3D12State.cpp
+++ b/Engine/Source/Runtime/D3D12RHI/Private/D3D12State.cpp
@@ -653,8 +653,7 @@ uint64 FD3D12DynamicRHI::RHIComputePrecachePSOHash(const FGraphicsPipelineStateI
 bool FD3D12DynamicRHI::RHIMatchPrecachePSOInitializers(const FGraphicsPipelineStateInitializer& LHS, const FGraphicsPipelineStateInitializer& RHS)
 {
 	// first check non pointer objects
-	if (LHS.ImmutableSamplerState != RHS.ImmutableSamplerState ||
-		LHS.PrimitiveType != RHS.PrimitiveType ||
+	if (LHS.PrimitiveType != RHS.PrimitiveType ||
 		LHS.bDepthBounds != RHS.bDepthBounds ||
 		LHS.MultiViewCount != RHS.MultiViewCount ||
 		LHS.ShadingRate != RHS.ShadingRate ||
diff --git a/Engine/Source/Runtime/Engine/Private/Materials/MaterialUniformExpressions.cpp b/Engine/Source/Runtime/Engine/Private/Materials/MaterialUniformExpressions.cpp
--- a/Engine/Source/Runtime/Engine/Private/Materials/MaterialUniformExpressions.cpp
+++ b/Engine/Source/Runtime/Engine/Private/Materials/MaterialUniformExpressions.cpp
@@ -1920,8 +1920,6 @@ void FUniformExpressionSet::FillUniformBuffer(const FMaterialRenderContext& Mate
 
 		// Cache external texture uniform expressions.
 		uint32 ImmutableSamplerIndex = 0;
-		FImmutableSamplerState& ImmutableSamplerState = MaterialRenderContext.MaterialRenderProxy->ImmutableSamplerState;
-		ImmutableSamplerState.Reset();
 		for (int32 ExpressionIndex = 0; ExpressionIndex < UniformExternalTextureParameters.Num(); ExpressionIndex++)
 		{
 			FTextureRHIRef TextureRHI;
@@ -1936,11 +1934,6 @@ void FUniformExpressionSet::FillUniformBuffer(const FMaterialRenderContext& Mate
 			{
 				*ResourceTableTexturePtr = TextureRHI;
 				*ResourceTableSamplerPtr = SamplerStateRHI;
-
-				if (SamplerStateRHI->IsImmutable())
-				{
-					ImmutableSamplerState.ImmutableSamplers[ImmutableSamplerIndex++] = SamplerStateRHI;
-				}
 			}
 			else
 			{
diff --git a/Engine/Source/Runtime/Engine/Public/Materials/MaterialRenderProxy.h b/Engine/Source/Runtime/Engine/Public/Materials/MaterialRenderProxy.h
--- a/Engine/Source/Runtime/Engine/Public/Materials/MaterialRenderProxy.h
+++ b/Engine/Source/Runtime/Engine/Public/Materials/MaterialRenderProxy.h
@@ -6,7 +6,6 @@
 #include "Math/Color.h"
 #include "Math/Vector2D.h"
 #include "RenderResource.h"
-#include "RHIImmutableSamplerState.h"
 #include "Async/Mutex.h"
 #include "Tasks/Task.h"
 #include "Engine/BlendableInterface.h"
@@ -108,9 +107,6 @@ public:
 	/** Cached uniform expressions. */
 	mutable FUniformExpressionCacheContainer UniformExpressionCache;
 
-	/** Cached external texture immutable samplers */
-	mutable FImmutableSamplerState ImmutableSamplerState;
-
 	/** Default constructor. */
 	ENGINE_API FMaterialRenderProxy(FString InMaterialName);
 
diff --git a/Engine/Source/Runtime/RHI/Private/DynamicRHI.cpp b/Engine/Source/Runtime/RHI/Private/DynamicRHI.cpp
--- a/Engine/Source/Runtime/RHI/Private/DynamicRHI.cpp
+++ b/Engine/Source/Runtime/RHI/Private/DynamicRHI.cpp
@@ -19,7 +19,6 @@
 #include "RHIFwd.h"
 #include "DataDrivenShaderPlatformInfo.h"
 #include "RHICommandList.h"
-#include "RHIImmutableSamplerState.h"
 #include "RHIStrings.h"
 #include "RHITextureReference.h"
 #include "RHIAllocators.h"
@@ -597,7 +596,6 @@ uint64 FDynamicRHI::RHIComputeStatePrecachePSOHash(const FGraphicsPipelineStateI
 		uint32 BlendState;
 		uint32 RasterizerState;
 		uint32 DepthStencilState;
-		uint32 ImmutableSamplerState;
 
 		uint32 DrawShadingRate : 8;
 		uint32 PrimitiveType : 8;
@@ -633,9 +631,6 @@ uint64 FDynamicRHI::RHIComputeStatePrecachePSOHash(const FGraphicsPipelineStateI
 		HashKey.DepthStencilState = GetTypeHash(DepthStencilStateInitializerRHI);
 	}
 
-	// Ignore immutable samplers for now
-	//HashKey.ImmutableSamplerState = GetTypeHash(ImmutableSamplerState);
-
 	HashKey.DrawShadingRate = Initializer.ShadingRate;
 	HashKey.PrimitiveType = Initializer.PrimitiveType;
 	HashKey.bDepthBounds = Initializer.bDepthBounds;
@@ -716,8 +711,7 @@ uint64 FDynamicRHI::RHIComputePrecachePSOHash(const FGraphicsPipelineStateInitia
 bool FDynamicRHI::RHIMatchPrecachePSOInitializers(const FGraphicsPipelineStateInitializer& LHS, const FGraphicsPipelineStateInitializer& RHS)
 {
 	// first check non pointer objects
-	if (LHS.ImmutableSamplerState != RHS.ImmutableSamplerState ||
-		LHS.PrimitiveType != RHS.PrimitiveType ||
+	if (LHS.PrimitiveType != RHS.PrimitiveType ||
 		LHS.bDepthBounds != RHS.bDepthBounds ||
 		LHS.MultiViewCount != RHS.MultiViewCount ||
 		LHS.ShadingRate != RHS.ShadingRate ||
diff --git a/Engine/Source/Runtime/RHI/Private/PipelineStateCache.cpp b/Engine/Source/Runtime/RHI/Private/PipelineStateCache.cpp
--- a/Engine/Source/Runtime/RHI/Private/PipelineStateCache.cpp
+++ b/Engine/Source/Runtime/RHI/Private/PipelineStateCache.cpp
@@ -16,7 +16,6 @@ PipelineStateCache.cpp: Pipeline state cache implementation.
 #include "ProfilingDebugging/CsvProfiler.h"
 #include "RHICommandList.h"
 #include "RHIFwd.h"
-#include "RHIImmutableSamplerState.h"
 #include "RHIBreadcrumbs.h"
 #include "Stats/StatsTrace.h"
 #include "Templates/TypeHash.h"
@@ -77,11 +76,6 @@ static inline uint32 GetTypeHash(const FBoundShaderStateInput& Input)
 	return Hash;
 }
 
-static inline uint32 GetTypeHash(const FImmutableSamplerState& Iss)
-{
-	return GetTypeHash(Iss.ImmutableSamplers);
-}
-
 inline uint32 GetTypeHash(const FExclusiveDepthStencil& Ds)
 {
 	return GetTypeHash(Ds.Value);
@@ -99,7 +93,6 @@ static inline uint32 GetTypeHash(const FGraphicsPipelineStateInitializer& Initia
 	Hash = HashCombineFast(Hash, GetTypeHash(Initializer.BlendState));
 	Hash = HashCombineFast(Hash, GetTypeHash(Initializer.RasterizerState));
 	Hash = HashCombineFast(Hash, GetTypeHash(Initializer.DepthStencilState));
-	Hash = HashCombineFast(Hash, GetTypeHash(Initializer.ImmutableSamplerState));
 	Hash = HashCombineFast(Hash, GetTypeHash(Initializer.PrimitiveType));
 	Hash = HashCombineFast(Hash, GetTypeHash(Initializer.RenderTargetsEnabled));
 	Hash = HashCombineFast(Hash, GetTypeHash(Initializer.RenderTargetFormats));
diff --git a/Engine/Source/Runtime/RHI/Public/RHIImmutableSamplerState.h b/Engine/Source/Runtime/RHI/Public/RHIImmutableSamplerState.h
--- a/Engine/Source/Runtime/RHI/Public/RHIImmutableSamplerState.h
+++ b/Engine/Source/Runtime/RHI/Public/RHIImmutableSamplerState.h
@@ -5,39 +5,4 @@
 #include "CoreTypes.h"
 #include "Containers/StaticArray.h"
 
-class FRHISamplerState;
-
-/** Maximum number of immutable samplers in a PSO. */
-enum
-{
-	MaxImmutableSamplers = 2
-};
-
-struct FImmutableSamplerState
-{
-	using TImmutableSamplers = TStaticArray<FRHISamplerState*, MaxImmutableSamplers>;
-
-	FImmutableSamplerState()
-		: ImmutableSamplers(InPlace, nullptr)
-	{}
-
-	void Reset()
-	{
-		for (uint32 Index = 0; Index < MaxImmutableSamplers; ++Index)
-		{
-			ImmutableSamplers[Index] = nullptr;
-		}
-	}
-
-	bool operator==(const FImmutableSamplerState& rhs) const
-	{
-		return ImmutableSamplers == rhs.ImmutableSamplers;
-	}
-
-	bool operator!=(const FImmutableSamplerState& rhs) const
-	{
-		return ImmutableSamplers != rhs.ImmutableSamplers;
-	}
-
-	TImmutableSamplers ImmutableSamplers;
-};
+UE_DEPRECATED_HEADER(5.8, "RHIImmutableSamplerState.h is no longer used")
\ No newline at end of file
diff --git a/Engine/Source/Runtime/RHI/Public/RHIResources.h b/Engine/Source/Runtime/RHI/Public/RHIResources.h
--- a/Engine/Source/Runtime/RHI/Public/RHIResources.h
+++ b/Engine/Source/Runtime/RHI/Public/RHIResources.h
@@ -16,7 +16,6 @@
 #include "PixelFormat.h"
 #include "Async/TaskGraphFwd.h"
 #include "RHIFwd.h"
-#include "RHIImmutableSamplerState.h"
 #include "RHITransition.h"
 #include "MultiGPU.h"
 #include "Math/IntPoint.h"
@@ -4604,7 +4603,6 @@ public:
 		FRHIBlendState*				InBlendState,
 		FRHIRasterizerState*		InRasterizerState,
 		FRHIDepthStencilState*		InDepthStencilState,
-		FImmutableSamplerState		InImmutableSamplerState,
 		EPrimitiveType				InPrimitiveType,
 		uint32						InRenderTargetsEnabled,
 		const TRenderTargetFormats&	InRenderTargetFormats,
@@ -4630,7 +4628,6 @@ public:
 		, BlendState(InBlendState)
 		, RasterizerState(InRasterizerState)
 		, DepthStencilState(InDepthStencilState)
-		, ImmutableSamplerState(InImmutableSamplerState)
 		, PrimitiveType(InPrimitiveType)
 		, RenderTargetsEnabled(InRenderTargetsEnabled)
 		, RenderTargetFormats(InRenderTargetFormats)
@@ -4668,7 +4665,6 @@ public:
 			BlendState != rhs.BlendState ||
 			RasterizerState != rhs.RasterizerState ||
 			DepthStencilState != rhs.DepthStencilState ||
-			ImmutableSamplerState != rhs.ImmutableSamplerState ||
 			PrimitiveType != rhs.PrimitiveType ||
 			bDepthBounds != rhs.bDepthBounds ||
 			MultiViewCount != rhs.MultiViewCount ||
@@ -4748,7 +4744,6 @@ public:
 	FRHIBlendState*					BlendState;
 	FRHIRasterizerState*			RasterizerState;
 	FRHIDepthStencilState*			DepthStencilState;
-	FImmutableSamplerState			ImmutableSamplerState;
 
 	EPrimitiveType					PrimitiveType;
 	uint32							RenderTargetsEnabled;
diff --git a/Engine/Source/Runtime/Renderer/Private/MeshPassProcessor.cpp b/Engine/Source/Runtime/Renderer/Private/MeshPassProcessor.cpp
--- a/Engine/Source/Runtime/Renderer/Private/MeshPassProcessor.cpp
+++ b/Engine/Source/Runtime/Renderer/Private/MeshPassProcessor.cpp
@@ -2243,7 +2243,6 @@ FGraphicsMinimalPipelineStateInitializer PSOCollectorStats::PatchMinimalPipeline
 	//PatchedInitializer.RasterizerState = nullptr;
 	//PatchedInitializer.BlendState = nullptr;
 	//PatchedInitializer.PrimitiveType = PT_TriangleList;
-	//PatchedInitializer.ImmutableSamplerState = FImmutableSamplerState();
 	//PatchedInitializer.BoundShaderState.VertexDeclarationRHI = nullptr;
 	//PatchedInitializer.UniqueEntry = false;
 
diff --git a/Engine/Source/Runtime/Renderer/Private/RayTracing/RayTracingMaterialHitShaders.h b/Engine/Source/Runtime/Renderer/Private/RayTracing/RayTracingMaterialHitShaders.h
--- a/Engine/Source/Runtime/Renderer/Private/RayTracing/RayTracingMaterialHitShaders.h
+++ b/Engine/Source/Runtime/Renderer/Private/RayTracing/RayTracingMaterialHitShaders.h
@@ -58,8 +58,6 @@ private:
 	{
 		const FVertexFactory* RESTRICT VertexFactory = MeshBatch.VertexFactory;
 
-		checkf(MaterialRenderProxy.ImmutableSamplerState.ImmutableSamplers[0] == nullptr, TEXT("Immutable samplers not yet supported in Mesh Draw Command pipeline"));
-
 		FRayTracingMeshCommand SharedCommand;
 
 		SetupRayTracingMeshCommandMaskAndStatus(SharedCommand, MeshBatch, PrimitiveSceneProxy, MaterialResource, RayTracingType);
diff --git a/Engine/Source/Runtime/Renderer/Public/MeshPassProcessor.h b/Engine/Source/Runtime/Renderer/Public/MeshPassProcessor.h
--- a/Engine/Source/Runtime/Renderer/Public/MeshPassProcessor.h
+++ b/Engine/Source/Runtime/Renderer/Public/MeshPassProcessor.h
@@ -17,7 +17,6 @@
 #include "PSOPrecacheMaterial.h"
 #include "Hash/CityHash.h"
 #include "Experimental/Containers/RobinHoodHashTable.h"
-#include "RHIImmutableSamplerState.h"
 #include <atomic>
 
 #define MESH_DRAW_COMMAND_DEBUG_DATA ((!UE_BUILD_SHIPPING && !UE_BUILD_TEST) || VALIDATE_MESH_COMMAND_BINDINGS || WANTS_DRAW_MESH_EVENTS || WITH_DEBUG_VIEW_MODES)
@@ -352,14 +351,12 @@ public:
 		FRHIBlendState*					InBlendState,
 		FRHIRasterizerState*			InRasterizerState,
 		FRHIDepthStencilState*			InDepthStencilState,
-		FImmutableSamplerState			InImmutableSamplerState,
 		EPrimitiveType					InPrimitiveType
 	)
 		: BoundShaderState(InBoundShaderState)
 		, BlendState(InBlendState)
 		, RasterizerState(InRasterizerState)
 		, DepthStencilState(InDepthStencilState)
-		, ImmutableSamplerState(InImmutableSamplerState)
 		, PrimitiveType(InPrimitiveType)
 	{
 	}
@@ -369,7 +366,6 @@ public:
 		, BlendState(InMinimalState.BlendState)
 		, RasterizerState(InMinimalState.RasterizerState)
 		, DepthStencilState(InMinimalState.DepthStencilState)
-		, ImmutableSamplerState(InMinimalState.ImmutableSamplerState)
 		, bDepthBounds(InMinimalState.bDepthBounds)
 		, bAllowVariableRateShading(InMinimalState.bAllowVariableRateShading)
 		, DrawShadingRate(InMinimalState.DrawShadingRate)
@@ -390,7 +386,6 @@ public:
 			, BlendState
 			, RasterizerState
 			, DepthStencilState
-			, ImmutableSamplerState
 			, PrimitiveType
 			, 0
 			, FGraphicsPipelineStateInitializer::TRenderTargetFormats(InPlace, PF_Unknown)
@@ -437,7 +432,6 @@ public:
 			BlendState != rhs.BlendState ||
 			RasterizerState != rhs.RasterizerState ||
 			DepthStencilState != rhs.DepthStencilState ||
-			ImmutableSamplerState != rhs.ImmutableSamplerState ||
 			bDepthBounds != rhs.bDepthBounds ||
 			bAllowVariableRateShading != rhs.bAllowVariableRateShading ||
 			DrawShadingRate != rhs.DrawShadingRate ||
@@ -543,7 +537,6 @@ public:
 	FRHIBlendState*					BlendState;
 	FRHIRasterizerState*			RasterizerState;
 	FRHIDepthStencilState*			DepthStencilState;
-	FImmutableSamplerState			ImmutableSamplerState;
 
 	// Note: FGraphicsMinimalPipelineStateInitializer is 8-byte aligned and can't have any implicit padding,
 	// as it is sometimes hashed and compared as raw bytes. Explicit padding is therefore required between
diff --git a/Engine/Source/Runtime/Renderer/Public/MeshPassProcessor.inl b/Engine/Source/Runtime/Renderer/Public/MeshPassProcessor.inl
--- a/Engine/Source/Runtime/Renderer/Public/MeshPassProcessor.inl
+++ b/Engine/Source/Runtime/Renderer/Public/MeshPassProcessor.inl
@@ -82,7 +82,6 @@ void FMeshPassProcessor::BuildMeshDrawCommands(
 
 	FGraphicsMinimalPipelineStateInitializer PipelineState;
 	PipelineState.PrimitiveType = (EPrimitiveType)MeshBatch.Type;
-	PipelineState.ImmutableSamplerState = MaterialRenderProxy.ImmutableSamplerState;
 
 	EVertexInputStreamType InputStreamType = EVertexInputStreamType::Default;
 	if ((MeshPassFeatures & EMeshPassFeatures::PositionOnly) != EMeshPassFeatures::Default)				InputStreamType = EVertexInputStreamType::PositionOnly;
@@ -268,9 +267,6 @@ void FMeshPassProcessor::AddGraphicsPipelineStateInitializer(
 		FGraphicsMinimalPipelineStateInitializer MinimalPipelineStateInitializer;
 		MinimalPipelineStateInitializer.PrimitiveType = PrimitiveType;
 
-		// Ignore immutable samplers for now - should be passed in?
-		//PipelineState.ImmutableSamplerState = MaterialRenderProxy.ImmutableSamplerState;
-
 		EVertexInputStreamType InputStreamType = EVertexInputStreamType::Default;
 		if ((MeshPassFeatures & EMeshPassFeatures::PositionOnly) != EMeshPassFeatures::Default)				InputStreamType = EVertexInputStreamType::PositionOnly;
 		if ((MeshPassFeatures & EMeshPassFeatures::PositionAndNormalOnly) != EMeshPassFeatures::Default)	InputStreamType = EVertexInputStreamType::PositionAndNormalOnly;
diff --git a/Engine/Source/Runtime/VulkanRHI/Private/VulkanPipeline.cpp b/Engine/Source/Runtime/VulkanRHI/Private/VulkanPipeline.cpp
--- a/Engine/Source/Runtime/VulkanRHI/Private/VulkanPipeline.cpp
+++ b/Engine/Source/Runtime/VulkanRHI/Private/VulkanPipeline.cpp
@@ -2352,9 +2352,8 @@ void FVulkanPipelineStateCacheManager::CreateGfxEntry(const FGraphicsPipelineSta
 		checkf((NumBindlessShaders == 0) || (NumBindlessShaders == NumActiveShaders), TEXT("All shaders must be bindless or non-bindless."));
 
 		// Second pass
-		const int32 NumImmutableSamplers = PSOInitializer.ImmutableSamplerState.ImmutableSamplers.Num();
-		TArrayView<FRHISamplerState*> ImmutableSamplers(NumImmutableSamplers > 0 ? &(FRHISamplerState*&)PSOInitializer.ImmutableSamplerState.ImmutableSamplers[0] : nullptr, NumImmutableSamplers);
-		DescriptorSetLayoutInfo.FinalizeBindings<false>(*Device, UBGatherInfo, ImmutableSamplers, (NumBindlessShaders != 0));
+		//todo: strip out ImmutableSamplers from 3rd parameter.
+		DescriptorSetLayoutInfo.FinalizeBindings<false>(*Device, UBGatherInfo, TArrayView<FRHISamplerState*>(), (NumBindlessShaders != 0));
 	}
 
 	checkf(PSOInitializer.SubpassIndex == 0 || PSOInitializer.SubpassHint != ESubpassHint::None, TEXT("Subpass hint must be set if subpass index > 0."));
diff --git a/Engine/Source/Runtime/VulkanRHI/Private/VulkanRHI.cpp b/Engine/Source/Runtime/VulkanRHI/Private/VulkanRHI.cpp
--- a/Engine/Source/Runtime/VulkanRHI/Private/VulkanRHI.cpp
+++ b/Engine/Source/Runtime/VulkanRHI/Private/VulkanRHI.cpp
@@ -1835,7 +1835,7 @@ IRHITransientResourceAllocator* FVulkanDynamicRHI::RHICreateTransientResourceAll
 
 uint32 FVulkanDynamicRHI::GetPrecachePSOHashVersion()
 {
-	static const uint32 PrecacheHashVersion = 5;
+	static const uint32 PrecacheHashVersion = 6;
 	return PrecacheHashVersion;
 }
 
@@ -1858,7 +1858,6 @@ uint64 FVulkanDynamicRHI::RHIComputeStatePrecachePSOHash(const FGraphicsPipeline
 		uint32 BlendState;
 		uint32 RasterizerState;
 		uint32 DepthStencilState;
-		uint32 ImmutableSamplerState;
 
 		uint32 DrawShadingRate : 8;
 		uint32 PrimitiveType : 8;
@@ -1901,9 +1900,6 @@ uint64 FVulkanDynamicRHI::RHIComputeStatePrecachePSOHash(const FGraphicsPipeline
 		HashKey.DepthStencilState = GetTypeHash(DepthStencilStateInitializerRHI);
 	}
 
-	// Ignore immutable samplers for now
-	//HashKey.ImmutableSamplerState = GetTypeHash(ImmutableSamplerState);
-
 	HashKey.DrawShadingRate = Initializer.ShadingRate;
 	HashKey.PrimitiveType = Initializer.PrimitiveType;
 	HashKey.bDepthBounds = Initializer.bDepthBounds;
@@ -1974,8 +1970,7 @@ uint64 FVulkanDynamicRHI::RHIComputePrecachePSOHash(const FGraphicsPipelineState
 bool FVulkanDynamicRHI::RHIMatchPrecachePSOInitializers(const FGraphicsPipelineStateInitializer& LHS, const FGraphicsPipelineStateInitializer& RHS)
 {
 	// first check non pointer objects
-	if (LHS.ImmutableSamplerState != RHS.ImmutableSamplerState ||
-		LHS.PrimitiveType != RHS.PrimitiveType ||
+	if (LHS.PrimitiveType != RHS.PrimitiveType ||
 		LHS.bDepthBounds != RHS.bDepthBounds ||
 		LHS.MultiViewCount != RHS.MultiViewCount ||
 		LHS.ShadingRate != RHS.ShadingRate ||
