Use render thread version of frame number in render thread function.
Fixes race condition found by TSAN.
#rnx
#Synced-CL 49100268

diff --git a/Engine/Source/Runtime/Renderer/Private/VT/VirtualTextureFeedbackResource.cpp b/Engine/Source/Runtime/Renderer/Private/VT/VirtualTextureFeedbackResource.cpp
--- a/Engine/Source/Runtime/Renderer/Private/VT/VirtualTextureFeedbackResource.cpp
+++ b/Engine/Source/Runtime/Renderer/Private/VT/VirtualTextureFeedbackResource.cpp
@@ -367,7 +367,7 @@ void GetFeedbackShaderParams(int32 InFrameIndex, int32 InVirtualTextureFeedbackT
 void GetFeedbackShaderParams(FFeedbackShaderParams& OutParams)
 {
 	const uint32 VirtualTextureFeedbackTileSize = AlignVirtualTextureFeedbackTileSize(VirtualTextureScalability::GetVirtualTextureFeedbackFactor());
-	VirtualTexture::GetFeedbackShaderParams(GFrameNumber, VirtualTextureFeedbackTileSize, OutParams);
+	VirtualTexture::GetFeedbackShaderParams(GFrameNumberRenderThread, VirtualTextureFeedbackTileSize, OutParams);
 }
 
 void UpdateViewUniformShaderParameters(FFeedbackShaderParams const& InParams, FViewUniformShaderParameters& ViewUniformShaderParameters)
