Skeletal mesh LOD:  Fix for LOD flicker when custom render passes are in use.  Need to check AllViews, which includes CRPs, when checking if there are multiple views, and fix an issue where the frame number in the CRP's view family wasn't initialized.

#Synced-CL 49087395

diff --git a/Engine/Source/Runtime/Engine/Private/SkeletalRender.cpp b/Engine/Source/Runtime/Engine/Private/SkeletalRender.cpp
--- a/Engine/Source/Runtime/Engine/Private/SkeletalRender.cpp
+++ b/Engine/Source/Runtime/Engine/Private/SkeletalRender.cpp
@@ -171,7 +171,7 @@ void FSkeletalMeshObject::UpdateMinDesiredLODLevel(const FSceneView* View, const
 	}
 
 	// When rendering multiple views we need to guard the assignment with a mutex since relevance can occur in parallel.
-	const bool bMultiView = View->Family->Views.Num() > 1;
+	const bool bMultiView = View->Family->AllViews.Num() > 1;
 	if (bMultiView)
 	{
 		DesiredLODLevelMutex.Lock();
diff --git a/Engine/Source/Runtime/Renderer/Private/SceneRendering.cpp b/Engine/Source/Runtime/Renderer/Private/SceneRendering.cpp
--- a/Engine/Source/Runtime/Renderer/Private/SceneRendering.cpp
+++ b/Engine/Source/Runtime/Renderer/Private/SceneRendering.cpp
@@ -3011,6 +3011,8 @@ FSceneRenderer::FSceneRenderer(const FSceneViewFamily* InViewFamily, FHitProxyCo
 
 		CustomRenderPassInfo->CustomRenderPass = CustomRenderPass;
 		CustomRenderPassInfo->ViewFamily.Time = ViewFamily.Time;
+		CustomRenderPassInfo->ViewFamily.FrameNumber = ViewFamily.FrameNumber;
+		CustomRenderPassInfo->ViewFamily.FrameCounter = ViewFamily.FrameCounter;
 		CustomRenderPassInfo->ViewFamily.SetSceneRenderer(this);
 		CustomRenderPassInfo->ViewFamily.bIsSceneTextureSizedCapture = SceneCaptureUserData.bMainViewResolution;
 
