Fix high res screenshot custom depth masking. Ensure custom depth is enabled when highres screenshot capture is enabled. Additionally, remove mobile only renderCustomDepth optimization from SceneVisibility as it was based on results not yet visible (i.e. View.bHasCustomDepthPrimitives not yet populated).

#Synced-CL 48848898

diff --git a/Engine/Source/Runtime/Renderer/Private/MobileShadingRenderer.cpp b/Engine/Source/Runtime/Renderer/Private/MobileShadingRenderer.cpp
--- a/Engine/Source/Runtime/Renderer/Private/MobileShadingRenderer.cpp
+++ b/Engine/Source/Runtime/Renderer/Private/MobileShadingRenderer.cpp
@@ -84,6 +84,7 @@
 #include "DepthCopy.h"
 #include "SingleLayerWaterRendering.h"
 #include "PostProcess/PostProcessCompositeDebugPrimitives.h"
+#include "HighResScreenshot.h"
 
 CSV_DECLARE_CATEGORY_EXTERN(LightCount);
 
@@ -180,8 +181,21 @@ static FMobileCustomDepthStencilUsage GetCustomDepthStencilUsage(const FViewInfo
 				}
 			}
 		}
+		
+		FHighResScreenshotConfig& Config = GetHighResScreenshotConfig();
+
+		// If the highres screenshot UI is open and we're not taking a highres screenshot this frame
+		if (
+	#if WITH_EDITOR
+			Config.bDisplayCaptureRegion ||
+	#endif	
+			GIsHighResScreenshot)
+		{
+			CustomDepthStencilUsage.bUsesCustomDepthStencil |= Config.bMaskEnabled;
+			CustomDepthStencilUsage.bSamplesCustomStencil |= Config.bMaskEnabled;
+		}
 	}
-	
+
 	return CustomDepthStencilUsage;
 }
 
diff --git a/Engine/Source/Runtime/Renderer/Private/SceneVisibility.cpp b/Engine/Source/Runtime/Renderer/Private/SceneVisibility.cpp
--- a/Engine/Source/Runtime/Renderer/Private/SceneVisibility.cpp
+++ b/Engine/Source/Runtime/Renderer/Private/SceneVisibility.cpp
@@ -1429,7 +1429,6 @@ void FRelevancePacket::ComputeRelevance(FDynamicPrimitiveIndexList& DynamicPrimi
 	const bool bIsTranslucentHoldoutEnabled = IsPrimitiveAlphaHoldoutEnabled(View);
 	const bool bIsTranslucentClippedDepthEnabled = NeedTSRThinGeometryDetectionPass(View);
 	const bool bHLODActive = Scene.SceneLODHierarchy.IsActive();
-	const bool bRenderCustomDepth = (ShadingPath == EShadingPath::Mobile) ? FMobileSceneRenderer::ShouldRenderCustomDepth(View) : true;
 	const bool bMobileSupportsSM5MaterialNodes = MobileSupportsSM5MaterialNodes(View.GetShaderPlatform());
 	const FHLODVisibilityState* const HLODState = bHLODActive && ViewState ? &ViewState->HLODVisibilityState : nullptr;
 	float MaxDrawDistanceScale = GetCachedScalabilityCVars().ViewDistanceScale;
@@ -1733,7 +1732,7 @@ void FRelevancePacket::ComputeRelevance(FDynamicPrimitiveIndexList& DynamicPrimi
 									DrawCommandPacket.AddCommandsForMesh(PrimitiveIndex, PrimitiveSceneInfo, StaticMeshRelevance, StaticMesh, CullingPayloadFlags, Scene, bCanCache, EMeshPass::AnisotropyPass);
 								}
 
-								if (ViewRelevance.bRenderCustomDepth && bRenderCustomDepth)
+								if (ViewRelevance.bRenderCustomDepth)
 								{
 									DrawCommandPacket.AddCommandsForMesh(PrimitiveIndex, PrimitiveSceneInfo, StaticMeshRelevance, StaticMesh, CullingPayloadFlags, Scene, bCanCache, EMeshPass::CustomDepth);
 								}
