[VSM] Added r.Shadow.Virtual.PrefilteredDistant.ProjectEnable to gate the prefiltered distant feature (shaders in particular) behind.

#Synced-CL 48982900

diff --git a/Engine/Source/Runtime/Renderer/Private/VirtualShadowMaps/VirtualShadowMapArray.cpp b/Engine/Source/Runtime/Renderer/Private/VirtualShadowMaps/VirtualShadowMapArray.cpp
--- a/Engine/Source/Runtime/Renderer/Private/VirtualShadowMaps/VirtualShadowMapArray.cpp
+++ b/Engine/Source/Runtime/Renderer/Private/VirtualShadowMaps/VirtualShadowMapArray.cpp
@@ -717,10 +717,18 @@ static TAutoConsoleVariable<float> CVarVSMPrefilteredDistantHistoryBlendFactor(
 static TAutoConsoleVariable<bool> CVarVirtualShadowMapPrefilteredDistant(
 	TEXT( "r.Shadow.Virtual.PrefilteredDistant" ),
 	false,
-	TEXT( "Enable/Disable the experimental prefiltered distant shadow feature." ),
+	TEXT( "Enable/Disable the experimental prefiltered distant shadow feature at runtime. r.Shadow.Virtual.PrefilteredDistant.ProjectEnable must be enabled for this to have any effect." ),
 	ECVF_Scalability | ECVF_RenderThreadSafe
 );
 
+bool GVirtualShadowMapPrefilteredDistantProjectEnable = false;
+static FAutoConsoleVariableRef CVarVirtualShadowMapPrefilteredDistantProjectEnable(
+	TEXT( "r.Shadow.Virtual.PrefilteredDistant.ProjectEnable" ),
+	GVirtualShadowMapPrefilteredDistantProjectEnable,
+	TEXT( "Enable/Disable the experimental prefiltered distant shadow feature for the project (enabling compiles additional shader permutations, e.g., for shadow projection)" ),
+	ECVF_ReadOnly
+);
+
 static DynamicRenderScaling::FHeuristicSettings GetDynamicVSMResolutionSettings()
 {
 	DynamicRenderScaling::FHeuristicSettings BucketSetting;
@@ -834,7 +842,7 @@ void FVirtualShadowMapArray::Initialize(
 	bCullBackfacingPixels = CVarCullBackfacingPixels.GetValueOnRenderThread() != 0;
 	bUseHzbOcclusion = CVarShadowsVirtualUseHZB.GetValueOnRenderThread() != 0;
 
-	bEnablePrefilteredDistant = CVarVirtualShadowMapPrefilteredDistant.GetValueOnRenderThread();
+	bEnablePrefilteredDistant = GVirtualShadowMapPrefilteredDistantProjectEnable && CVarVirtualShadowMapPrefilteredDistant.GetValueOnRenderThread();
 
 	UniformParameters.NumFullShadowMaps = 0;
 	UniformParameters.NumSinglePageShadowMaps = 0;
@@ -1722,10 +1730,20 @@ class FMergeStaticPhysicalPagesIndirectCS : public FVirtualShadowMapPageManageme
 };
 IMPLEMENT_GLOBAL_SHADER(FMergeStaticPhysicalPagesIndirectCS, "/Engine/Private/VirtualShadowMaps/VirtualShadowMapPhysicalPageManagement.usf", "MergeStaticPhysicalPagesIndirectCS", SF_Compute);
 
-class FUpdateFilterablePagesSeparableX_CS : public FVirtualShadowMapPageManagementShader
+// Helper to gate relevant shaders with the project setting.
+class FVirtualShadowMapPrefilteredDistantShader : public FVirtualShadowMapPageManagementShader
+{
+public:
+	static bool ShouldCompilePermutation(const FGlobalShaderPermutationParameters& Parameters)
+	{
+		return GVirtualShadowMapPrefilteredDistantProjectEnable && FVirtualShadowMapPageManagementShader::ShouldCompilePermutation(Parameters);
+	}
+};
+
+class FUpdateFilterablePagesSeparableX_CS : public FVirtualShadowMapPrefilteredDistantShader
 {
 	DECLARE_GLOBAL_SHADER(FUpdateFilterablePagesSeparableX_CS);
-	SHADER_USE_PARAMETER_STRUCT(FUpdateFilterablePagesSeparableX_CS, FVirtualShadowMapPageManagementShader)
+	SHADER_USE_PARAMETER_STRUCT(FUpdateFilterablePagesSeparableX_CS, FVirtualShadowMapPrefilteredDistantShader)
 
 	// using FPermutationDomain = FUpdateFilterablePagesSeparablePermutationDomain;
 
@@ -1739,10 +1757,10 @@ class FUpdateFilterablePagesSeparableX_CS : public FVirtualShadowMapPageManageme
 };
 IMPLEMENT_GLOBAL_SHADER(FUpdateFilterablePagesSeparableX_CS, "/Engine/Private/VirtualShadowMaps/VirtualShadowMapPhysicalPageManagement.usf", "UpdateFilterablePagesSeparableX_CS", SF_Compute);
 
-class FUpdateFilterablePagesSeparableY_CS : public FVirtualShadowMapPageManagementShader
+class FUpdateFilterablePagesSeparableY_CS : public FVirtualShadowMapPrefilteredDistantShader
 {
 	DECLARE_GLOBAL_SHADER(FUpdateFilterablePagesSeparableY_CS);
-	SHADER_USE_PARAMETER_STRUCT(FUpdateFilterablePagesSeparableY_CS, FVirtualShadowMapPageManagementShader)
+	SHADER_USE_PARAMETER_STRUCT(FUpdateFilterablePagesSeparableY_CS, FVirtualShadowMapPrefilteredDistantShader)
 
 	BEGIN_SHADER_PARAMETER_STRUCT(FParameters, )
 		SHADER_PARAMETER_RDG_UNIFORM_BUFFER(FVirtualShadowMapUniformParameters, VirtualShadowMap)
@@ -1765,10 +1783,10 @@ class FUpdateFilterablePagesSeparableY_CS : public FVirtualShadowMapPageManageme
 };
 IMPLEMENT_GLOBAL_SHADER(FUpdateFilterablePagesSeparableY_CS, "/Engine/Private/VirtualShadowMaps/VirtualShadowMapPhysicalPageManagement.usf", "UpdateFilterablePagesSeparableY_CS", SF_Compute);
 
-class FGatherFilterableBorderTexels_CS : public FVirtualShadowMapPageManagementShader
+class FGatherFilterableBorderTexels_CS : public FVirtualShadowMapPrefilteredDistantShader
 {
 	DECLARE_GLOBAL_SHADER(FGatherFilterableBorderTexels_CS);
-	SHADER_USE_PARAMETER_STRUCT(FGatherFilterableBorderTexels_CS, FVirtualShadowMapPageManagementShader)
+	SHADER_USE_PARAMETER_STRUCT(FGatherFilterableBorderTexels_CS, FVirtualShadowMapPrefilteredDistantShader)
 
 	BEGIN_SHADER_PARAMETER_STRUCT(FParameters, )
 		SHADER_PARAMETER_RDG_UNIFORM_BUFFER(FVirtualShadowMapUniformParameters, VirtualShadowMap)
@@ -1781,10 +1799,10 @@ class FGatherFilterableBorderTexels_CS : public FVirtualShadowMapPageManagementS
 };
 IMPLEMENT_GLOBAL_SHADER(FGatherFilterableBorderTexels_CS, "/Engine/Private/VirtualShadowMaps/VirtualShadowMapPhysicalPageManagement.usf", "GatherFilterableBorderTexels_CS", SF_Compute);
 
-class FMergeDownFilterablePagesCS : public FVirtualShadowMapPageManagementShader
+class FMergeDownFilterablePagesCS : public FVirtualShadowMapPrefilteredDistantShader
 {
 	DECLARE_GLOBAL_SHADER(FMergeDownFilterablePagesCS);
-	SHADER_USE_PARAMETER_STRUCT(FMergeDownFilterablePagesCS, FVirtualShadowMapPageManagementShader)
+	SHADER_USE_PARAMETER_STRUCT(FMergeDownFilterablePagesCS, FVirtualShadowMapPrefilteredDistantShader)
 
 	BEGIN_SHADER_PARAMETER_STRUCT(FParameters, )
 		SHADER_PARAMETER_RDG_UNIFORM_BUFFER(FVirtualShadowMapUniformParameters, VirtualShadowMap)
diff --git a/Engine/Source/Runtime/Renderer/Private/VirtualShadowMaps/VirtualShadowMapProjection.cpp b/Engine/Source/Runtime/Renderer/Private/VirtualShadowMaps/VirtualShadowMapProjection.cpp
--- a/Engine/Source/Runtime/Renderer/Private/VirtualShadowMaps/VirtualShadowMapProjection.cpp
+++ b/Engine/Source/Runtime/Renderer/Private/VirtualShadowMaps/VirtualShadowMapProjection.cpp
@@ -58,6 +58,7 @@ static TAutoConsoleVariable<int32> CVarSubsurfaceShadowMinSourceAngle(
 );
 
 extern TAutoConsoleVariable<float> CVarMaxDOFResolutionBias;
+extern bool GVirtualShadowMapPrefilteredDistantProjectEnable;
 
 #if MAX_TEST_PERMUTATION > 0
 static TAutoConsoleVariable<int32> CVarTestPermutation(
@@ -214,7 +215,7 @@ class FVirtualShadowMapProjectionCS : public FGlobalShader
 		}
 		
 		// Only need the far Prefiltered Distant Clipmap permutation for directional lights.
-		if( PermutationVector.Get< FPrefilteredDistantClipmap >() && !PermutationVector.Get< FDirectionalLightDim >() )
+		if( PermutationVector.Get< FPrefilteredDistantClipmap >() && (!GVirtualShadowMapPrefilteredDistantProjectEnable || !PermutationVector.Get< FDirectionalLightDim >()) )
 		{
 			return false;
 		}
