Add RENDERER_API tag to accessor for ray tracing scene view handle. Allows inline raytracing shaders in plugins

#Synced-CL 49017163

diff --git a/Engine/Source/Runtime/Renderer/Private/SceneRendering.cpp b/Engine/Source/Runtime/Renderer/Private/SceneRendering.cpp
--- a/Engine/Source/Runtime/Renderer/Private/SceneRendering.cpp
+++ b/Engine/Source/Runtime/Renderer/Private/SceneRendering.cpp
@@ -1133,6 +1133,16 @@ FRDGBufferUAVRef FViewInfo::GetRayTracingInstanceHitCountUAV(FRDGBuilder& GraphB
 	return nullptr;
 }
 
+const FRayTracingScene::FViewHandle& FViewInfo::GetRayTracingSceneViewHandle() const
+{
+	if (UNLIKELY(IStereoRendering::IsStereoEyePass(StereoPass) && IStereoRendering::IsASecondaryPass(StereoPass)))
+	{
+		return GetPrimaryView()->RayTracingSceneViewHandle;
+	}
+
+	return RayTracingSceneViewHandle;
+}
+
 #endif // RHI_RAYTRACING
 
 #if DO_CHECK || USING_CODE_ANALYSIS
diff --git a/Engine/Source/Runtime/Renderer/Private/SceneRendering.h b/Engine/Source/Runtime/Renderer/Private/SceneRendering.h
--- a/Engine/Source/Runtime/Renderer/Private/SceneRendering.h
+++ b/Engine/Source/Runtime/Renderer/Private/SceneRendering.h
@@ -1558,15 +1558,7 @@ public:
 
 	FRDGBufferUAVRef GetRayTracingInstanceHitCountUAV(FRDGBuilder& GraphBuilder) const;
 
-	inline const FRayTracingScene::FViewHandle& GetRayTracingSceneViewHandle() const
-	{
-		if (UNLIKELY(IStereoRendering::IsStereoEyePass(StereoPass) && IStereoRendering::IsASecondaryPass(StereoPass)))
-		{
-			return GetPrimaryView()->RayTracingSceneViewHandle;
-		}
-
-		return RayTracingSceneViewHandle;
-	}
+	RENDERER_API const FRayTracingScene::FViewHandle& GetRayTracingSceneViewHandle() const;
 
 	inline void SetRayTracingSceneViewHandle(const FRayTracingScene::FViewHandle& InHandle)
 	{
