Custom Render Pass:  support for the Scene Capture bExcludeFromSceneTextureExtents flag, to allow large custom render passes to run without increasing other Scene Texture memory.

#Synced-CL 49169569

diff --git a/Engine/Source/Runtime/Renderer/Private/CustomRenderPassSceneCapture.h b/Engine/Source/Runtime/Renderer/Private/CustomRenderPassSceneCapture.h
--- a/Engine/Source/Runtime/Renderer/Private/CustomRenderPassSceneCapture.h
+++ b/Engine/Source/Runtime/Renderer/Private/CustomRenderPassSceneCapture.h
@@ -20,6 +20,7 @@ public:
 	bool bMainViewResolution = false;
 	bool bMainViewCamera = false;
 	bool bIgnoreScreenPercentage = false;
+	bool bExcludeFromSceneTextureExtents = false;
 	FIntPoint SceneTextureDivisor = FIntPoint(1, 1);
 	FName UserSceneTextureBaseColor;
 	FName UserSceneTextureNormal;
diff --git a/Engine/Source/Runtime/Renderer/Private/DeferredShadingRenderer.cpp b/Engine/Source/Runtime/Renderer/Private/DeferredShadingRenderer.cpp
--- a/Engine/Source/Runtime/Renderer/Private/DeferredShadingRenderer.cpp
+++ b/Engine/Source/Runtime/Renderer/Private/DeferredShadingRenderer.cpp
@@ -2383,23 +2383,6 @@ void FDeferredShadingSceneRenderer::Render(FRDGBuilder& GraphBuilder, const FSce
 		RDG_EVENT_SCOPE_STAT(GraphBuilder, CustomRenderPasses, "CustomRenderPasses");
 		RDG_GPU_STAT_SCOPE(GraphBuilder, CustomRenderPasses);
 
-		// If the main view family has MSAA enabled, initialize and use the separate non-MSAA version of the FSceneTextures stored in the first
-		// custom render pass (also used by the other custom render passes).
-		FSceneTextures* CustomRenderPassSceneTextures = &SceneTextures;
-		if (SceneTextures.Config.NumSamples > 1)
-		{
-			FSceneTextures::InitializeViewFamily(GraphBuilder, CustomRenderPassInfos[0].ViewFamily, FamilySize);
-			CustomRenderPassSceneTextures = const_cast<FSceneTextures*>(&CustomRenderPassInfos[0].ViewFamily.GetSceneTextures());
-			
-			// Make sure a separate FSceneTextures structure was allocated in the FSceneRenderer constructor when custom render passes were initialized!
-			check(CustomRenderPassSceneTextures != &SceneTextures);
-		}
-
-		// We want to reset the scene texture uniform buffer to its original state after custom render passes,
-		// so they can't affect downstream rendering.
-		ESceneTextureSetupMode OriginalSceneTextureSetupMode = CustomRenderPassSceneTextures->SetupMode;
-		TRDGUniformBufferRef<FSceneTextureUniformParameters> OriginalSceneTextureUniformBuffer = CustomRenderPassSceneTextures->UniformBuffer;
-
 		for (int32 i = 0; i < CustomRenderPassInfos.Num(); ++i)
 		{
 			FCustomRenderPassBase* CustomRenderPass = CustomRenderPassInfos[i].CustomRenderPass;
@@ -2407,6 +2390,20 @@ void FDeferredShadingSceneRenderer::Render(FRDGBuilder& GraphBuilder, const FSce
 			FNaniteShadingCommands& NaniteBasePassShadingCommands = CustomRenderPassInfos[i].NaniteBasePassShadingCommands;
 			check(CustomRenderPass);
 
+			// Initialize scene textures if owned by this Custom Render Pass (not reusing the main view family or shared non-MSAA SceneTextures).
+			FViewFamilyInfo& CustomRenderPassViewFamily = CustomRenderPassInfos[i].ViewFamily;
+			if (CustomRenderPassViewFamily.IsSceneTexturesOwner())
+			{
+				FSceneTextures::InitializeViewFamily(GraphBuilder, CustomRenderPassViewFamily, FamilySize);
+			}
+			FSceneTextures* CustomRenderPassSceneTextures = &CustomRenderPassViewFamily.GetSceneTextures();
+
+			// We want to reset the scene texture uniform buffer to its original state after custom render passes,
+			// so they can't affect downstream rendering.  This includes both cases where CRPs use the main ViewFamily's
+			// SceneTextures, and cases where CRPs share non-MSAA SceneTextures.
+			ESceneTextureSetupMode OriginalSceneTextureSetupMode = CustomRenderPassSceneTextures->SetupMode;
+			TRDGUniformBufferRef<FSceneTextureUniformParameters> OriginalSceneTextureUniformBuffer = CustomRenderPassSceneTextures->UniformBuffer;
+
 			CustomRenderPass->BeginPass(GraphBuilder);
 
 			{
diff --git a/Engine/Source/Runtime/Renderer/Private/MobileShadingRenderer.cpp b/Engine/Source/Runtime/Renderer/Private/MobileShadingRenderer.cpp
--- a/Engine/Source/Runtime/Renderer/Private/MobileShadingRenderer.cpp
+++ b/Engine/Source/Runtime/Renderer/Private/MobileShadingRenderer.cpp
@@ -1345,17 +1345,25 @@ void FMobileSceneRenderer::Render(FRDGBuilder& GraphBuilder, const FSceneRenderU
 		RDG_EVENT_SCOPE_STAT(GraphBuilder, CustomRenderPasses, "CustomRenderPasses");
 		RDG_GPU_STAT_SCOPE(GraphBuilder, CustomRenderPasses);
 
-		// We want to reset the scene texture uniform buffer to its original state after custom render passes,
-		// so they can't affect downstream rendering.
-		EMobileSceneTextureSetupMode OriginalSceneTextureSetupMode = SceneTextures.MobileSetupMode;
-		TRDGUniformBufferRef<FMobileSceneTextureUniformParameters> OriginalSceneTextureUniformBuffer = SceneTextures.MobileUniformBuffer;
-
 		for (int32 i = 0; i < CustomRenderPassInfos.Num(); ++i)
 		{
 			FCustomRenderPassBase* CustomRenderPass = CustomRenderPassInfos[i].CustomRenderPass;
 			TArray<FViewInfo>& CustomRenderPassViews = CustomRenderPassInfos[i].Views;
 			check(CustomRenderPass);
 
+			// Initialize scene textures if owned by this Custom Render Pass (not reusing the main view family).
+			FViewFamilyInfo& CustomRenderPassViewFamily = CustomRenderPassInfos[i].ViewFamily;
+			if (CustomRenderPassViewFamily.IsSceneTexturesOwner())
+			{
+				FSceneTextures::InitializeViewFamily(GraphBuilder, CustomRenderPassViewFamily, FamilySize);
+			}
+			FSceneTextures& CustomRenderPassSceneTextures = CustomRenderPassViewFamily.GetSceneTextures();
+
+			// We want to reset the scene texture uniform buffer to its original state after custom render passes,
+			// so they can't affect downstream rendering.
+			EMobileSceneTextureSetupMode OriginalSceneTextureSetupMode = CustomRenderPassSceneTextures.MobileSetupMode;
+			TRDGUniformBufferRef<FMobileSceneTextureUniformParameters> OriginalSceneTextureUniformBuffer = CustomRenderPassSceneTextures.MobileUniformBuffer;
+
 			CustomRenderPass->BeginPass(GraphBuilder);
 
 			{
@@ -1369,23 +1377,23 @@ void FMobileSceneRenderer::Render(FRDGBuilder& GraphBuilder, const FSceneRenderU
 
 				if (bIsFullDepthPrepassEnabled)
 				{
-					RenderFullDepthPrepass(GraphBuilder, CustomRenderPassViews, SceneTextures, InstanceCullingManager, true);
+					RenderFullDepthPrepass(GraphBuilder, CustomRenderPassViews, CustomRenderPassSceneTextures, InstanceCullingManager, true);
 					if (!bRequiresSceneDepthAux)
 					{
-						AddResolveSceneDepthPass(GraphBuilder, CustomRenderPassViews, SceneTextures.Depth);
+						AddResolveSceneDepthPass(GraphBuilder, CustomRenderPassViews, CustomRenderPassSceneTextures.Depth);
 					}
 				}
 
 				// Render base pass if the custom pass requires it. Otherwise if full depth prepass is not enabled, then depth is generated in the base pass.
 				if (CustomRenderPass->GetRenderMode() == FCustomRenderPassBase::ERenderMode::DepthAndBasePass || (CustomRenderPass->GetRenderMode() == FCustomRenderPassBase::ERenderMode::DepthPass && !bIsFullDepthPrepassEnabled))
 				{
-					RenderCustomRenderPassBasePass(GraphBuilder, CustomRenderPassViews, ViewFamilyTexture, SceneTextures, InstanceCullingManager, CustomRenderPass->IsTranslucentIncluded());
+					RenderCustomRenderPassBasePass(GraphBuilder, CustomRenderPassViews, ViewFamilyTexture, CustomRenderPassSceneTextures, InstanceCullingManager, CustomRenderPass->IsTranslucentIncluded());
 				}
 
-				SceneTextures.MobileSetupMode = EMobileSceneTextureSetupMode::SceneColor | EMobileSceneTextureSetupMode::SceneDepth | EMobileSceneTextureSetupMode::SceneDepthAux;
-				SceneTextures.MobileUniformBuffer = CreateMobileSceneTextureUniformBuffer(GraphBuilder, &SceneTextures, SceneTextures.MobileSetupMode);
+				CustomRenderPassSceneTextures.MobileSetupMode = EMobileSceneTextureSetupMode::SceneColor | EMobileSceneTextureSetupMode::SceneDepth | EMobileSceneTextureSetupMode::SceneDepthAux;
+				CustomRenderPassSceneTextures.MobileUniformBuffer = CreateMobileSceneTextureUniformBuffer(GraphBuilder, &CustomRenderPassSceneTextures, CustomRenderPassSceneTextures.MobileSetupMode);
 
-				CopySceneCaptureComponentToTarget(GraphBuilder, SceneTextures, CustomRenderPass->GetRenderTargetTexture(), ViewFamily, CustomRenderPassViews);
+				CopySceneCaptureComponentToTarget(GraphBuilder, CustomRenderPassSceneTextures, CustomRenderPass->GetRenderTargetTexture(), ViewFamily, CustomRenderPassViews);
 
 				CustomRenderPass->PostRender(GraphBuilder);
 
@@ -1399,8 +1407,8 @@ void FMobileSceneRenderer::Render(FRDGBuilder& GraphBuilder, const FSceneRenderU
 
 			CustomRenderPass->EndPass(GraphBuilder);
 
-			SceneTextures.MobileSetupMode = OriginalSceneTextureSetupMode;
-			SceneTextures.MobileUniformBuffer = OriginalSceneTextureUniformBuffer;
+			CustomRenderPassSceneTextures.MobileSetupMode = OriginalSceneTextureSetupMode;
+			CustomRenderPassSceneTextures.MobileUniformBuffer = OriginalSceneTextureUniformBuffer;
 		}
 	}
 
diff --git a/Engine/Source/Runtime/Renderer/Private/SceneCaptureRendering.cpp b/Engine/Source/Runtime/Renderer/Private/SceneCaptureRendering.cpp
--- a/Engine/Source/Runtime/Renderer/Private/SceneCaptureRendering.cpp
+++ b/Engine/Source/Runtime/Renderer/Private/SceneCaptureRendering.cpp
@@ -1039,6 +1039,7 @@ public:
 		UserData->bMainViewResolution = CaptureComponent->ShouldRenderWithMainViewResolution();
 		UserData->bMainViewCamera = CaptureComponent->ShouldRenderWithMainViewCamera();
 		UserData->bIgnoreScreenPercentage = CaptureComponent->ShouldIgnoreScreenPercentage();
+		UserData->bExcludeFromSceneTextureExtents = CaptureComponent->bExcludeFromSceneTextureExtents;
 		UserData->SceneTextureDivisor = CaptureComponent->MainViewResolutionDivisor.ComponentMax(FIntPoint(1,1));
 		UserData->UserSceneTextureBaseColor = CaptureComponent->UserSceneTextureBaseColor;
 		UserData->UserSceneTextureNormal = CaptureComponent->UserSceneTextureNormal;
diff --git a/Engine/Source/Runtime/Renderer/Private/SceneRendering.cpp b/Engine/Source/Runtime/Renderer/Private/SceneRendering.cpp
--- a/Engine/Source/Runtime/Renderer/Private/SceneRendering.cpp
+++ b/Engine/Source/Runtime/Renderer/Private/SceneRendering.cpp
@@ -2958,6 +2958,9 @@ FSceneRenderer::FSceneRenderer(const FSceneViewFamily* InViewFamily, FHitProxyCo
 	// Prepare custom render passes and their views:
 	CustomRenderPassInfos.Empty(Scene->CustomRenderPassRendererInputs.Num());
 
+	// If multiple CRPs are in use with MSAA active, we share the non-MSAA version of scene textures across them.
+	FCustomRenderPassInfo* CustomRenderPassInfoNoMSAA = nullptr;
+
 	int32 NumAdditionalViews = 0;
 	int32 IncrementIfNotRemoved = 1;
 	for (int32 i = 0; i < Scene->CustomRenderPassRendererInputs.Num(); i+=IncrementIfNotRemoved)
@@ -2998,19 +3001,25 @@ FSceneRenderer::FSceneRenderer(const FSceneViewFamily* InViewFamily, FHitProxyCo
 		// The GetDefaultMSAACount call below must match the logic that initializes NumSamples in FSceneTexturesConfig::Init.  Downstream
 		// code uses NumSamples to detect MSAA, after the config has been initialized.
 		FCustomRenderPassInfo* CustomRenderPassInfo;
-		if (GetDefaultMSAACount(Scene->GetFeatureLevel(), GDynamicRHI->RHIGetPlatformTextureMaxSampleCount()) > 1)
+		if (SceneCaptureUserData.bExcludeFromSceneTextureExtents)
+		{
+			// If excluded from scene texture extents, generate its own scene textures struct, by not passing a source view family.
+			CustomRenderPassInfo = &CustomRenderPassInfos.Emplace_GetRef(FamilyCVS);
+		}
+		else if (GetDefaultMSAACount(Scene->GetFeatureLevel(), GDynamicRHI->RHIGetPlatformTextureMaxSampleCount()) > 1)
 		{
 			// If main pass is MSAA, create a view family with its own FSceneTextures, by calling the constructor that doesn't accept another view family.
 			// MSAA isn't supported for custom render passes, so we'll initialize this one without MSAA (see FSceneRenderer::OnRenderBegin, where we set
 			// NumSamples to 1 for Custom Render Pass scene texture configs).  First custom render pass gets its own, subsequent custom render passes share
 			// with the first custom render pass.
-			if (CustomRenderPassInfos.Num() == 0)
+			if (!CustomRenderPassInfoNoMSAA)
 			{
 				CustomRenderPassInfo = &CustomRenderPassInfos.Emplace_GetRef(FamilyCVS);
+				CustomRenderPassInfoNoMSAA = CustomRenderPassInfo;
 			}
 			else
 			{
-				CustomRenderPassInfo = &CustomRenderPassInfos.Emplace_GetRef(FamilyCVS, CustomRenderPassInfos[0].ViewFamily);
+				CustomRenderPassInfo = &CustomRenderPassInfos.Emplace_GetRef(FamilyCVS, CustomRenderPassInfoNoMSAA->ViewFamily);
 			}
 		}
 		else
@@ -3037,6 +3046,7 @@ FSceneRenderer::FSceneRenderer(const FSceneViewFamily* InViewFamily, FHitProxyCo
 		ViewInitOptions.ViewActor = PassInput.ViewActor;
 		ViewInitOptions.ShowOnlyPrimitives = PassInput.ShowOnlyPrimitives;
 		ViewInitOptions.HiddenPrimitives = PassInput.HiddenPrimitives;
+		ViewInitOptions.bExcludeFromSceneTextureExtents = SceneCaptureUserData.bExcludeFromSceneTextureExtents;
 
 		FSceneView NewView(ViewInitOptions);
 		FViewInfo* ViewInfo = &CustomRenderPassInfo->Views.Emplace_GetRef(&NewView);
@@ -3181,6 +3191,22 @@ FIntPoint FSceneRenderer::QuantizeViewRectMin(const FIntPoint& ViewRectMin)
 // static
 FIntPoint FSceneRenderer::GetDesiredInternalBufferSize(const FSceneViewFamily& ViewFamily)
 {
+	// If this function is called for a CustomRenderPass view family, it means it is getting its own scene textures,
+	// separate from the main renderer, and only its own views should be counted, not AllViews.
+	if (ViewFamily.Views[0]->CustomRenderPass)
+	{
+		FIntPoint FamilySizeUpperBound(0, 0);
+		for (const FSceneView* View : ViewFamily.Views)
+		{
+			FamilySizeUpperBound.X = FMath::Max(FamilySizeUpperBound.X, View->UnscaledViewRect.Max.X);
+			FamilySizeUpperBound.Y = FMath::Max(FamilySizeUpperBound.Y, View->UnscaledViewRect.Max.Y);
+		}
+
+		FIntPoint DesiredBufferSize;
+		QuantizeSceneBufferSize(FamilySizeUpperBound, DesiredBufferSize);
+		return DesiredBufferSize;
+	}
+
 	// If not supporting screen percentage, bypass all computation.
 	// This bypasses not only resolution scaling, but also quantizing the buffer size.
 	// The only configurations that do not support screen percentage are Forward Mobile non-HDR and Forward Mobile w/ Tonemap Subpass.
@@ -3191,6 +3217,12 @@ FIntPoint FSceneRenderer::GetDesiredInternalBufferSize(const FSceneViewFamily& V
 
 		for (const FSceneView* View : ViewFamily.AllViews)
 		{
+			// If a custom render pass has bExcludeFromSceneTextureExtents set, it will get its own scene textures, and shouldn't be counted in the main view family.
+			if (View->CustomRenderPass && View->bExcludeFromSceneTextureExtents)
+			{
+				continue;
+			}
+
 			FamilySizeUpperBound.X = FMath::Max(FamilySizeUpperBound.X, View->UnscaledViewRect.Max.X);
 			FamilySizeUpperBound.Y = FMath::Max(FamilySizeUpperBound.Y, View->UnscaledViewRect.Max.Y);
 		}
@@ -3226,6 +3258,12 @@ FIntPoint FSceneRenderer::GetDesiredInternalBufferSize(const FSceneViewFamily& V
 	float MaxOverscanResolutionFraction = 1.0f;
 	for (const FSceneView* View : ViewFamily.AllViews)
 	{
+		// If a custom render pass has bExcludeFromSceneTextureExtents set, it will get its own scene textures, and shouldn't be counted in the main view family.
+		if (View->CustomRenderPass && View->bExcludeFromSceneTextureExtents)
+		{
+			continue;
+		}
+
 		MaxOverscanResolutionFraction = FMath::Max(MaxOverscanResolutionFraction, View->SceneViewInitOptions.OverscanResolutionFraction);
 	}
 
@@ -3233,6 +3271,12 @@ FIntPoint FSceneRenderer::GetDesiredInternalBufferSize(const FSceneViewFamily& V
 	
 	for (const FSceneView* View : ViewFamily.AllViews)
 	{
+		// If a custom render pass has bExcludeFromSceneTextureExtents set, it will get its own scene textures, and shouldn't be counted in the main view family.
+		if (View->CustomRenderPass && View->bExcludeFromSceneTextureExtents)
+		{
+			continue;
+		}
+
 		// Note: This ensures that custom passes (rendered with the main renderer) ignore screen percentage, like regular scene captures.
 		const float AdjustedResolutionFractionUpperBounds = View->CustomRenderPass ? 1.0f : (View->SceneViewInitOptions.OverridePrimaryResolutionFraction > 0.0 ? (View->SceneViewInitOptions.OverridePrimaryResolutionFraction * ViewFamily.SecondaryViewFraction)  : ResolutionFractionUpperBound);
 		
@@ -4022,7 +4066,15 @@ IVisibilityTaskData* FSceneRenderer::OnRenderBegin(FRDGBuilder& GraphBuilder, co
 				// needs to be copied.  The FSceneTextures structure itself is pointer shared, and doesn't need to be copied.
 				for (FCustomRenderPassInfo& CustomRenderPass : Renderer->CustomRenderPassInfos)
 				{
-					CustomRenderPass.ViewFamily.SceneTexturesConfig = Renderer->ViewFamily.SceneTexturesConfig;
+					// If this is excluded from scene texture extents, it computes its own scene texture config, rather than copying.
+					if (CustomRenderPass.ViewFamily.Views[0]->bExcludeFromSceneTextureExtents)
+					{
+						InitializeSceneTexturesConfig(CustomRenderPass.ViewFamily.SceneTexturesConfig, CustomRenderPass.ViewFamily);
+					}
+					else
+					{
+						CustomRenderPass.ViewFamily.SceneTexturesConfig = Renderer->ViewFamily.SceneTexturesConfig;
+					}
 
 					// Custom Render Passes don't support MSAA.  If MSAA is enabled, the first Custom Render Pass will allocate a separate non-MSAA
 					// FSceneTextures, initialized using this config (see logic in the FSceneRenderer constructor that fills in CustomRenderPassInfos).
diff --git a/Engine/Source/Runtime/Renderer/Private/SceneRendering.h b/Engine/Source/Runtime/Renderer/Private/SceneRendering.h
--- a/Engine/Source/Runtime/Renderer/Private/SceneRendering.h
+++ b/Engine/Source/Runtime/Renderer/Private/SceneRendering.h
@@ -2007,6 +2007,11 @@ public:
 		return bIsViewFamilyInfo && SceneTextures->bIsSceneTexturesInitialized ? SceneTextures : nullptr;
 	}
 
+	inline bool IsSceneTexturesOwner() const
+	{
+		return bIsViewFamilyInfo && SceneTextures->Owner == this;
+	}
+
 private:
 	friend struct FMinimalSceneTextures;
 	friend struct FSceneTextures;
diff --git a/Engine/Source/Runtime/Renderer/Private/SceneTextures.cpp b/Engine/Source/Runtime/Renderer/Private/SceneTextures.cpp
--- a/Engine/Source/Runtime/Renderer/Private/SceneTextures.cpp
+++ b/Engine/Source/Runtime/Renderer/Private/SceneTextures.cpp
@@ -163,12 +163,32 @@ public:
 		// If explicitly requested, don't include this in history tracking for future requests.  This extent request can still be adjusted based on existing history.
 		bool bExcludeFromHistoryUpdate = false;
 
-		for (const FSceneView* View : ViewFamily.AllViews)
+		// If this function is called for a CustomRenderPass view family, it means it is getting its own scene textures,
+		// separate from the main renderer, and only its own views should be counted, not AllViews.
+		if (ViewFamily.Views[0]->CustomRenderPass)
 		{
-			bIsSceneCapture |= View->bIsSceneCapture;
-			bIsReflectionCapture |= View->bIsReflectionCapture;
-			bIsVRScene |= (IStereoRendering::IsStereoEyeView(*View) && GEngine->XRSystem.IsValid());
-			bExcludeFromHistoryUpdate |= View->bExcludeFromSceneTextureExtents;
+			for (const FSceneView* View : ViewFamily.Views)
+			{
+				bIsSceneCapture |= View->bIsSceneCapture;
+				bIsReflectionCapture |= View->bIsReflectionCapture;
+				bExcludeFromHistoryUpdate |= View->bExcludeFromSceneTextureExtents;
+			}
+		}
+		else
+		{
+			for (const FSceneView* View : ViewFamily.AllViews)
+			{
+				// If a custom render pass has bExcludeFromSceneTextureExtents set, it will get its own scene textures, and shouldn't be counted in the main view family.
+				if (View->CustomRenderPass && View->bExcludeFromSceneTextureExtents)
+				{
+					continue;
+				}
+
+				bIsSceneCapture |= View->bIsSceneCapture;
+				bIsReflectionCapture |= View->bIsReflectionCapture;
+				bIsVRScene |= (IStereoRendering::IsStereoEyeView(*View) && GEngine->XRSystem.IsValid());
+				bExcludeFromHistoryUpdate |= View->bExcludeFromSceneTextureExtents;
+			}
 		}
 
 		FIntPoint DesiredExtent = FIntPoint::ZeroValue;
