[HWRT] Fix issues with dynamic primitive data.
- incorrect InstanceSceneDataOffset and number of instances being used in some paths.
- not taking MaterialsView into account.

#Synced-CL 49186836

diff --git a/Engine/Source/Runtime/Renderer/Private/RayTracing/RayTracing.cpp b/Engine/Source/Runtime/Renderer/Private/RayTracing/RayTracing.cpp
--- a/Engine/Source/Runtime/Renderer/Private/RayTracing/RayTracing.cpp
+++ b/Engine/Source/Runtime/Renderer/Private/RayTracing/RayTracing.cpp
@@ -917,7 +917,6 @@ namespace RayTracing
 			FPrimitiveSceneProxy* SceneProxy = Scene.PrimitiveSceneProxies[PrimitiveIndex];
 			FPrimitiveSceneInfo* SceneInfo = Scene.Primitives[PrimitiveIndex];
 			const FScene::FPrimitiveRayTracingData& PrimitiveRayTracingData = Scene.PrimitiveRayTracingDatas[PrimitiveIndex];
-			const FPersistentPrimitiveIndex PersistentPrimitiveIndex = SceneInfo->GetPersistentIndex();
 
 			if (DynamicRayTracingPrimitive.InstancesRange.Size<int32>() > 0)
 			{
@@ -957,17 +956,19 @@ namespace RayTracing
 						RayTracingScene.GeometriesToBuild.Add(Geometry);
 					}
 
+					TConstArrayView<FMeshBatch> InstanceMaterials = Instance.GetMaterials();
+
 					// Validate the material/segment counts
-					if (!ensureMsgf(Instance.GetMaterials().Num() <= NumSegments,
+					if (!ensureMsgf(InstanceMaterials.Num() <= NumSegments,
 						TEXT("Ray tracing material assignment validation failed for geometry '%s'. "
 							"Instance.GetMaterials().Num() = %d, Geometry->Initializer.Segments.Num() = %d."),
-						*Geometry->Initializer.DebugName.ToString(), Instance.GetMaterials().Num(),
+						*Geometry->Initializer.DebugName.ToString(), InstanceMaterials.Num(),
 						NumSegments))
 					{
 						continue;
 					}
 
-					if (Instance.GetMaterials().IsEmpty())
+					if (InstanceMaterials.IsEmpty())
 					{
 						// If the material list is empty, skip this instance altogether
 						continue;
@@ -1010,14 +1011,13 @@ namespace RayTracing
 						continue;
 					}
 
-					int32 PrimitiveId = PersistentPrimitiveIndex.Index;
 					int32 InstanceSceneDataOffset = SceneInfo->GetInstanceSceneDataOffset();
+					uint32 NumInstanceSceneDataEntries = SceneInfo->GetNumInstanceSceneDataEntries();
 
-					if (Instance.Materials.Num() > 0 && Instance.Materials[0].Elements.Num() > 0 && Instance.Materials[0].Elements[0].DynamicPrimitiveData != nullptr)
+					if (InstanceMaterials.Num() > 0 && InstanceMaterials[0].Elements.Num() > 0 && InstanceMaterials[0].Elements[0].DynamicPrimitiveData != nullptr)
 					{
-						check(Instance.NumTransforms == Instance.Materials[0].Elements[0].NumInstances);
-						PrimitiveId = ViewDynamicPrimitiveId + Instance.Materials[0].Elements[0].DynamicPrimitiveIndex;
-						InstanceSceneDataOffset = ViewInstanceSceneDataOffset + Instance.Materials[0].Elements[0].DynamicPrimitiveInstanceSceneDataOffset;
+						InstanceSceneDataOffset = ViewInstanceSceneDataOffset + InstanceMaterials[0].Elements[0].DynamicPrimitiveInstanceSceneDataOffset;
+						NumInstanceSceneDataEntries = InstanceMaterials[0].Elements[0].NumInstances;
 					}
 
 					FRayTracingGeometryInstance RayTracingInstance;
@@ -1049,8 +1049,8 @@ namespace RayTracing
 						TArrayView<uint32> InstanceSceneDataOffsets = RayTracingScene.Allocate<uint32>(PrimitiveInstanceIndices.Num());
 						for (int32 InstanceIndex = 0; InstanceIndex < PrimitiveInstanceIndices.Num(); ++InstanceIndex)
 						{
-							check(PrimitiveInstanceIndices[InstanceIndex] <= uint32(SceneInfo->GetNumInstanceSceneDataEntries()));
-							InstanceSceneDataOffsets[InstanceIndex] = SceneInfo->GetInstanceSceneDataOffset() + PrimitiveInstanceIndices[InstanceIndex];
+							check(PrimitiveInstanceIndices[InstanceIndex] <= NumInstanceSceneDataEntries);
+							InstanceSceneDataOffsets[InstanceIndex] = InstanceSceneDataOffset + PrimitiveInstanceIndices[InstanceIndex];
 						}
 
 						RayTracingInstance.InstanceSceneDataOffsets = InstanceSceneDataOffsets;
@@ -1083,7 +1083,7 @@ namespace RayTracing
 					else
 					{
 						// If array of transforms was not provided, get the instance transforms from GPU Scene
-						check(Instance.NumTransforms <= uint32(SceneInfo->GetNumInstanceSceneDataEntries()));
+						check(Instance.NumTransforms <= NumInstanceSceneDataEntries);
 						RayTracingInstance.NumTransforms = Instance.NumTransforms;
 						RayTracingInstance.BaseInstanceSceneDataOffset = InstanceSceneDataOffset;
 					}
@@ -1159,7 +1159,7 @@ namespace RayTracing
 						FRayTracingMeshBatchWorkItem& WorkItem = ViewMeshBatchTaskData.Page->WorkItems[ViewMeshBatchTaskData.Page->NumWorkItems];
 						ViewMeshBatchTaskData.Page->NumWorkItems++;
 
-						ViewMeshBatchTaskData.NumPendingMeshBatches += Instance.GetMaterials().Num();
+						ViewMeshBatchTaskData.NumPendingMeshBatches += InstanceMaterials.Num();
 
 						if (Instance.OwnsMaterials())
 						{
@@ -1177,7 +1177,6 @@ namespace RayTracing
 					}
 					else
 					{
-						TArrayView<const FMeshBatch> InstanceMaterials = Instance.GetMaterials();
 						for (const FMeshBatch& MeshBatch : InstanceMaterials)
 						{
 							FDynamicRayTracingMeshCommandContext CommandContext(View.DynamicRayTracingMeshCommandStorage, View.VisibleRayTracingShaderBindings, Geometry->GetRHI(), MeshBatch.SegmentIndex, Geometry->Initializer.bAllowUpdate, SBTAllocation);
