Rolling back change for allowing TAA.History texture to be transient since there are validation errors

#Synced-CL 49027503

diff --git a/Engine/Source/Runtime/Renderer/Private/PostProcess/TemporalAA.cpp b/Engine/Source/Runtime/Renderer/Private/PostProcess/TemporalAA.cpp
--- a/Engine/Source/Runtime/Renderer/Private/PostProcess/TemporalAA.cpp
+++ b/Engine/Source/Runtime/Renderer/Private/PostProcess/TemporalAA.cpp
@@ -985,7 +985,7 @@ FTAAOutputs AddTemporalAAPass(
 		{
 			if (bUseHistoryTexture[i])
 			{
-				GraphBuilder.QueueTextureExtraction(NewHistoryTexture[i], &OutputHistory->RT[i], ERDGResourceExtractionFlags::AllowTransient);
+				GraphBuilder.QueueTextureExtraction(NewHistoryTexture[i], &OutputHistory->RT[i]);
 			}
 		}
 
