SkyAtmosphere - fix for race condition with shadow map creation, when dispatching skyatmosphere lut generation on async compute pipe earlier in the frame.

#tests PC editor dx12 + debug
#rnx
#Synced-CL 44689559

diff --git a/Engine/Source/Runtime/Renderer/Private/SkyAtmosphereRendering.cpp b/Engine/Source/Runtime/Renderer/Private/SkyAtmosphereRendering.cpp
--- a/Engine/Source/Runtime/Renderer/Private/SkyAtmosphereRendering.cpp
+++ b/Engine/Source/Runtime/Renderer/Private/SkyAtmosphereRendering.cpp
@@ -228,7 +228,7 @@ static TAutoConsoleVariable<int32> CVarSkyAtmosphereEditorNotifications(
 
 static TAutoConsoleVariable<int32> CVarSkyAtmosphereAsyncCompute(
 	TEXT("r.SkyAtmosphereAsyncCompute"), 0,
-	TEXT("SkyAtmosphere on async compute (default: false). When running on the async pipe, SkyAtmosphere lut generation will overlap with the occlusion pass.\n"),
+	TEXT("SkyAtmosphere on async compute (default: false). When running on the async pipe, SkyAtmosphere LUT generation will overlap with the occlusion pass.\n When enabled, opaque shadow map sampling on atmosphere will be disabled to avoid race conditions because these textures can be created in an overlaping parallel thread."),
 	ECVF_RenderThreadSafe);
 
 
@@ -312,9 +312,14 @@ auto GetDefaultSkyAtmosphereRealTimeReflectionLUTUniformBuffer()
 #define KM_TO_CM  100000.0f
 #define CM_TO_KM  (1.0f / KM_TO_CM)
 
+static bool GetSkyAtmosphereAsyncComputeEnabled()
+{
+	return CVarSkyAtmosphereAsyncCompute.GetValueOnRenderThread() != 0;
+}
+
 ESkyAtmospherePassLocation GetSkyAtmospherePassLocation()
 {
-	if (CVarSkyAtmosphereAsyncCompute.GetValueOnAnyThread() != 0)
+	if (GetSkyAtmosphereAsyncComputeEnabled())
 	{
 		return ESkyAtmospherePassLocation::BeforePrePass;
 	}
@@ -387,7 +392,16 @@ void GetSkyAtmosphereLightsUniformBuffers(
 {
 	FVolumeShadowingShaderParametersGlobal0& LightShadowShaderParams0 = *GraphBuilder.AllocParameters<FVolumeShadowingShaderParametersGlobal0>();
 	FVolumeShadowingShaderParametersGlobal1& LightShadowShaderParams1 = *GraphBuilder.AllocParameters<FVolumeShadowingShaderParametersGlobal1>();
-	if (bShouldSampleOpaqueShadow && LightShadowData.LightVolumetricShadowSceneinfo0)
+
+	bool bSampleOpaqueShadow = bShouldSampleOpaqueShadow;
+	if (GetSkyAtmosphereAsyncComputeEnabled())
+	{
+		// When enabling async compute, we can run RenderSkyAtmosphereLookupTables in parallel with FilterDynamicShadows, which allocates the depth buffers.
+		// So when running sky atmosphere LUT generation on the async compute pipe, we disable opaque shadow sampling.
+		bSampleOpaqueShadow = false;
+	}
+
+	if (bSampleOpaqueShadow && LightShadowData.LightVolumetricShadowSceneinfo0)
 	{
 		SetVolumeShadowingShaderParameters(GraphBuilder, LightShadowShaderParams0, ViewInfo,
 			LightShadowData.LightVolumetricShadowSceneinfo0, LightShadowData.ProjectedShadowInfo0);
@@ -396,7 +410,7 @@ void GetSkyAtmosphereLightsUniformBuffers(
 	{
 		SetVolumeShadowingDefaultShaderParameters(GraphBuilder, LightShadowShaderParams0);
 	}
-	if (bShouldSampleOpaqueShadow && LightShadowData.LightVolumetricShadowSceneinfo1)
+	if (bSampleOpaqueShadow && LightShadowData.LightVolumetricShadowSceneinfo1)
 	{
 		SetVolumeShadowingShaderParameters(GraphBuilder, LightShadowShaderParams1, ViewInfo,
 			LightShadowData.LightVolumetricShadowSceneinfo1, LightShadowData.ProjectedShadowInfo1);
