[Nanite] Fix a bug where World Position Offset will be disabled at an incorrect distance for any material with WPO but without pixel programmable features (e.g. Masked, PDO) on any mesh component that has Nanite Pixel Programmable set greater than zero, but less than its World Position Offset Disable Distance.

#Synced-CL 49347449

diff --git a/Engine/Source/Runtime/Renderer/Private/PrimitiveSceneInfo.cpp b/Engine/Source/Runtime/Renderer/Private/PrimitiveSceneInfo.cpp
--- a/Engine/Source/Runtime/Renderer/Private/PrimitiveSceneInfo.cpp
+++ b/Engine/Source/Runtime/Renderer/Private/PrimitiveSceneInfo.cpp
@@ -906,6 +906,12 @@ void BuildNaniteMaterialBins(FScene* Scene, FPrimitiveSceneInfo* PrimitiveSceneI
 		// Pre-allocate the max possible material slots for the slot array here, before contexts are applied serially.
 		const int32 NumMaterialSections = NaniteProxy->GetMaterialSections().Num();
 
+		float WPODistance;
+		const bool bHasWPODistance = NaniteProxy->GetInstanceWorldPositionOffsetDisableDistance(WPODistance);
+		const float PixelProgrammableDistance = NaniteProxy->GetPixelProgrammableDistance();
+		const float DisplacementFadeOutSize = NaniteProxy->GetMaterialDisplacementFadeOutSize();
+		const bool bAnyPixelProgrammableMaterial = NaniteProxy->HasPixelProgrammableRaster();
+
 		TArray<Nanite::FSceneProxyBase::FMaterialSection>& NaniteMaterialSections = NaniteProxy->GetMaterialSections();
 		if (NaniteMaterialSections.Num() > 0)
 		{
@@ -966,17 +972,21 @@ void BuildNaniteMaterialBins(FScene* Scene, FPrimitiveSceneInfo* PrimitiveSceneI
 					RasterPipeline.DisplacementScaling = MaterialSection.DisplacementScaling;
 					RasterPipeline.DisplacementFadeRange = MaterialSection.DisplacementFadeRange;
 
-					float WPODistance;
-					RasterPipeline.bHasWPODistance =
-						RasterPipeline.bWPOEnabled &&
-						!MaterialSection.bAlwaysEvaluateWPO &&
-						NaniteProxy->GetInstanceWorldPositionOffsetDisableDistance(WPODistance);
-					RasterPipeline.bHasPixelDistance =
-						RasterPipeline.bPerPixelEval &&
-						NaniteProxy->GetPixelProgrammableDistance() > 0.0f;
-					RasterPipeline.bHasDisplacementFadeOut =
-						RasterPipeline.bDisplacementEnabled &&
-						NaniteProxy->GetMaterialDisplacementFadeOutSize() > 0.0f;
+					RasterPipeline.bHasWPODistance = RasterPipeline.bWPOEnabled && !MaterialSection.bAlwaysEvaluateWPO && bHasWPODistance;
+					RasterPipeline.bHasPixelDistance = RasterPipeline.bPerPixelEval && PixelProgrammableDistance > 0.0f;
+					RasterPipeline.bHasDisplacementFadeOut = RasterPipeline.bDisplacementEnabled && DisplacementFadeOutSize > 0.0f;
+
+					// If the proxy has both pixel programmable distance and WPO disable distance, and the pixel distance is shorter,
+					// then treat the pipeline as if it has no WPO distance. Fallback raster is chosen for the whole cluster, which
+					// can have multiple associated raster bins. It's easier and more performant to branch over WPO than to push out
+					// pixel programmable distances to fix visual issues, so we opt for the former.
+					if (RasterPipeline.bHasWPODistance &&
+						PixelProgrammableDistance > 0.0f &&
+						PixelProgrammableDistance < WPODistance &&
+						bAnyPixelProgrammableMaterial)
+					{
+						RasterPipeline.bHasWPODistance = false;
+					}
 
 					FNaniteShadingPipeline& TriangleShadingPipeline = PipelinesCommand.TriangleShadingPipelines.Emplace_GetRef();
 					switch (MeshPassIndex)
