Added point light shadow rendering to multipass deferred shading for mobile renderer

#rnx
#Synced-CL 49269509

diff --git a/Engine/Shaders/Private/MobileDeferredShading.usf b/Engine/Shaders/Private/MobileDeferredShading.usf
--- a/Engine/Shaders/Private/MobileDeferredShading.usf
+++ b/Engine/Shaders/Private/MobileDeferredShading.usf
@@ -41,6 +41,14 @@
 #define RADIAL_LIGHT_TYPE 0
 #endif
 
+#ifndef SUPPORT_RADIAL_LIGHT_SHADOW
+#define SUPPORT_RADIAL_LIGHT_SHADOW 0
+#endif 
+
+#define SUPPORT_SPOTLIGHTS_SHADOW	SUPPORT_RADIAL_LIGHT_SHADOW && RADIAL_LIGHT_TYPE==LIGHT_TYPE_SPOT
+#define SUPPORT_POINTLIGHTS_SHADOW	SUPPORT_RADIAL_LIGHT_SHADOW && RADIAL_LIGHT_TYPE==LIGHT_TYPE_POINT
+
+
 float4x4 TranslatedWorldToLight;
 float2 LightFunctionParameters2;
 
@@ -300,6 +308,11 @@ FDeferredLightData InitDeferredLightFromLightParameters(FLightShaderParameters I
 
 void MobileRadialLightPS(
 	float4 InScreenPosition : TEXCOORD0,
+#if INSTANCED_STEREO && MOBILE_MULTI_VIEW
+	in nointerpolation uint EyeIndex : VIEW_ID,
+#elif MOBILE_MULTI_VIEW
+	in nointerpolation uint ViewId : SV_ViewID,
+#endif
 	float4 SVPos			: SV_POSITION,
 #if USE_GLES_FBF_DEFERRED
 	out HALF4_TYPE OutProxyAdditive : SV_Target0,
@@ -311,7 +324,16 @@ void MobileRadialLightPS(
 #endif
 )
 {
+#if INSTANCED_STEREO && MOBILE_MULTI_VIEW
+	ResolvedView = ResolveView(EyeIndex);
+#elif MOBILE_MULTI_VIEW
+	ResolvedView = ResolveView(ViewId);
+	const uint EyeIndex = ViewId;
+#else
 	ResolvedView = ResolveView();
+	const uint EyeIndex = 0;
+#endif
+
 
 	float2 ScreenUV = InScreenPosition.xy / InScreenPosition.w * View.ScreenPositionScaleBias.xy + View.ScreenPositionScaleBias.wz;
 	FGBufferData GBuffer = MobileFetchAndDecodeGBuffer(ScreenUV, SVPos.xy);
@@ -382,8 +404,13 @@ void MobileRadialLightPS(
 		GBuffer.DiffuseColor * LightData.DiffuseScale, 
 		GBuffer.SpecularColor * LightData.SpecularScale,
 		GBuffer.Roughness, SimpleLightData);
+#else
+#if SUPPORT_POINTLIGHTS_SHADOW
+	half4 LightAttenuation = SampleScreenSpaceShadowMask(ScreenUV, EyeIndex);
 #else
 	half4 LightAttenuation = half4(1, 1, Shadow, Shadow);
+#endif
+
 	float SurfaceShadow = 0;
 	const float Dither = InterleavedGradientNoise(SVPos.xy, View.StateFrameIndexMod8);
 	DirectLighting = AccumulateDynamicLighting(TranslatedWorldPosition, CameraVector, GBuffer, 1, LightData, LightAttenuation, Dither, uint2(0, 0), SurfaceShadow);
diff --git a/Engine/Source/Runtime/Renderer/Private/LightSceneInfo.cpp b/Engine/Source/Runtime/Renderer/Private/LightSceneInfo.cpp
--- a/Engine/Source/Runtime/Renderer/Private/LightSceneInfo.cpp
+++ b/Engine/Source/Runtime/Renderer/Private/LightSceneInfo.cpp
@@ -311,16 +311,33 @@ FLightPrimitiveInteraction* FLightSceneInfo::GetDynamicInteractionStaticPrimitiv
 	return DynamicInteractionStaticPrimitiveList;
 }
 
-bool FLightSceneInfo::SetupMobileMovableLocalLightShadowParameters(const FViewInfo& View, TConstArrayView<FVisibleLightInfo> VisibleLightInfos, FMobileMovableLocalLightShadowParameters& MobileMovableLocalLightShadowParameters) const
+bool FLightSceneInfo::IsAValidShadowPointLight(const FViewInfo& View, const TArray<FVisibleLightInfo, SceneRenderingAllocator>& VisibleLightInfos, EShaderPlatform ShaderPlatform) const
+{
+	if (Proxy->GetLightType() != LightType_Point)
+	{
+		return false;
+	}
+
+	if (!FReadOnlyCVARCache::EnablePointLightShadows(ShaderPlatform))
+	{
+		return false;
+	}
+
+	bool bShouldCastPointShadow = View.Family->EngineShowFlags.DynamicShadows
+		&& GetShadowQuality() > 0
+		&& VisibleLightInfos[GetPersistentIndex()].AllProjectedShadows.Num() > 0
+		&& VisibleLightInfos[GetPersistentIndex()].AllProjectedShadows.Last()->bAllocated;
+
+	return bShouldCastPointShadow;
+}
+
+bool FLightSceneInfo::SetupMobileMovableLocalSpotLightShadowParameters(const FViewInfo& View, TConstArrayView<FVisibleLightInfo> VisibleLightInfos, FMobileMovableLocalLightShadowParameters& MobileMovableLocalLightShadowParameters) const
 {
 	QUICK_SCOPE_CYCLE_COUNTER(STAT_FLightSceneProxy_SetupMobileMovableLocalLightShadowParameters);
 
 	check(Proxy->IsLocalLight());
 	const uint8 LightType = Proxy->GetLightType();
 
-	FLightRenderParameters LightParameters;
-	Proxy->GetLightShaderParameters(LightParameters);
-
 	bool bShouldCastShadow = View.Family->EngineShowFlags.DynamicShadows
 		&& GetShadowQuality() > 0
 		&& LightType == LightType_Spot
@@ -341,6 +358,9 @@ bool FLightSceneInfo::SetupMobileMovableLocalLightShadowParameters(const FViewIn
 		MobileMovableLocalLightShadowParameters.SpotLightShadowSharpenAndFadeFractionAndReceiverDepthBiasAndSoftTransitionScale = FVector4f(Proxy->GetShadowSharpen() * 7.0f + 1.0f, ShadowFadeFraction, ProjectedShadowInfo->GetShaderReceiverDepthBias(), SoftTransitionScale);
 
 		const FMatrix WorldToShadowMatrix = ProjectedShadowInfo->GetWorldToShadowMatrix(MobileMovableLocalLightShadowParameters.SpotLightShadowmapMinMax);
+		FLightRenderParameters LightParameters;
+		Proxy->GetLightShaderParameters(LightParameters);
+
 		MobileMovableLocalLightShadowParameters.SpotLightShadowWorldToShadowMatrix = FMatrix44f(FTranslationMatrix(LightParameters.WorldPosition) * WorldToShadowMatrix);
 
 		const FIntPoint ShadowBufferResolution = ProjectedShadowInfo->GetShadowBufferResolution();
diff --git a/Engine/Source/Runtime/Renderer/Private/LightSceneInfo.h b/Engine/Source/Runtime/Renderer/Private/LightSceneInfo.h
--- a/Engine/Source/Runtime/Renderer/Private/LightSceneInfo.h
+++ b/Engine/Source/Runtime/Renderer/Private/LightSceneInfo.h
@@ -337,7 +337,9 @@ public:
 		return (uint32)LightSceneInfo->Id;
 	}
 
-	bool SetupMobileMovableLocalLightShadowParameters(const FViewInfo& View, TConstArrayView<FVisibleLightInfo> VisibleLightInfos, FMobileMovableLocalLightShadowParameters& MobileMovableLocalLightShadowParameters) const;
+	bool IsAValidShadowPointLight(const FViewInfo& View, const TArray<FVisibleLightInfo, SceneRenderingAllocator>& VisibleLightInfos, EShaderPlatform ShaderPlatform) const;
+
+	bool SetupMobileMovableLocalSpotLightShadowParameters(const FViewInfo& View, TConstArrayView<FVisibleLightInfo> VisibleLightInfos, FMobileMovableLocalLightShadowParameters& MobileMovableLocalLightShadowParameters) const;
 
 	bool ShouldRecordShadowSubjectsForMobile() const;
 
diff --git a/Engine/Source/Runtime/Renderer/Private/MobileDeferredShadingPass.cpp b/Engine/Source/Runtime/Renderer/Private/MobileDeferredShadingPass.cpp
--- a/Engine/Source/Runtime/Renderer/Private/MobileDeferredShadingPass.cpp
+++ b/Engine/Source/Runtime/Renderer/Private/MobileDeferredShadingPass.cpp
@@ -17,7 +17,7 @@
 #include "MobileSSR.h"
 #include "DistanceFieldLightingShared.h"
 #include "DistanceFieldAmbientOcclusion.h"
-#include "DistanceFieldLightingShared.h"
+#include "ClearQuad.h"
 
 DECLARE_GPU_STAT(DeferredShading);
 
@@ -214,13 +214,13 @@ public:
 	DECLARE_SHADER_TYPE(FMobileRadialLightFunctionPS,Material);
 	SHADER_USE_PARAMETER_STRUCT_WITH_LEGACY_BASE(FMobileRadialLightFunctionPS, FMaterialShader)
 
-	class FEnableShadingModelSupport: SHADER_PERMUTATION_BOOL("ENABLE_SHADINGMODEL_SUPPORT_MOBILE_DEFERRED");
-	class FRadialLightTypeDim		: SHADER_PERMUTATION_RANGE_INT("RADIAL_LIGHT_TYPE", LIGHT_TYPE_POINT, LIGHT_TYPE_RECT);
-	class FIESProfileDim			: SHADER_PERMUTATION_BOOL("USE_IES_PROFILE");
-	class FSpotLightShadowDim		: SHADER_PERMUTATION_BOOL("SUPPORT_SPOTLIGHTS_SHADOW");
-	class FCapsuleLightDim			: SHADER_PERMUTATION_BOOL("MOBILE_SHADING_PATH_SUPPORT_CAPSULE_LIGHT");
-	class FSimpleLightDim			: SHADER_PERMUTATION_BOOL("SIMPLE_LIGHT");
-	using FPermutationDomain = TShaderPermutationDomain<FEnableShadingModelSupport, FRadialLightTypeDim, FIESProfileDim, FSpotLightShadowDim, FCapsuleLightDim, FSimpleLightDim>;
+	class FEnableShadingModelSupport		: SHADER_PERMUTATION_BOOL("ENABLE_SHADINGMODEL_SUPPORT_MOBILE_DEFERRED");
+	class FRadialLightTypeDim				: SHADER_PERMUTATION_RANGE_INT("RADIAL_LIGHT_TYPE", LIGHT_TYPE_POINT, LIGHT_TYPE_RECT);
+	class FIESProfileDim					: SHADER_PERMUTATION_BOOL("USE_IES_PROFILE");
+	class FSupportRadialLightShadowDim		: SHADER_PERMUTATION_BOOL("SUPPORT_RADIAL_LIGHT_SHADOW");
+	class FCapsuleLightDim					: SHADER_PERMUTATION_BOOL("MOBILE_SHADING_PATH_SUPPORT_CAPSULE_LIGHT");
+	class FSimpleLightDim					: SHADER_PERMUTATION_BOOL("SIMPLE_LIGHT");
+	using FPermutationDomain = TShaderPermutationDomain<FEnableShadingModelSupport, FRadialLightTypeDim, FIESProfileDim, FSupportRadialLightShadowDim, FCapsuleLightDim, FSimpleLightDim>;
 
 	BEGIN_SHADER_PARAMETER_STRUCT(FParameters, )
 		SHADER_PARAMETER_STRUCT(FLightShaderParameters, Light)
@@ -255,9 +255,16 @@ public:
 
 	static FPermutationDomain RemapPermutationVector(FPermutationDomain PermutationVector, EShaderPlatform Platform)
 	{
-		if (!IsMobileMovableSpotlightShadowsEnabled(Platform))
+		if (PermutationVector.Get<FSupportRadialLightShadowDim>() &&
+			((IsMobileMovableSpotlightShadowsEnabled(Platform) && PermutationVector.Get<FRadialLightTypeDim>() == LIGHT_TYPE_SPOT) ||
+			(MobileUsesShadowMaskTexture(Platform) && FReadOnlyCVARCache::EnablePointLightShadows(Platform) && PermutationVector.Get<FRadialLightTypeDim>() == LIGHT_TYPE_POINT))
+			)
+		{
+			PermutationVector.Set<FSupportRadialLightShadowDim>(true);
+		}
+		else
 		{
-			PermutationVector.Set<FSpotLightShadowDim>(false);
+			PermutationVector.Set<FSupportRadialLightShadowDim>(false);
 		}
 
 		if (!MobileUsesGBufferCustomData(Platform))
@@ -270,7 +277,7 @@ public:
 			PermutationVector.Set<FEnableShadingModelSupport>(false);
 			PermutationVector.Set<FRadialLightTypeDim>(LIGHT_TYPE_POINT);
 			PermutationVector.Set<FIESProfileDim>(false);
-			PermutationVector.Set<FSpotLightShadowDim>(false);
+			PermutationVector.Set<FSupportRadialLightShadowDim>(false);
 			PermutationVector.Set<FCapsuleLightDim>(false);
 		}
 
@@ -287,12 +294,15 @@ public:
 		FMaterialShader::ModifyCompilationEnvironment(Parameters, OutEnvironment);
 		OutEnvironment.SetDefine(TEXT("USE_LIGHT_FUNCTION"), Parameters.MaterialParameters.bIsDefaultMaterial ? 0 : 1);
 		OutEnvironment.SetDefine(TEXT("MATERIAL_SHADER"), 1);
-		OutEnvironment.SetDefine(TEXT("ENABLE_SHADOWMASKTEXTURE"), 0);
 		OutEnvironment.SetDefine(TEXT("ENABLE_CLUSTERED_LIGHTS"), 0);
 		OutEnvironment.SetDefine(TEXT("IS_MOBILE_DEFERREDSHADING_SUBPASS"), 1u);
 
 		const bool bMobileForceDepthRead = MobileUsesFullDepthPrepass(Parameters.Platform);
 		OutEnvironment.SetDefine(TEXT("FORCE_DEPTH_TEXTURE_READS"), bMobileForceDepthRead ? 1u : 0u);
+
+		FPermutationDomain PermutationVector(Parameters.PermutationId);
+		bool bPointLightCastShadow = (PermutationVector.Get<FSupportRadialLightShadowDim>() && PermutationVector.Get<FRadialLightTypeDim>() == LIGHT_TYPE_POINT);
+		OutEnvironment.SetDefine(TEXT("ENABLE_SHADOWMASKTEXTURE"), bPointLightCastShadow ? 1 : 0);
 	}
 
 	void SetParameters(FRHIBatchedShaderParameters& BatchedParameters, const FViewInfo& View, const FMaterialRenderProxy* Proxy, const FMaterial& Material)
@@ -916,7 +926,8 @@ static void RenderLocalLight(
 	}
 
 	FMobileRadialLightFunctionPS::FParameters PassParameters;
-	const bool bShouldCastShadow = LightSceneInfo.SetupMobileMovableLocalLightShadowParameters(View, VisibleLightInfos, PassParameters.MobileMovableLocalLightShadow);
+	bool bShouldCastLocalLightShadow = LightSceneInfo.SetupMobileMovableLocalSpotLightShadowParameters(View, VisibleLightInfos, PassParameters.MobileMovableLocalLightShadow);
+	bShouldCastLocalLightShadow |= LightSceneInfo.IsAValidShadowPointLight(View, VisibleLightInfos, View.GetShaderPlatform());
 
 	PassParameters.Light = GetDeferredLightParameters(View, LightSceneInfo).LightParameters;
 	const float TanOuterAngle = bIsSpotLight ? FMath::Tan(LightSceneInfo.Proxy->GetOuterConeAngle()) : 1.0f;
@@ -953,7 +964,7 @@ static void RenderLocalLight(
 			PermutationVector.Set<FMobileRadialLightFunctionPS::FEnableShadingModelSupport>(bEnableShadingModelSupport);
 			PermutationVector.Set<FMobileRadialLightFunctionPS::FRadialLightTypeDim>(LightType);
 			PermutationVector.Set<FMobileRadialLightFunctionPS::FIESProfileDim>(bUseIESTexture);
-			PermutationVector.Set<FMobileRadialLightFunctionPS::FSpotLightShadowDim>(bShouldCastShadow);
+			PermutationVector.Set<FMobileRadialLightFunctionPS::FSupportRadialLightShadowDim>(bShouldCastLocalLightShadow);
 			PermutationVector.Set<FMobileRadialLightFunctionPS::FCapsuleLightDim>(bCapsuleLight);
 			FCachedLightMaterial LightMaterial;
 			TShaderRef<FMobileRadialLightFunctionPS> PixelShader;
@@ -1055,6 +1066,23 @@ static void RenderSimpleLights(
 	}
 }
 
+void MobileDeferredRenderLocalLight(
+	FRHICommandList& RHICmdList,
+	const FScene& Scene,
+	const FViewInfo& View,
+	const FLightSceneInfo& LightSceneInfo,
+	const TArray<FVisibleLightInfo, SceneRenderingAllocator>& VisibleLightInfos)
+{
+	// Default material for light rendering
+	FCachedLightMaterial DefaultMaterial;
+	DefaultMaterial.MaterialProxy = UMaterial::GetDefaultMaterial(MD_LightFunction)->GetRenderProxy();
+	DefaultMaterial.Material = DefaultMaterial.MaterialProxy->GetMaterialNoFallback(ERHIFeatureLevel::ES3_1);
+	check(DefaultMaterial.Material);
+	
+	RHICmdList.SetViewport(View.ViewRect.Min.X, View.ViewRect.Min.Y, 0.0f, View.ViewRect.Max.X, View.ViewRect.Max.Y, 1.0f);
+	RenderLocalLight(RHICmdList, Scene, View, LightSceneInfo, DefaultMaterial, VisibleLightInfos);
+}
+
 void MobileDeferredShadingPass(
 	FRHICommandList& RHICmdList,
 	int32 ViewIndex,
@@ -1079,7 +1107,8 @@ void MobileDeferredShadingPass(
 
 	int NumDirLights = RenderDirectionalLights(RHICmdList, Scene, View, DefaultMaterial, MobileSSRQuality, DynamicBentNormalAOTexture);
 	
-	const bool bMobileUseClusteredDeferredShading = UseClusteredDeferredShading(View.GetShaderPlatform()) && NumDirLights > 0;
+	EShaderPlatform ShaderPlatform = View.GetShaderPlatform();
+	const bool bMobileUseClusteredDeferredShading = UseClusteredDeferredShading(ShaderPlatform) && NumDirLights > 0;
 	if (!bMobileUseClusteredDeferredShading)
 	{
 		// Render non-clustered simple lights
@@ -1108,6 +1137,14 @@ void MobileDeferredShadingPass(
 	{
 		const FSortedLightSceneInfo& SortedLight = SortedLightSet.SortedLights[LightIdx];
 		const FLightSceneInfo& LightSceneInfo = *SortedLight.LightSceneInfo;
+
+		if (LightSceneInfo.IsAValidShadowPointLight(View, VisibleLightInfos, ShaderPlatform))
+		{
+			// Skip point light shadow as they 2 need RDG Passes (render point shadow to shadowscreenspace and renderlight).
+			// they'll be rendered after "the main pass"
+			continue;
+		}
+
 		RenderLocalLight(RHICmdList, Scene, View, LightSceneInfo, DefaultMaterial, VisibleLightInfos);
 	}
 }
@@ -1317,7 +1354,7 @@ void FDeferredMobileLightMaterialPSOCollector::CollectPSOInitializersLocal(
 			PermutationVector.Set<FMobileRadialLightFunctionPS::FEnableShadingModelSupport>(bEnableShadingModelSupport);
 			PermutationVector.Set<FMobileRadialLightFunctionPS::FRadialLightTypeDim>(LightType);
 			PermutationVector.Set<FMobileRadialLightFunctionPS::FIESProfileDim>(bUseIESTexture);
-			PermutationVector.Set<FMobileRadialLightFunctionPS::FSpotLightShadowDim>(bShouldCastShadow);
+			PermutationVector.Set<FMobileRadialLightFunctionPS::FSupportRadialLightShadowDim>(bShouldCastShadow);
 			PermutationVector.Set<FMobileRadialLightFunctionPS::FCapsuleLightDim>(bCapsuleLight);
 			FMaterialShaderTypes ShaderTypesToGet;
 			ShaderTypesToGet.AddShaderType<FMobileRadialLightFunctionPS>(PermutationVector.ToDimensionValueId());
diff --git a/Engine/Source/Runtime/Renderer/Private/MobileDeferredShadingPass.h b/Engine/Source/Runtime/Renderer/Private/MobileDeferredShadingPass.h
--- a/Engine/Source/Runtime/Renderer/Private/MobileDeferredShadingPass.h
+++ b/Engine/Source/Runtime/Renderer/Private/MobileDeferredShadingPass.h
@@ -17,12 +17,19 @@ struct FSortedLightSetSceneInfo;
 
 extern int32 GMobileUseClusteredDeferredShading;
 
+void MobileDeferredRenderLocalLight(
+	FRHICommandList& RHICmdList,
+	const FScene& Scene,
+	const FViewInfo& View,
+	const FLightSceneInfo& LightSceneInfo,
+	const TArray<FVisibleLightInfo, SceneRenderingAllocator>& VisibleLightInfos);
+
 void MobileDeferredShadingPass(
 	FRHICommandList& RHICmdList,
 	int32 ViewIndex,
 	int32 NumViews,
 	const FViewInfo& View,
-	const FScene& Scene, 
+	const FScene& Scene,
 	const FSortedLightSetSceneInfo& SortedLightSet,
 	const TArray<FVisibleLightInfo, SceneRenderingAllocator>& VisibleLightInfos,
 	EMobileSSRQuality MobileSSRQuality,
diff --git a/Engine/Source/Runtime/Renderer/Private/MobileShadingRenderer.cpp b/Engine/Source/Runtime/Renderer/Private/MobileShadingRenderer.cpp
--- a/Engine/Source/Runtime/Renderer/Private/MobileShadingRenderer.cpp
+++ b/Engine/Source/Runtime/Renderer/Private/MobileShadingRenderer.cpp
@@ -2641,16 +2641,59 @@ void FMobileSceneRenderer::RenderDeferredMultiPass(FRDGBuilder& GraphBuilder, FS
 
 		const EMobileSSRQuality MobileSSRQuality = ActiveMobileSSRQuality(View, bShouldRenderVelocities);
 
-		GraphBuilder.AddPass(
-			RDG_EVENT_NAME("LightingAndFog"),
-			PassParameters,
-			ERDGPassFlags::Raster,
-			[this, PassParameters, ViewContext, MobileSSRQuality, &SortedLightSet](FRHICommandList& RHICmdList)
+		// Determine if a shadow point light is present.
+		TArray<int32, TInlineAllocator<8>> ShadowPointLightIndexes;
+		GatherShadowPointLights(View, SortedLightSet, ShadowPointLightIndexes);
+		const bool bSeparateShadowPointLights = (ShadowPointLightIndexes.Num() > 0);
+
+		{
+			RDG_EVENT_SCOPE(GraphBuilder, "LightingAndFog", ViewIndex);
+
+			GraphBuilder.AddPass(
+				RDG_EVENT_NAME("Lighting"),
+				PassParameters,
+				ERDGPassFlags::Raster,
+				[this, PassParameters, ViewContext, MobileSSRQuality, &SortedLightSet](FRHICommandList& RHICmdList)
+				{
+					FViewInfo& View = *ViewContext.ViewInfo;
+					MobileDeferredShadingPass(RHICmdList, ViewContext.ViewIndex, Views.Num(), View, *Scene, SortedLightSet, VisibleLightInfos, MobileSSRQuality, PassParameters->BentNormalAOTexture);
+				});
+
+
+			if (bSeparateShadowPointLights)
 			{
-				FViewInfo& View = *ViewContext.ViewInfo;
-				MobileDeferredShadingPass(RHICmdList, ViewContext.ViewIndex, Views.Num(), View, *Scene, SortedLightSet, VisibleLightInfos, MobileSSRQuality, PassParameters->BentNormalAOTexture);
-				RenderFog(RHICmdList, View);
-			});
+				for (int32 Index = 0; Index < ShadowPointLightIndexes.Num(); ++Index)
+				{
+					const FSortedLightSceneInfo& SortedLight = SortedLightSet.SortedLights[ShadowPointLightIndexes[Index]];
+					const FLightSceneInfo& LightSceneInfo = *SortedLight.LightSceneInfo;
+
+					RDG_EVENT_SCOPE(GraphBuilder, "PointLightShadow_%s", LightSceneInfo.Proxy->GetOwnerNameOrLabel());
+					{
+						ClearMobileScreenSpaceShadows(GraphBuilder, View, LightSceneInfo.Proxy, BasePassTextures.ScreenSpaceShadowMaskTexture);
+						RenderShadowProjections(GraphBuilder, SceneTextures, BasePassTextures.ScreenSpaceShadowMaskTexture, nullptr, &LightSceneInfo, false);
+
+						GraphBuilder.AddPass(
+							RDG_EVENT_NAME("PointLightShadowPass"),
+							PassParameters,
+							ERDGPassFlags::Raster,
+							[this, &View, LightSceneInfo](FRHICommandList& RHICmdList)
+							{
+								MobileDeferredRenderLocalLight(RHICmdList, *Scene, View, LightSceneInfo, VisibleLightInfos);
+							});
+					}
+				}
+			}
+
+			GraphBuilder.AddPass(
+				RDG_EVENT_NAME("Fog"),
+				PassParameters,
+				ERDGPassFlags::Raster,
+				[this, &View](FRHICommandList& RHICmdList)
+				{
+					RenderFog(RHICmdList, View);
+				});
+		}
+
 
 		if (bShouldRenderSingleLayerWater 
 			&& View.bHasSingleLayerWaterMaterial
diff --git a/Engine/Source/Runtime/Renderer/Private/SceneRendering.h b/Engine/Source/Runtime/Renderer/Private/SceneRendering.h
--- a/Engine/Source/Runtime/Renderer/Private/SceneRendering.h
+++ b/Engine/Source/Runtime/Renderer/Private/SceneRendering.h
@@ -2945,6 +2945,7 @@ protected:
 
 	FRDGTextureRef RenderAmbientOcclusion(FRDGBuilder& GraphBuilder, const FSceneTextures& SceneTextures);
 
+	void ClearMobileScreenSpaceShadows(FRDGBuilder& GraphBuilder, const FViewInfo& View, const FLightSceneProxy* LightSceneProxy, FRDGTextureRef ScreenShadowMaskTexture);
 	void RenderMobileScreenSpaceShadows(FRDGBuilder& GraphBuilder, FIntPoint Extent, FRDGTextureRef ScreenShadowMaskTexture);
 	FRDGTextureRef RenderMobileShadowProjections(FRDGBuilder& GraphBuilder, FIntPoint Extent, bool bIsMobileMultiView);
 
@@ -2952,6 +2953,8 @@ protected:
 	FRenderTargetBindingSlots InitRenderTargetBindings_Forward(FRDGTextureRef ViewFamilyTexture, FSceneTextures& SceneTextures);
 	FColorTargets GetColorTargets_Deferred(FSceneTextures& SceneTextures);
 
+	void GatherShadowPointLights(const FViewInfo& View, const FSortedLightSetSceneInfo& SortedLightSet, TArray<int32, TInlineAllocator<8>>& OutShadowPointLightIndexes);
+
 private:
 	const bool bGammaSpace;
 	const bool bDeferredShading;
diff --git a/Engine/Source/Runtime/Renderer/Private/ShadowRendering.cpp b/Engine/Source/Runtime/Renderer/Private/ShadowRendering.cpp
--- a/Engine/Source/Runtime/Renderer/Private/ShadowRendering.cpp
+++ b/Engine/Source/Runtime/Renderer/Private/ShadowRendering.cpp
@@ -2465,6 +2465,31 @@ void FMobileSceneRenderer::RenderModulatedShadowProjections(FRHICommandList& RHI
 	}
 }
 
+void FMobileSceneRenderer::ClearMobileScreenSpaceShadows(FRDGBuilder& GraphBuilder, const FViewInfo& View, const FLightSceneProxy* LightSceneProxy, FRDGTextureRef ScreenShadowMaskTexture)
+{
+	FRenderTargetParameters* ClearPassParameters = GraphBuilder.AllocParameters<FRenderTargetParameters>();
+	ClearPassParameters->RenderTargets[0] = FRenderTargetBinding(ScreenShadowMaskTexture, ERenderTargetLoadAction::ENoAction);
+
+	GraphBuilder.AddPass(
+		RDG_EVENT_NAME("ScreenSpaceShadowMaskClearQuad"),
+		ClearPassParameters,
+		ERDGPassFlags::Raster,
+		[&View, LightSceneProxy](FRHICommandList& RHICmdList)
+		{
+			FIntRect ScissorRect;
+			if (!LightSceneProxy->GetScissorRect(ScissorRect, View, View.ViewRect))
+			{
+				ScissorRect = View.ViewRect;
+			}
+
+			if (ScissorRect.Min.X < ScissorRect.Max.X && ScissorRect.Min.Y < ScissorRect.Max.Y)
+			{
+				RHICmdList.SetViewport(ScissorRect.Min.X, ScissorRect.Min.Y, 0.0f, ScissorRect.Max.X, ScissorRect.Max.Y, 1.0f);
+				DrawClearQuad(RHICmdList, true, FLinearColor(1, 1, 1, 1), false, 0, false, 0);
+			}
+		});
+}
+
 void FMobileSceneRenderer::RenderMobileScreenSpaceShadows(FRDGBuilder& GraphBuilder, FIntPoint Extent, FRDGTextureRef ScreenShadowMaskTexture)
 {
 	static auto* ContactShadowsCVar = IConsoleManager::Get().FindTConsoleVariableDataInt(TEXT("r.ContactShadows"));
@@ -2514,10 +2539,10 @@ FRDGTextureRef FMobileSceneRenderer::RenderMobileShadowProjections(FRDGBuilder&
 		const FVisibleLightInfo& VisibleLightInfo = VisibleLightInfos[LightSceneInfo->Id];
 		const FLightSceneProxy* LightSceneProxy = LightSceneInfo->Proxy;
 
-		const bool bProjectingForForwardShading = true;
+		const bool bProjectingForForwardShading = !(IsMobileDeferredShadingEnabled(ShaderPlatform) && LightSceneProxy->GetLightType() == LightType_Point);
 
-		// Local light shadows don't render to shadow mask texture on mobile deferred
-		if (LightSceneProxy->GetLightType() == LightType_Directional || !IsMobileDeferredShadingEnabled(ShaderPlatform))
+		// Other than point light, local light shadows don't render to shadow mask texture on mobile deferred
+		if (LightSceneProxy->GetLightType() == LightType_Directional || LightSceneProxy->GetLightType() == LightType_Point || !IsMobileDeferredShadingEnabled(ShaderPlatform))
 		{
 			RenderShadowProjections(GraphBuilder, SceneTextures,
 				ScreenShadowMaskTexture,
diff --git a/Engine/Source/Runtime/Renderer/Private/ShadowSetupMobile.cpp b/Engine/Source/Runtime/Renderer/Private/ShadowSetupMobile.cpp
--- a/Engine/Source/Runtime/Renderer/Private/ShadowSetupMobile.cpp
+++ b/Engine/Source/Runtime/Renderer/Private/ShadowSetupMobile.cpp
@@ -617,3 +617,26 @@ void FMobileSceneRenderer::BuildCSMVisibilityState(FLightSceneInfo* LightSceneIn
 		}
 	}
 }
+
+void FMobileSceneRenderer::GatherShadowPointLights(
+	const FViewInfo& View,
+	const FSortedLightSetSceneInfo& SortedLightSet,
+	TArray<int32, TInlineAllocator<8>>& OutShadowPointLightIndexes)
+{
+	if (!FReadOnlyCVARCache::EnablePointLightShadows(ShaderPlatform))
+	{
+		return;
+	}
+
+	int32 NumLights = SortedLightSet.SortedLights.Num();
+	for (int32 LightIdx = SortedLightSet.UnbatchedLightStart; LightIdx < NumLights; ++LightIdx)
+	{
+		const FSortedLightSceneInfo& SortedLight = SortedLightSet.SortedLights[LightIdx];
+		const FLightSceneInfo& LightSceneInfo = *SortedLight.LightSceneInfo;
+
+		if (LightSceneInfo.IsAValidShadowPointLight(View, VisibleLightInfos, ShaderPlatform))
+		{
+			OutShadowPointLightIndexes.Add(LightIdx);
+		}
+	}
+}
