Remove HairComplexTransmittance permutation when Substrate is enabled, as it is no longer needed with classification.

#Synced-CL 49403711

diff --git a/Engine/Source/Runtime/Renderer/Private/LightRendering.cpp b/Engine/Source/Runtime/Renderer/Private/LightRendering.cpp
--- a/Engine/Source/Runtime/Renderer/Private/LightRendering.cpp
+++ b/Engine/Source/Runtime/Renderer/Private/LightRendering.cpp
@@ -1011,16 +1011,27 @@ class FDeferredLightPS : public FGlobalShader
 		{
 			return false;
 		}
-		
-		// (Hair Lighting == 1) requires FHairComplexTransmittance 
-		if (PermutationVector.Get<FHairLighting>() == 1 && !PermutationVector.Get<FDeferredLightPS::FHairComplexTransmittance>())
+
+		// When using Substrate, complex transmittance permutation is not used
+		if (Substrate::IsSubstrateEnabled())
 		{
-			return false;
+			if (!PermutationVector.Get<FDeferredLightPS::FHairComplexTransmittance>())
+			{
+				return false;
+			}
 		}
-
-		if (PermutationVector.Get<FHairLighting>() == 0 && PermutationVector.Get<FDeferredLightPS::FHairComplexTransmittance>() && !SupportsHairComplexTransmittance(Parameters.Platform))
+		else
 		{
-			return false;
+			// (Hair Lighting == 1) requires FHairComplexTransmittance 
+			if (PermutationVector.Get<FHairLighting>() == 1 && !PermutationVector.Get<FDeferredLightPS::FHairComplexTransmittance>())
+			{
+				return false;
+			}
+
+			if (PermutationVector.Get<FHairLighting>() == 0 && PermutationVector.Get<FDeferredLightPS::FHairComplexTransmittance>() && !SupportsHairComplexTransmittance(Parameters.Platform))
+			{
+				return false;
+			}
 		}
 
 		if (Substrate::IsSubstrateEnabled() && PermutationVector.Get<FHairLighting>() == 1 && PermutationVector.Get<Substrate::FSubstrateTileType>() != ESubstrateTileType::EComplex)
@@ -1091,10 +1102,19 @@ class FDeferredLightPS : public FGlobalShader
 			PermutationVector.Set< FHairComplexTransmittance >(false);
 		}
 
-		// We must keep FAnistropicMaterials 0/1 permutations for tiles of Complex and ComplexSpecial complexity. That is because rect light or hair can disable anisotropy explicitely.
-		if (Substrate::IsSubstrateEnabled() && PermutationVector.Get<Substrate::FSubstrateTileType>() < ESubstrateTileType::EComplex)
+		if (Substrate::IsSubstrateEnabled())
 		{
-			PermutationVector.Set<FAnistropicMaterials>(false);		// No anisotropy for tiles of Simple and Single complexity
+			// We must keep FAnistropicMaterials 0/1 permutations for tiles of Complex and ComplexSpecial complexity. That is because rect light or hair can disable anisotropy explicitely.
+			if (PermutationVector.Get<Substrate::FSubstrateTileType>() < ESubstrateTileType::EComplex)
+			{
+				PermutationVector.Set<FAnistropicMaterials>(false);		// No anisotropy for tiles of Simple and Single complexity
+			}
+
+			// With Substrate, complex transmittance is always enabled (dynamic switch)
+			if (!PermutationVector.Get<FVisualizeCullingDim>())
+			{
+				PermutationVector.Set<FHairComplexTransmittance>(true);
+			}
 		}
 
 		return PermutationVector;
