Ensure Primitives RequestStaticMeshUpdate when Skylight changes to ensure the ReflectionCaptureUniformBuffer gets updated to reflect the change.

#Synced-CL 49329359

diff --git a/Engine/Source/Runtime/Renderer/Private/RendererScene.cpp b/Engine/Source/Runtime/Renderer/Private/RendererScene.cpp
--- a/Engine/Source/Runtime/Renderer/Private/RendererScene.cpp
+++ b/Engine/Source/Runtime/Renderer/Private/RendererScene.cpp
@@ -2430,7 +2430,8 @@ void FScene::SetSkyLight(FSkyLightSceneProxy* LightProxy)
 
 			// Use the most recently enabled skylight
 			Scene->SkyLight = LightProxy;
-
+			
+			Scene->ReflectionSceneData.bSkylightHasChanged = true;
 			Scene->InvalidatePathTracedOutput();
 		});
 }
@@ -2457,6 +2458,7 @@ void FScene::DisableSkyLight(FSkyLightSceneProxy* LightProxy)
 			Scene->SkyLight = NULL;
 		}
 
+		Scene->ReflectionSceneData.bSkylightHasChanged = true;
 		Scene->InvalidatePathTracedOutput();
 	});
 }
@@ -5120,14 +5122,16 @@ void UpdateReflectionSceneData(FScene* Scene)
 	{
 		QUICK_SCOPE_CYCLE_COUNTER(STAT_MarkAllPrimitivesForReflectionProxyUpdate);
 
-		if (Scene->ReflectionSceneData.bRegisteredReflectionCapturesHasChanged)
-		{
-			// Mobile needs to re-cache all mesh commands when scene capture data has changed
-			const bool bNeedsStaticMeshUpdate = GetFeatureLevelShadingPath(Scene->GetFeatureLevel()) == EShadingPath::Mobile;
 
+		// Mobile needs to re-cache all mesh commands when scene capture data has changed
+		bool bNeedsStaticMeshUpdate = GetFeatureLevelShadingPath(Scene->GetFeatureLevel()) == EShadingPath::Mobile;
+
+		if (ReflectionSceneData.bRegisteredReflectionCapturesHasChanged)
+		{
 			// Mark all primitives as needing an update
 			// Note: Only visible primitives will actually update their reflection proxy
-			for (int32 PrimitiveIndex = 0; PrimitiveIndex < Scene->Primitives.Num(); PrimitiveIndex++)
+			const int32 NumPrimitives = Scene->Primitives.Num();
+			for (int32 PrimitiveIndex = 0; PrimitiveIndex < NumPrimitives; PrimitiveIndex++)
 			{
 				FPrimitiveSceneInfo* Primitive = Scene->Primitives[PrimitiveIndex];
 				Primitive->RemoveCachedReflectionCaptures();
@@ -5139,7 +5143,23 @@ void UpdateReflectionSceneData(FScene* Scene)
 				}
 			}
 
-			Scene->ReflectionSceneData.bRegisteredReflectionCapturesHasChanged = false;
+			ReflectionSceneData.bRegisteredReflectionCapturesHasChanged = false;
+			bNeedsStaticMeshUpdate = false;
+		}
+		
+		if (ReflectionSceneData.bSkylightHasChanged)
+		{
+			if (bNeedsStaticMeshUpdate)
+			{
+				// Mark all primitives as needing an update
+				// Note: Only visible primitives will actually update their reflection proxy
+				const int32 NumPrimitives = Scene->Primitives.Num();
+				for (int32 PrimitiveIndex = 0; PrimitiveIndex < NumPrimitives; PrimitiveIndex++)
+				{
+					Scene->Primitives[PrimitiveIndex]->RequestStaticMeshUpdate();
+				}
+			}
+			ReflectionSceneData.bSkylightHasChanged = false;
 		}
 	}
 }
diff --git a/Engine/Source/Runtime/Renderer/Private/ScenePrivate.h b/Engine/Source/Runtime/Renderer/Private/ScenePrivate.h
--- a/Engine/Source/Runtime/Renderer/Private/ScenePrivate.h
+++ b/Engine/Source/Runtime/Renderer/Private/ScenePrivate.h
@@ -1762,6 +1762,9 @@ public:
 	 */
 	bool bRegisteredReflectionCapturesHasChanged;
 
+	/** Set to true for one frame whenever Scene's Skylight changes. Which allows one frame to RequestStaticMeshUpdate for the Scene's Primitives.*/
+	bool bSkylightHasChanged;
+	
 	/** True if AllocatedReflectionCaptureState has changed. Allows to update cached single capture id. */
 	bool AllocatedReflectionCaptureStateHasChanged;
 
@@ -1805,6 +1808,7 @@ public:
 
 	FReflectionEnvironmentSceneData(ERHIFeatureLevel::Type InFeatureLevel) :
 		bRegisteredReflectionCapturesHasChanged(true),
+		bSkylightHasChanged(false),
 		AllocatedReflectionCaptureStateHasChanged(false),
 		CubemapArray(InFeatureLevel),
 		MaxAllocatedReflectionCubemapsGameThread(0),
