[HWRT] Added extra logging to IsCachedRayTracingGeometryValid() to help investigate failing assert.

diff --git a/Engine/Source/Runtime/Renderer/Private/PrimitiveSceneInfo.cpp b/Engine/Source/Runtime/Renderer/Private/PrimitiveSceneInfo.cpp
--- a/Engine/Source/Runtime/Renderer/Private/PrimitiveSceneInfo.cpp
+++ b/Engine/Source/Runtime/Renderer/Private/PrimitiveSceneInfo.cpp
@@ -389,17 +389,48 @@ bool FPrimitiveSceneInfo::IsCachedRayTracingGeometryValid(int8 LODIndex, int32 I
 	{
 		// TODO: Doesn't take Nanite Ray Tracing into account
 
+		check(!CachedRayTracingGeometry->GetRequiresBuild() && !CachedRayTracingGeometry->HasPendingBuildRequest());
+
+#if DO_CHECK
+		auto LogErrorDebugInfo = [this, CachedRayTracingGeometry](uint32 RayTracingInstanceIndex)
+			{
+				FRHIRayTracingGeometry* CachedInstanceGeometryRHI = Scene->RayTracingScene.GetCachedInstanceGeometry(RayTracingInstanceIndex);
+				FRHIRayTracingGeometry* CachedRayTracingGeometryRHI = CachedRayTracingGeometry->GetRHI();
+				UE_LOG(LogRenderer, Error, TEXT("IsCachedRayTracingGeometryValid() Error Debug Info \
+					PrimitiveRayTracingFlags: %u \
+					bIsCachedRayTracingInstanceValid: %d \
+					bCachedInstanceHasGeometryRHI: %d, \
+					bCachedRayTracingGeometryHasRHI: %d, \
+					bCachedRayTracingGeometryIsValid: %d, \
+					bCachedRayTracingGeometryIsEvicted: %d, \
+					CachedRayTracingGeometryFlags: %u"),
+					Scene->PrimitiveRayTracingFlags[GetIndex()],
+					bIsCachedRayTracingInstanceValid ? 1 : 0,
+					CachedInstanceGeometryRHI != nullptr ? 1 : 0,
+					CachedRayTracingGeometryRHI != nullptr ? 1 : 0,
+					CachedRayTracingGeometry->IsValid() ? 1 : 0,
+					CachedRayTracingGeometry->IsEvicted() ? 1 : 0,
+					uint32(CachedRayTracingGeometry->GetGeometryState()));
+			};
+
 		if (InstanceData.RayTracingInstanceIndexMain != UINT32_MAX)
 		{
-			check(Scene->RayTracingScene.GetCachedInstanceGeometry(InstanceData.RayTracingInstanceIndexMain) == CachedRayTracingGeometry->GetRHI());
+			if (Scene->RayTracingScene.GetCachedInstanceGeometry(InstanceData.RayTracingInstanceIndexMain) != CachedRayTracingGeometry->GetRHI())
+			{
+				LogErrorDebugInfo(InstanceData.RayTracingInstanceIndexMain);
+				checkf(false, TEXT("Cached ray tracing instance has an unexpected geometry. Check log above."))
+			}
 		}
 
 		if (InstanceData.RayTracingInstanceIndexDecal != UINT32_MAX)
 		{
-			check(Scene->RayTracingScene.GetCachedInstanceGeometry(InstanceData.RayTracingInstanceIndexDecal) == CachedRayTracingGeometry->GetRHI());
+			if (Scene->RayTracingScene.GetCachedInstanceGeometry(InstanceData.RayTracingInstanceIndexDecal) != CachedRayTracingGeometry->GetRHI())
+			{
+				LogErrorDebugInfo(InstanceData.RayTracingInstanceIndexDecal);
+				checkf(false, TEXT("Cached ray tracing instance has an unexpected geometry. Check log above."))
+			}
 		}
-
-		check(!CachedRayTracingGeometry->GetRequiresBuild() && !CachedRayTracingGeometry->HasPendingBuildRequest());
+#endif // DO_CHECK
 
 		return CachedRayTracingGeometry->IsValid() && !CachedRayTracingGeometry->IsEvicted();
 	}
