[HWRT] Fix crash when LODToRender.GetRayTracedLOD() returns out-of-bounds index.

#Synced-CL 49373806

diff --git a/Engine/Source/Runtime/Renderer/Private/RayTracing/RayTracing.cpp b/Engine/Source/Runtime/Renderer/Private/RayTracing/RayTracing.cpp
--- a/Engine/Source/Runtime/Renderer/Private/RayTracing/RayTracing.cpp
+++ b/Engine/Source/Runtime/Renderer/Private/RayTracing/RayTracing.cpp
@@ -1504,15 +1504,17 @@ namespace RayTracing
 				// we should replace this condition with an assert instead
 				else if (!EnumHasAnyFlags(Flags, ERayTracingPrimitiveFlags::Skip))
 				{
+					const uint32 NumRayTracingLODs = SceneInfo->GetRayTracingLODDataNum();
+
 #if DO_CHECK
-					if (bUsingReferenceBasedResidency)
+					// TODO: Should have an assert here but disabled it due to UE-309710
+					if (bUsingReferenceBasedResidency && LODIndex < int8(NumRayTracingLODs))
 					{
 						const FPrimitiveSceneInfo::FRayTracingLODData& RTLODData = SceneInfo->GetRayTracingLODData(LODIndex);
 
 						for (const FPrimitiveSceneInfo::FRayTracingInstanceData& RTInstanceData : RTLODData.Instances)
 						{
-							// TODO: Should have an assert here but disabled it due to UE-112448
-							if (RTInstanceData.RayTracingGeometry != nullptr)
+							if (ensure(RTInstanceData.RayTracingGeometry != nullptr))
 							{
 								// It is not safe to directly call FRayTracingGeometryManager::IsGeometryGroupReferenced(...) here since other threads might be modifying it.
 								// Instead we gather the group handles to validate later.
@@ -1522,8 +1524,6 @@ namespace RayTracing
 					}
 #endif
 
-					const uint32 NumRayTracingLODs = SceneInfo->GetRayTracingLODDataNum();
-
 					// find highest LOD whose instances all have a valid RayTracingGeometry
 					for (; LODIndex < int8(NumRayTracingLODs); ++LODIndex)
 					{
