Fix missing shader permutation for light complexity view mode.

#Synced-CL 49556441

diff --git a/Engine/Source/Runtime/Renderer/Private/LightRendering.cpp b/Engine/Source/Runtime/Renderer/Private/LightRendering.cpp
--- a/Engine/Source/Runtime/Renderer/Private/LightRendering.cpp
+++ b/Engine/Source/Runtime/Renderer/Private/LightRendering.cpp
@@ -1015,7 +1015,7 @@ class FDeferredLightPS : public FGlobalShader
 		// When using Substrate, complex transmittance permutation is not used
 		if (Substrate::IsSubstrateEnabled())
 		{
-			if (!PermutationVector.Get<FDeferredLightPS::FHairComplexTransmittance>())
+			if (!PermutationVector.Get<FDeferredLightPS::FHairComplexTransmittance>() && !PermutationVector.Get< FVisualizeCullingDim >())
 			{
 				return false;
 			}
