Custom Render Pass:  fix when a CRP render target is accessed by Niagara.  Niagara switches the render target to external access mode, and it needs to be switched back to internal access mode in that case.  This will be a nop in the common case.  A licensee on UDN proposed the fix, and I confirmed it works.

#rnx

#Synced-CL 49564039

diff --git a/Engine/Source/Runtime/Renderer/Private/SceneCaptureRendering.cpp b/Engine/Source/Runtime/Renderer/Private/SceneCaptureRendering.cpp
--- a/Engine/Source/Runtime/Renderer/Private/SceneCaptureRendering.cpp
+++ b/Engine/Source/Runtime/Renderer/Private/SceneCaptureRendering.cpp
@@ -1060,6 +1060,10 @@ public:
 		((FTextureRenderTarget2DResource*)SceneCaptureRenderTarget)->Resize(GraphBuilder.RHICmdList, RenderTargetSize.X, RenderTargetSize.Y, bAutoGenerateMips);
 
 		RenderTargetTexture = SceneCaptureRenderTarget->GetRenderTargetTexture(GraphBuilder);
+
+		// If a Custom Render Pass output was accessed by Niagara, it will have converted it to external access, so switch it back to internal access.
+		// If it's already in internal access mode (the common expected case), this is a nop.
+		GraphBuilder.UseInternalAccessMode(RenderTargetTexture);
 	}
 
 	virtual void OnEndPass(FRDGBuilder& GraphBuilder) override
