Fix AO mesh decals not showing.
There were two bugs. The first is that the number of bits allocated for the decal render target mask in mesh relevance was off by one. So the AO target bit was lost. Added the extra bit back in and confirmed that it doesn't bump the struct size. The second bug was that AO only (or emissive only) mesh decals didn't compile the material vertex shader.
#Synced-CL 49559171

diff --git a/Engine/Source/Runtime/Renderer/Private/CompositionLighting/PostProcessMeshDecals.cpp b/Engine/Source/Runtime/Renderer/Private/CompositionLighting/PostProcessMeshDecals.cpp
--- a/Engine/Source/Runtime/Renderer/Private/CompositionLighting/PostProcessMeshDecals.cpp
+++ b/Engine/Source/Runtime/Renderer/Private/CompositionLighting/PostProcessMeshDecals.cpp
@@ -62,8 +62,7 @@ public:
 
 	static bool ShouldCompilePermutation(const FMeshMaterialShaderPermutationParameters& Parameters)
 	{
-		return (Parameters.MaterialParameters.MaterialDomain == MD_DeferredDecal) &&
-			DecalRendering::GetBaseRenderStage(DecalRendering::ComputeDecalBlendDesc(Parameters.Platform, Parameters.MaterialParameters)) != EDecalRenderStage::None;
+		return (Parameters.MaterialParameters.MaterialDomain == MD_DeferredDecal);
 	}
 
 	static void ModifyCompilationEnvironment(const FMaterialShaderPermutationParameters& Parameters, FShaderCompilerEnvironment& OutEnvironment)
diff --git a/Engine/Source/Runtime/Renderer/Private/DecalRenderingCommon.cpp b/Engine/Source/Runtime/Renderer/Private/DecalRenderingCommon.cpp
--- a/Engine/Source/Runtime/Renderer/Private/DecalRenderingCommon.cpp
+++ b/Engine/Source/Runtime/Renderer/Private/DecalRenderingCommon.cpp
@@ -8,7 +8,7 @@
 #include "MaterialShared.h"
 #include "StaticMeshBatch.h"
 
-static_assert(DecalRendering::DecalRenderTargetMode_NumBits == ((int32)EDecalRenderTargetMode::Num - 1));
+static_assert(DecalRendering::DecalRenderTargetMode_NumBits == (int32)EDecalRenderTargetMode::Num);
 
 namespace DecalRendering
 {
diff --git a/Engine/Source/Runtime/Renderer/Public/StaticMeshBatch.h b/Engine/Source/Runtime/Renderer/Public/StaticMeshBatch.h
--- a/Engine/Source/Runtime/Renderer/Public/StaticMeshBatch.h
+++ b/Engine/Source/Runtime/Renderer/Public/StaticMeshBatch.h
@@ -31,7 +31,7 @@ private:
 	FStaticMeshBatch(const FStaticMeshBatch& InStaticMesh);
 };
 
-namespace DecalRendering { enum { DecalRenderTargetMode_NumBits = 5	}; }
+namespace DecalRendering { enum { DecalRenderTargetMode_NumBits = 6 }; }
 
 /**
  * FStaticMeshBatch data which is InitViews specific. Stored separately for cache efficiency.
