Distance Field Shadows optimization to prevent the distance field clipmap from being generated when the feature is disabled for a specific scene renderer via ShowFlags.

#rnx

#Synced-CL 49564039

diff --git a/Engine/Source/Runtime/Renderer/Private/DistanceFieldShadowing.cpp b/Engine/Source/Runtime/Renderer/Private/DistanceFieldShadowing.cpp
--- a/Engine/Source/Runtime/Renderer/Private/DistanceFieldShadowing.cpp
+++ b/Engine/Source/Runtime/Renderer/Private/DistanceFieldShadowing.cpp
@@ -763,7 +763,7 @@ bool SupportsHeightFieldShadows(ERHIFeatureLevel::Type FeatureLevel, EShaderPlat
 
 bool ShouldPrepareForDistanceFieldShadows(const FSceneRenderUpdateInputs& SceneUpdateInputs)
 {
-	if (!SceneUpdateInputs.HasAnyShowFlags([] (const FEngineShowFlags& EngineShowFlags) { return EngineShowFlags.DynamicShadows; }) ||
+	if (!SceneUpdateInputs.HasAnyShowFlags([] (const FEngineShowFlags& EngineShowFlags) { return EngineShowFlags.DynamicShadows && EngineShowFlags.RayTracedDistanceFieldShadows; }) ||
 		!SupportsDistanceFieldShadows(SceneUpdateInputs.FeatureLevel, SceneUpdateInputs.ShaderPlatform))
 	{
 		return false;
