[HWRT] Add indirect TLAS build support

#rnx
#tests replayrun
#Synced-CL 49576663

diff --git a/Engine/Source/Runtime/RHI/Public/RHIContext.h b/Engine/Source/Runtime/RHI/Public/RHIContext.h
--- a/Engine/Source/Runtime/RHI/Public/RHIContext.h
+++ b/Engine/Source/Runtime/RHI/Public/RHIContext.h
@@ -617,6 +617,10 @@ struct FRayTracingSceneBuildParams
 	FRHIBuffer* HitGroupContributionsBuffer = nullptr;
 	uint32 HitGroupContributionsBufferOffset = 0;
 
+	// Buffer of actual number of instances used for indirect acceleration structure build (if supported)
+	FRHIBuffer* InstanceCountBuffer = nullptr;
+	uint32 InstanceCountBufferOffset = 0;
+
 	uint32 NumInstances = 0;
 
 	// Unique list of geometries referenced by all instances in this scene.
diff --git a/Engine/Source/Runtime/RHI/Public/RHIGlobals.h b/Engine/Source/Runtime/RHI/Public/RHIGlobals.h
--- a/Engine/Source/Runtime/RHI/Public/RHIGlobals.h
+++ b/Engine/Source/Runtime/RHI/Public/RHIGlobals.h
@@ -433,6 +433,9 @@ struct FRHIGlobals
 
 		/** Whether or not the RHI supports async building ray tracing acceleration structures. */
 		bool SupportsAsyncBuildAccelerationStructure = false;
+
+		/** Whether or not the RHI supports indirect building of ray tracing top level acceleration structures. */
+		bool SupportsBuildTopLevelAccelerationStructureIndirect = false;
 		
 		/** Whether or not the RHI supports async ray trace dispatch calls. */
 		bool SupportsAsyncRayTraceDispatch = false;
diff --git a/Engine/Source/Runtime/Renderer/Private/RayTracing/RayTracingScene.cpp b/Engine/Source/Runtime/Renderer/Private/RayTracing/RayTracingScene.cpp
--- a/Engine/Source/Runtime/Renderer/Private/RayTracing/RayTracingScene.cpp
+++ b/Engine/Source/Runtime/Renderer/Private/RayTracing/RayTracingScene.cpp
@@ -259,7 +259,7 @@ void FRayTracingScene::Update(FRDGBuilder& GraphBuilder, FSceneUniformBuffer& Sc
 
 			uint32 NumNativeInstances = LayerView.InstanceBufferBuilder.GetMaxNumInstances();
 
-			if (bCompactInstanceBuffer)
+			if (bCompactInstanceBuffer && !GRHIGlobals.RayTracing.SupportsBuildTopLevelAccelerationStructureIndirect)
 			{
 				NumNativeInstances = FMath::Max<uint32>(CVarRayTracingSceneCompactInstancesMin.GetValueOnRenderThread(), LayerView.NumActiveInstances + CVarRayTracingSceneCompactInstancesMargin.GetValueOnRenderThread());
 				NumNativeInstances = FMath::Min<uint32>(NumNativeInstances, LayerView.InstanceBufferBuilder.GetMaxNumInstances());
@@ -518,6 +518,7 @@ void FRayTracingScene::Build(FRDGBuilder& GraphBuilder, ERDGPassFlags ComputePas
 		[PassParams, this](FRHICommandList& RHICmdList)
 		{
 			const bool bUseBatchedBuild = CVarRayTracingSceneBatchedBuild.GetValueOnRenderThread();
+			const uint32 MaxNumViews = ActiveViews.GetMaxIndex();
 
 			TArray<FRayTracingSceneBuildParams> BatchedBuildParams;
 			BatchedBuildParams.Reserve(NumLayers); // TODO: should also take num views into account
@@ -543,6 +544,12 @@ void FRayTracingScene::Build(FRDGBuilder& GraphBuilder, ERDGPassFlags ComputePas
 						BuildParams.HitGroupContributionsBuffer = LayerView.HitGroupContributionsBuffer->GetRHI();
 						BuildParams.HitGroupContributionsBufferOffset = 0;
 					}
+
+					if (GRHIGlobals.RayTracing.SupportsBuildTopLevelAccelerationStructureIndirect)
+					{
+						BuildParams.InstanceCountBuffer = InstanceStatsBuffer ? InstanceStatsBuffer->GetRHI() : nullptr;
+						BuildParams.InstanceCountBufferOffset = (LayerIndex * MaxNumViews + ViewIndex) * sizeof(FInstanceBufferStats);
+					}
 					
 					BuildParams.NumInstances = LayerView.MaxNumInstances;
 					BuildParams.ReferencedGeometries = LayerView.InstanceBufferBuilder.GetReferencedGeometries();
