LightFunctionAtlas - visualisation now show material that have failed to compiled or when the default material is expliciterly specified.

#tests PC editor light function atlas visualisation.
#rnx
#Synced-CL 49646537

diff --git a/Engine/Source/Runtime/Renderer/Private/LightFunctionAtlas.cpp b/Engine/Source/Runtime/Renderer/Private/LightFunctionAtlas.cpp
--- a/Engine/Source/Runtime/Renderer/Private/LightFunctionAtlas.cpp
+++ b/Engine/Source/Runtime/Renderer/Private/LightFunctionAtlas.cpp
@@ -928,6 +928,7 @@ FScreenPassTexture FLightFunctionAtlas::AddDebugVisualizationPasses(FRDGBuilder&
 			const uint32 LFMaterialUniqueID = LightFunctionMaterial->GetMaterialInterface() ? MurmurFinalize32(MurmurFinalize32(LightFunctionMaterial->GetMaterialInterface()->GetUniqueID())) : 0xFFFFFFFF;
 
 			const FLinearColor MaterialColor(FColor(LFMaterialUniqueID & 0xFF, (LFMaterialUniqueID >> 8) & 0xFF, (LFMaterialUniqueID >> 16) & 0xFF));
+			const FLinearColor WarningColor(1.0f, 100.0f / 255.0f, 20.0f / 255.0f);
 
 			uint32 LightCountUsingThisMaterial = 0;
 			for (auto& LocalLight : EffectiveLocalLightSlotArray)
@@ -940,8 +941,29 @@ FScreenPassTexture FLightFunctionAtlas::AddDebugVisualizationPasses(FRDGBuilder&
 
 			Text = FString::Printf(TEXT("%2d"), LightCountUsingThisMaterial);
 			Canvas.DrawShadowedString(DrawPosX, DrawPosY, *Text, GEngine->GetLargeFont(), MaterialColor);
+
+			FString MaterialDetails;
+			const bool bDefaultMaterialSpecified = LightFunctionMaterial->GetFallback(View.GetFeatureLevel()) == nullptr;
+			if (bDefaultMaterialSpecified)
+			{
+				MaterialDetails = TEXT("[Default Material]   ");
+			}
+			const bool bMaterialCompilationFailed= !bDefaultMaterialSpecified && LightFunctionMaterial->GetIncompleteMaterialWithFallback(View.GetFeatureLevel()).IsDefaultMaterial();
+			if (bMaterialCompilationFailed)
+			{
+				MaterialDetails = TEXT("[Compilation Failed] ");
+			}
+
 			Text = FString::Printf(TEXT("lights - %s"), *MaterialName);
-			Canvas.DrawShadowedString(DrawPosX + 20, DrawPosY, *Text, GEngine->GetLargeFont(), MaterialColor);
+			if (!MaterialDetails.IsEmpty())
+			{
+				Canvas.DrawShadowedString(DrawPosX +  20, DrawPosY, *MaterialDetails, GEngine->GetLargeFont(), WarningColor);
+				Canvas.DrawShadowedString(DrawPosX + 140, DrawPosY, *Text, GEngine->GetLargeFont(), MaterialColor);
+			}
+			else
+			{
+				Canvas.DrawShadowedString(DrawPosX + 20, DrawPosY, *Text, GEngine->GetLargeFont(), MaterialColor);
+			}
 
 			// Draw a line around the corresponding atlas tile
 			auto OutLineAtlasSlot = [&](float X0, float Y0, float X1, float Y1)
