Adding CVarVTTexturePageTableSizeAlign to tune PageTable size alignment, default to 256 in order to reduce RT re allocation churn frame by frame

#Synced-CL 49557800

diff --git a/Engine/Source/Runtime/Renderer/Private/VT/VirtualTextureSpace.cpp b/Engine/Source/Runtime/Renderer/Private/VT/VirtualTextureSpace.cpp
--- a/Engine/Source/Runtime/Renderer/Private/VT/VirtualTextureSpace.cpp
+++ b/Engine/Source/Runtime/Renderer/Private/VT/VirtualTextureSpace.cpp
@@ -34,6 +34,14 @@ static TAutoConsoleVariable<int32> CVarVTMaskedPageTableUpdates(
 	ECVF_RenderThreadSafe
 	);
 
+// Align value of VT PageTable texture dimension, reduce frame to frame variation of the render target size 
+static int32 GVTTexturePageTableSizeAlign = 256;
+static FAutoConsoleVariableRef CVarVTTexturePageTableSizeAlign(
+	TEXT("r.VT.TexturePageTableSizeAlign"),
+	GVTTexturePageTableSizeAlign,
+	TEXT("Align value the size of the Virtual Texture PageTable texture width and height. (default: 256 bytes)"),
+	ECVF_RenderThreadSafe);
+
 static EPixelFormat GetFormatForNumLayers(uint32 NumLayers, EVTPageTableFormat Format)
 {
 	const bool bUse16Bits = (Format == EVTPageTableFormat::UInt16);
@@ -144,11 +152,14 @@ void FVirtualTextureSpace::ReleaseRHI()
 FUintPoint FVirtualTextureSpace::GetRequiredPageTableAllocationSize() const
 {
 	// Private spaces should allocate the full page table texture up front.
-	const uint32 Width = Description.bPrivateSpace ? Description.MaxSpaceSize : Allocator.GetAllocatedWidth();
-	const uint32 Height = Description.bPrivateSpace ? Description.MaxSpaceSize : Allocator.GetAllocatedHeight();
-	// We align on some minimum size. Maybe minimum, and align sizes should be different? But OK for now.
-	const uint32 WidthAligned = Align(Width, VIRTUALTEXTURE_MIN_PAGETABLE_SIZE);
-	const uint32 HeightAligned = Align(Height, VIRTUALTEXTURE_MIN_PAGETABLE_SIZE);
+
+	const uint32 Width =  FMath::Max(  Description.bPrivateSpace ? Description.MaxSpaceSize : Allocator.GetAllocatedWidth(), VIRTUALTEXTURE_MIN_PAGETABLE_SIZE );
+	const uint32 Height = FMath::Max(  Description.bPrivateSpace ? Description.MaxSpaceSize : Allocator.GetAllocatedHeight(), VIRTUALTEXTURE_MIN_PAGETABLE_SIZE) ;
+
+	const uint32 AlignValue = (Description.bPrivateSpace) ? VIRTUALTEXTURE_MIN_PAGETABLE_SIZE : FMath::RoundUpToPowerOfTwo( FMath::Clamp(GVTTexturePageTableSizeAlign, 1, GRHIGlobals.MaxTextureDimensions) );
+
+	const uint32 WidthAligned	= Align(Width, AlignValue);
+	const uint32 HeightAligned	= Align(Height, AlignValue);
 	return FUintPoint(WidthAligned, HeightAligned);
 }
 
