Ray Tracing: Fix exposure handling for debug views

In the ray tracing case, exposure was incorrectly disabled for some modes such as PrimaryRays which do need to be exposed

In the path tracing case, the exposure was being always applied even for modes that don't require it (most debug modes)

Fixed both and removed uncessary logic from RayTracingDebug

#Synced-CL 49691844

diff --git a/Engine/Shaders/Private/RayTracing/RayTracingDebug.usf b/Engine/Shaders/Private/RayTracing/RayTracingDebug.usf
--- a/Engine/Shaders/Private/RayTracing/RayTracingDebug.usf
+++ b/Engine/Shaders/Private/RayTracing/RayTracingDebug.usf
@@ -39,7 +39,6 @@ RaytracingAccelerationStructure FarFieldTLAS;
 
 uint VisualizationMode;
 uint PickerDomain;
-uint ShouldUsePreExposure;
 uint OpaqueOnly;
 uint TraceFarField;
 uint TintFarField;
@@ -502,11 +501,8 @@ RAY_TRACING_ENTRY_RAYGEN_OR_INLINE(RayTracingDebugMain)
 		Result *= float4(1.0f, 0.5f, 0.5f, 1.0f);
 	}
 
-	if (ShouldUsePreExposure)
-	{
-		// Only  when the output is tonemapped
-		Result.rgb *= View.PreExposure;
-	}
+	// When the visualization mode does not require exposure, we disable it via RaytracingDebugViewModeNeedsTonemapping which makes this factor 1.0
+	Result.rgb *= View.PreExposure;
 
 	if (!IsFinite(Result.r) || !IsFinite(Result.g) || !IsFinite(Result.b))
 	{
diff --git a/Engine/Shaders/Private/RayTracing/RayTracingPrimaryRays.usf b/Engine/Shaders/Private/RayTracing/RayTracingPrimaryRays.usf
--- a/Engine/Shaders/Private/RayTracing/RayTracingPrimaryRays.usf
+++ b/Engine/Shaders/Private/RayTracing/RayTracingPrimaryRays.usf
@@ -41,7 +41,6 @@ int ShadowsTranslucencyType;
 int ShouldDoDirectLighting;
 int ShouldDoEmissiveAndIndirectLighting;
 int UpscaleFactor;
-int ShouldUsePreExposure;
 uint PrimaryRayFlags;
 
 float TranslucencyMinRayDistance;
diff --git a/Engine/Source/Runtime/Renderer/Private/PostProcess/PostProcessEyeAdaptation.cpp b/Engine/Source/Runtime/Renderer/Private/PostProcess/PostProcessEyeAdaptation.cpp
--- a/Engine/Source/Runtime/Renderer/Private/PostProcess/PostProcessEyeAdaptation.cpp
+++ b/Engine/Source/Runtime/Renderer/Private/PostProcess/PostProcessEyeAdaptation.cpp
@@ -370,7 +370,11 @@ bool IsAutoExposureDebugMode(const FViewInfo& View)
 	return View.Family->UseDebugViewPS() ||
 		!EngineShowFlags.Lighting ||
 		(EngineShowFlags.VisualizeBuffer && View.CurrentBufferVisualizationMode != NAME_None && !GetBufferVisualizationData().GetMaterialApplyAutoExposure(View.CurrentBufferVisualizationMode)) ||
-		EngineShowFlags.RayTracingDebug ||
+		
+		// NOTE: Debug modes that are not tonemapped (such as normals, basecolor, etc ...) should not have exposure applied
+		(EngineShowFlags.RayTracingDebug && !RaytracingDebugViewModeNeedsTonemapping(View)) ||
+		(EngineShowFlags.PathTracing && !PathTracing::NeedsTonemapping()) ||
+
 		EngineShowFlags.VisualizeDistanceFieldAO ||
 		EngineShowFlags.VisualizeVolumetricCloudConservativeDensity ||
 		EngineShowFlags.VisualizeVolumetricCloudEmptySpaceSkipping ||
diff --git a/Engine/Source/Runtime/Renderer/Private/RayTracing/RayTracingDebug.cpp b/Engine/Source/Runtime/Renderer/Private/RayTracing/RayTracingDebug.cpp
--- a/Engine/Source/Runtime/Renderer/Private/RayTracing/RayTracingDebug.cpp
+++ b/Engine/Source/Runtime/Renderer/Private/RayTracing/RayTracingDebug.cpp
@@ -198,7 +198,6 @@ IMPLEMENT_UNIFORM_BUFFER_STRUCT(FRayTracingDebugHitStatsUniformBufferParameters,
 BEGIN_SHADER_PARAMETER_STRUCT(FRayTracingDebugParameters, )
 	SHADER_PARAMETER(uint32, VisualizationMode)
 	SHADER_PARAMETER(uint32, PickerDomain)
-	SHADER_PARAMETER(uint32, ShouldUsePreExposure)
 	SHADER_PARAMETER(float, TimingScale)
 	SHADER_PARAMETER(float, MaxTraceDistance)
 	SHADER_PARAMETER(float, FarFieldMaxTraceDistance)
@@ -1758,7 +1757,6 @@ void RenderRayTracingDebug(FRDGBuilder& GraphBuilder, const FScene& Scene, const
 	FRayTracingDebugParameters SharedParameters;
 	SharedParameters.VisualizationMode = DebugVisualizationMode;
 	SharedParameters.PickerDomain = CVarRayTracingVisualizePickerDomain.GetValueOnRenderThread();
-	SharedParameters.ShouldUsePreExposure = View.Family->EngineShowFlags.Tonemapper;
 	SharedParameters.TimingScale = GetRayTracingDebugTimingScale();
 	SharedParameters.OpaqueOnly = CVarRayTracingVisualizeOpaqueOnly.GetValueOnRenderThread();
 	SharedParameters.TraceFarField = CVarRayTracingVisualizeFarField.GetValueOnRenderThread() ? 1 : 0;
diff --git a/Engine/Source/Runtime/Renderer/Private/RayTracing/RayTracingPrimaryRays.cpp b/Engine/Source/Runtime/Renderer/Private/RayTracing/RayTracingPrimaryRays.cpp
--- a/Engine/Source/Runtime/Renderer/Private/RayTracing/RayTracingPrimaryRays.cpp
+++ b/Engine/Source/Runtime/Renderer/Private/RayTracing/RayTracingPrimaryRays.cpp
@@ -39,7 +39,6 @@ class FRayTracingPrimaryRaysRGS : public FGlobalShader
 		SHADER_PARAMETER(int32, ShadowsTranslucencyType)
 		SHADER_PARAMETER(int32, ShouldDoEmissiveAndIndirectLighting)
 		SHADER_PARAMETER(int32, UpscaleFactor)
-		SHADER_PARAMETER(int32, ShouldUsePreExposure)
 		SHADER_PARAMETER(uint32, PrimaryRayFlags)
 		SHADER_PARAMETER(float, TranslucencyMinRayDistance)
 		SHADER_PARAMETER(float, TranslucencyMaxRayDistance)
@@ -159,7 +158,6 @@ void RenderRayTracingPrimaryRaysView(
 	PassParameters->TranslucencyMaxRoughness = FMath::Clamp(TranslucencyOptions.MaxRoughness >= 0 ? TranslucencyOptions.MaxRoughness : View.FinalPostProcessSettings.RayTracingTranslucencyMaxRoughness, 0.01f, 1.0f);
 	PassParameters->TranslucencyRefraction = TranslucencyOptions.EnableRefraction >= 0 ? TranslucencyOptions.EnableRefraction : View.FinalPostProcessSettings.RayTracingTranslucencyRefraction;
 	PassParameters->MaxNormalBias = GetRaytracingMaxNormalBias();
-	PassParameters->ShouldUsePreExposure = View.Family->EngineShowFlags.Tonemapper;
 	PassParameters->PrimaryRayFlags = (uint32)Flags;
 	PassParameters->TLAS = View.GetRayTracingSceneLayerViewChecked(ERayTracingSceneLayer::Base);
 	PassParameters->ViewUniformBuffer = View.ViewUniformBuffer;
