Mobile: Added a separate shader permutation for desktop-like Single Layer Water (SLW) rendering
Can be toggled at runtime on supported configurations (deferred multi-pass) via r.Water.SingleLayer.MobileFallback
The default remains the translucent fallback rendering.
#Synced-CL 49714310

diff --git a/Engine/Shaders/Private/MaterialTemplate.ush b/Engine/Shaders/Private/MaterialTemplate.ush
--- a/Engine/Shaders/Private/MaterialTemplate.ush
+++ b/Engine/Shaders/Private/MaterialTemplate.ush
@@ -2465,10 +2465,14 @@ float MaterialExpressionSceneDepthWithoutWater(float2 ViewportUV, float Fallback
 {
 #if MATERIAL_SHADINGMODEL_SINGLELAYERWATER && !HIT_PROXY_SHADER && !SCENE_TEXTURES_DISABLED && (SINGLE_LAYER_WATER_SHADING_QUALITY != SINGLE_LAYER_WATER_SHADING_QUALITY_MOBILE_WITH_DEPTH_BUFFER)
 
-#if MOBILE_SUPPORTS_SM5_MATERIAL_NODES
-	const float4 SLWBufferBilinearUVMinMax = View.BufferSizeAndInvSize.zwzw * (ResolvedView.ViewRectMinAndSize.xyxy + float4(0.55f, 0.55f, ResolvedView.ViewRectMinAndSize.zw - 0.55f));
-	const float2 ClampedViewportUV = clamp(ViewportUV, SLWBufferBilinearUVMinMax.xy, SLWBufferBilinearUVMinMax.zw);
-	float WaterDeviceZ = WaterSampleSceneDepthWithoutWater(MobileBasePass.SceneDepthCopyTexture, MobileBasePass.SceneDepthCopySampler, ClampedViewportUV, View.BufferSizeAndInvSize.xy, View.BufferSizeAndInvSize.zw);
+#if SHADING_PATH_MOBILE
+	#if MOBILE_SINGLELAYERWATER_FALLBACK
+		return FallbackDepth;
+	#else
+		const float4 SLWBufferBilinearUVMinMax = View.BufferSizeAndInvSize.zwzw * (ResolvedView.ViewRectMinAndSize.xyxy + float4(0.55f, 0.55f, ResolvedView.ViewRectMinAndSize.zw - 0.55f));
+		const float2 ClampedViewportUV = clamp(ViewportUV, SLWBufferBilinearUVMinMax.xy, SLWBufferBilinearUVMinMax.zw);
+		float WaterDeviceZ = WaterSampleSceneDepthWithoutWater(MobileBasePass.SceneDepthCopyTexture, MobileBasePass.SceneDepthCopySampler, ClampedViewportUV, View.BufferSizeAndInvSize.xy, View.BufferSizeAndInvSize.zw);
+	#endif
 #else
 	const float2 ClampedViewportUV = clamp(ViewportUV, SingleLayerWater.SceneWithoutSingleLayerWaterMinMaxUV.xy, SingleLayerWater.SceneWithoutSingleLayerWaterMinMaxUV.zw);
 	float WaterDeviceZ = WaterSampleSceneDepthWithoutWater(SingleLayerWater.SceneDepthWithoutSingleLayerWaterTexture, SingleLayerWaterSceneDepthSampler, ClampedViewportUV, 
diff --git a/Engine/Shaders/Private/MobileBasePassPixelShader.usf b/Engine/Shaders/Private/MobileBasePassPixelShader.usf
--- a/Engine/Shaders/Private/MobileBasePassPixelShader.usf
+++ b/Engine/Shaders/Private/MobileBasePassPixelShader.usf
@@ -22,8 +22,9 @@
 #define MOBILE_QL_FORCE_DISABLE_PREINTEGRATEDGF 0
 #endif
 
-#if SUBSTRATE_ENABLED && MERGED_LOCAL_LIGHTS_MOBILE == 2
-#define ALLOW_LOCAL_LIGHT_DISTANCE_ATTENUATION 1
+// By default we use translucent fallback for SLW rendering
+#ifndef MOBILE_SINGLELAYERWATER_FALLBACK
+#define MOBILE_SINGLELAYERWATER_FALLBACK 1
 #endif
 
 #define UseBasePassSkylightDiffuse  (1)
@@ -44,7 +45,7 @@
 	#undef SINGLE_LAYER_WATER_SHADING_QUALITY
 	#endif
 
-	#if MOBILE_SUPPORTS_SM5_MATERIAL_NODES
+	#if !MOBILE_SINGLELAYERWATER_FALLBACK
 		// Value must match SINGLE_LAYER_WATER_SHADING_QUALITY_FULL in SingleLayerWaterCommon.ush!
 		#define SINGLE_LAYER_WATER_SHADING_QUALITY 0
 		// SINGLE_LAYER_WATER_SHADING_QUALITY_FULL normally uses a refraction mask texture which we don't have on mobile
@@ -82,7 +83,7 @@
 #define FORCE_VERTEX_NORMAL (MOBILE_QL_DISABLE_MATERIAL_NORMAL)
 #define SUPPORT_SPOTLIGHTS_SHADOW (MATERIALBLENDING_SOLID || MATERIALBLENDING_MASKED) && PROJECT_MOBILE_ENABLE_MOVABLE_SPOTLIGHT_SHADOWS
 // Sky does not write to depth and must not overwrite GBuffer
-#define MOBILE_USE_GBUFFER (MOBILE_DEFERRED_SHADING && ((MATERIALBLENDING_SOLID || MATERIALBLENDING_MASKED) && (!MATERIAL_SHADINGMODEL_SINGLELAYERWATER || MOBILE_SUPPORTS_SM5_MATERIAL_NODES) && !MATERIAL_IS_SKY))
+#define MOBILE_USE_GBUFFER (MOBILE_DEFERRED_SHADING && ((MATERIALBLENDING_SOLID || MATERIALBLENDING_MASKED) && !(MATERIAL_SHADINGMODEL_SINGLELAYERWATER && MOBILE_SINGLELAYERWATER_FALLBACK) && !MATERIAL_IS_SKY))
 
 #if MATERIAL_SHADINGMODEL_THIN_TRANSLUCENT || SUBSTRATE_BLENDING_TRANSLUCENT_COLOREDTRANSMITTANCE
 	#define MOBILE_TRANSLUCENT_COLOR_TRANSMITTANCE 1
@@ -1180,7 +1181,7 @@ void Main(
 		WaterDiffuseIndirectIlluminance = lerp(WaterDiffuseIndirectIlluminance, PI, View.UnlitViewmodeMask);
 #endif
 
-#if MOBILE_SUPPORTS_SM5_MATERIAL_NODES
+#if !MOBILE_SINGLELAYERWATER_FALLBACK
 		const float4 SLWBufferBilinearUVMinMax = View.BufferSizeAndInvSize.zwzw * (ResolvedView.ViewRectMinAndSize.xyxy + float4(0.55f, 0.55f, ResolvedView.ViewRectMinAndSize.zw - 0.55f));
 #endif
 
@@ -1192,7 +1193,7 @@ void Main(
 		WaterVolumeLightingOutput WaterLighting = EvaluateWaterVolumeLighting(
 			MaterialParameters, PixelMaterialInputs, ResolvedView,
 			DirectionalLightShadow, 
-#if MOBILE_SUPPORTS_SM5_MATERIAL_NODES
+#if !MOBILE_SINGLELAYERWATER_FALLBACK
 			MobileBasePass.SceneDepthCopyTexture, MobileBasePass.SceneDepthCopySampler,
 			ResolvedView.BufferSizeAndInvSize.xy,
 			ResolvedView.BufferSizeAndInvSize.zw,
@@ -1207,7 +1208,7 @@ void Main(
 
 		// Accumulate luminance and occlude the background according to transmittance to view and mean transmittance to lights.
 		LightAccumulator_AddSplit(DirectLighting, WaterLighting.Luminance, 0.0f, WaterLighting.Luminance, 1.0f, false);
-#if MOBILE_SUPPORTS_SM5_MATERIAL_NODES
+#if !MOBILE_SINGLELAYERWATER_FALLBACK
 		Opacity = 1.0f;
 #else
 		Opacity = 1.0 - ((1.0 - Opacity) * dot(WaterLighting.WaterToSceneToLightTransmittance, float3(1.0 / 3.0, 1.0 / 3.0, 1.0 / 3.0)));
diff --git a/Engine/Shaders/Private/MobileLightingCommon.ush b/Engine/Shaders/Private/MobileLightingCommon.ush
--- a/Engine/Shaders/Private/MobileLightingCommon.ush
+++ b/Engine/Shaders/Private/MobileLightingCommon.ush
@@ -18,7 +18,7 @@
 #if MOBILE_SSR_ENABLED && !(MOBILE_QL_FORCE_FULLY_ROUGH || MATERIAL_FULLY_ROUGH)
 	#if MOBILE_DEFERRED_SHADING
 		#if IS_MOBILE_BASE_PASS
-			#define MOBILE_USE_SSR (MATERIAL_SSR || (MOBILE_SUPPORTS_SM5_MATERIAL_NODES && MATERIAL_SHADINGMODEL_SINGLELAYERWATER))
+			#define MOBILE_USE_SSR (MATERIAL_SSR || (MATERIAL_SHADINGMODEL_SINGLELAYERWATER && !MOBILE_SINGLELAYERWATER_FALLBACK))
 		#else
 			#define MOBILE_USE_SSR 1
 		#endif
diff --git a/Engine/Source/Runtime/Engine/Private/Materials/MaterialInterface.cpp b/Engine/Source/Runtime/Engine/Private/Materials/MaterialInterface.cpp
--- a/Engine/Source/Runtime/Engine/Private/Materials/MaterialInterface.cpp
+++ b/Engine/Source/Runtime/Engine/Private/Materials/MaterialInterface.cpp
@@ -517,7 +517,7 @@ FMaterialRelevance UMaterialInterface::GetRelevance_Internal(const UMaterial* Ma
 
 	const bool bIsMobile = IsMobilePlatform(InShaderPlatform);
 	const bool bUsesSingleLayerWaterMaterial = MaterialResource->GetShadingModels().HasShadingModel(MSM_SingleLayerWater);
-	const bool bIsSinglePassWaterTranslucent = bIsMobile && bUsesSingleLayerWaterMaterial && !MobileSupportsSM5MaterialNodes(InShaderPlatform);
+	const bool bIsSinglePassWaterTranslucent = bIsMobile && bUsesSingleLayerWaterMaterial && IsUsingSingleLayerWaterFallback(InShaderPlatform);
 
 	// Note that even though XX_GameThread() api is called, this function can be called on non game thread via 
 	// GetRelevance_Concurrent()
diff --git a/Engine/Source/Runtime/RenderCore/Private/RenderUtils.cpp b/Engine/Source/Runtime/RenderCore/Private/RenderUtils.cpp
--- a/Engine/Source/Runtime/RenderCore/Private/RenderUtils.cpp
+++ b/Engine/Source/Runtime/RenderCore/Private/RenderUtils.cpp
@@ -1553,6 +1553,26 @@ bool IsSingleLayerWaterDepthPrepassEnabled(const FStaticShaderPlatform Platform,
 
 }
 
+bool IsUsingSingleLayerWaterFallback(FStaticShaderPlatform Platform)
+{
+	static const auto CVarWaterSingleLayerMobileFallback = IConsoleManager::Get().FindConsoleVariable(TEXT("r.Water.SingleLayer.MobileFallback"));
+
+	if (!IsMobilePlatform(Platform))
+	{
+		return false;
+	}
+
+	// Only multi-pass deferred supports SLW rendering
+	if (IsMobileDeferredShadingEnabled(Platform) 
+		&& !FReadOnlyCVARCache::MobileAllowFramebufferFetch(Platform))
+	{
+		return (!CVarWaterSingleLayerMobileFallback || CVarWaterSingleLayerMobileFallback->GetInt() != 0);
+	}
+
+	return true;
+}
+
+
 RENDERCORE_API bool IsUsingDBuffers(const FStaticShaderPlatform Platform)
 {
 	// Dbuffer is enabled:
diff --git a/Engine/Source/Runtime/RenderCore/Public/RenderUtils.h b/Engine/Source/Runtime/RenderCore/Public/RenderUtils.h
--- a/Engine/Source/Runtime/RenderCore/Public/RenderUtils.h
+++ b/Engine/Source/Runtime/RenderCore/Public/RenderUtils.h
@@ -500,6 +500,11 @@ RENDERCORE_API bool UseSparseVolumeTextures(EShaderPlatform Platform);
 */
 RENDERCORE_API bool IsSingleLayerWaterDepthPrepassEnabled(const FStaticShaderPlatform Platform, FStaticFeatureLevel FeatureLevel);
 
+/**
+*	(Runtime) Whether platform uses translucent fallback for SLW rendering.
+*/
+RENDERCORE_API bool IsUsingSingleLayerWaterFallback(FStaticShaderPlatform Platform);
+
 /**
 *	Checks if virtual texturing lightmap enabled and supported
 */
diff --git a/Engine/Source/Runtime/Renderer/Private/DebugViewModeRendering.cpp b/Engine/Source/Runtime/Renderer/Private/DebugViewModeRendering.cpp
--- a/Engine/Source/Runtime/Renderer/Private/DebugViewModeRendering.cpp
+++ b/Engine/Source/Runtime/Renderer/Private/DebugViewModeRendering.cpp
@@ -550,7 +550,7 @@ void FDebugViewModeMeshProcessor::UpdateInstructionCount(FDebugViewModeShaderEle
 		{
 			TShaderRef<TMobileBasePassVSPolicyParamType<FUniformLightMapPolicy>> MobileVS;
 			TShaderRef<TMobileBasePassPSPolicyParamType<FUniformLightMapPolicy>> MobilePS;
-			if (MobileBasePass::GetShaders(LMP_NO_LIGHTMAP, EMobileLocalLightSetting::LOCAL_LIGHTS_DISABLED, *InBatchMaterial, InVertexFactoryType, MobileVS, MobilePS))
+			if (MobileBasePass::GetShaders(FUniformLightMapPolicy(LMP_NO_LIGHTMAP), EMobileLocalLightSetting::LOCAL_LIGHTS_DISABLED, *InBatchMaterial, InVertexFactoryType, MobileVS, MobilePS))
 			{
 				OutShaderElementData.NumVSInstructions = MobileVS.IsValid() ? MobileVS->GetNumInstructions() : 0;
 				OutShaderElementData.NumPSInstructions = MobilePS.IsValid() ? MobilePS->GetNumInstructions() : 0;
diff --git a/Engine/Source/Runtime/Renderer/Private/EditorPrimitivesRendering.cpp b/Engine/Source/Runtime/Renderer/Private/EditorPrimitivesRendering.cpp
--- a/Engine/Source/Runtime/Renderer/Private/EditorPrimitivesRendering.cpp
+++ b/Engine/Source/Runtime/Renderer/Private/EditorPrimitivesRendering.cpp
@@ -132,10 +132,11 @@ bool FEditorPrimitivesBasePassMeshProcessor::ProcessMobileShadingPath(const FMes
 	const EMobileLocalLightSetting LocalLightSetting = EMobileLocalLightSetting::LOCAL_LIGHTS_DISABLED;
 
 	TMeshProcessorShaders<
-		TMobileBasePassVSPolicyParamType<FUniformLightMapPolicy>,
-		TMobileBasePassPSPolicyParamType<FUniformLightMapPolicy>> BasePassShaders;
+		TMobileBasePassVSPolicyParamType<LightMapPolicyType>,
+		TMobileBasePassPSPolicyParamType<LightMapPolicyType>> BasePassShaders;
+	
 	if (!MobileBasePass::GetShaders(
-		NoLightmapPolicy.GetIndirectPolicy(),
+		NoLightmapPolicy,
 		LocalLightSetting,
 		Material,
 		VertexFactory->GetType(),
diff --git a/Engine/Source/Runtime/Renderer/Private/LightMapRendering.cpp b/Engine/Source/Runtime/Renderer/Private/LightMapRendering.cpp
--- a/Engine/Source/Runtime/Renderer/Private/LightMapRendering.cpp
+++ b/Engine/Source/Runtime/Renderer/Private/LightMapRendering.cpp
@@ -256,6 +256,41 @@ void FMobileDirectionalLightCSMAndSHIndirectPolicy::ModifyCompilationEnvironment
 	Super::ModifyCompilationEnvironment(Parameters, OutEnvironment);
 }
 
+FMobileSingleLayerWaterPolicy::FMobileSingleLayerWaterPolicy()
+{
+}
+
+bool FMobileSingleLayerWaterPolicy::ShouldCompilePermutation(const FMeshMaterialShaderPermutationParameters& Parameters)
+{
+	return IsMobilePlatform(Parameters.Platform) 
+		&& IsMobileDeferredShadingEnabled(Parameters.Platform)
+		&& !FReadOnlyCVARCache::MobileAllowFramebufferFetch(Parameters.Platform)
+		&& Parameters.MaterialParameters.ShadingModels.HasShadingModel(MSM_SingleLayerWater);
+}
+
+void FMobileSingleLayerWaterPolicy::ModifyCompilationEnvironment(const FMaterialShaderPermutationParameters& Parameters, FShaderCompilerEnvironment& OutEnvironment)
+{
+	OutEnvironment.SetDefine(TEXT("DIRECTIONAL_LIGHT_CSM"), TEXT("1"));
+	OutEnvironment.SetDefine(TEXT("MOBILE_SINGLELAYERWATER_FALLBACK"), 0u);
+	OutEnvironment.SetDefine(TEXT("MOBILE_SSR_ENABLED"), 1u); // should it be optional?
+}
+
+void FMobileSingleLayerWaterPolicy::GetVertexShaderBindings(
+	const FPrimitiveSceneProxy* PrimitiveSceneProxy,
+	const ElementDataType& ShaderElementData,
+	const VertexParametersType* VertexShaderParameters,
+	FMeshDrawSingleShaderBindings& ShaderBindings)
+{
+}
+
+void FMobileSingleLayerWaterPolicy::GetPixelShaderBindings(
+	const FPrimitiveSceneProxy* PrimitiveSceneProxy,
+	const ElementDataType& ShaderElementData,
+	const PixelParametersType* PixelShaderParameters,
+	FMeshDrawSingleShaderBindings& ShaderBindings)
+{
+}
+
 bool FSelfShadowedCachedPointIndirectLightingPolicy::ShouldCompilePermutation(const FMeshMaterialShaderPermutationParameters& Parameters)
 {
 	return Parameters.MaterialParameters.ShadingModels.IsLit()
diff --git a/Engine/Source/Runtime/Renderer/Private/LightMapRendering.h b/Engine/Source/Runtime/Renderer/Private/LightMapRendering.h
--- a/Engine/Source/Runtime/Renderer/Private/LightMapRendering.h
+++ b/Engine/Source/Runtime/Renderer/Private/LightMapRendering.h
@@ -298,6 +298,52 @@ public:
 	static void ModifyCompilationEnvironment(const FMaterialShaderPermutationParameters& Parameters, FShaderCompilerEnvironment& OutEnvironment);
 };
 
+/** Mobile Specific: SLW rendering */
+class FMobileSingleLayerWaterPolicy
+{
+public:
+	FMobileSingleLayerWaterPolicy();
+
+	typedef FRHIUniformBuffer* ElementDataType;
+
+	class VertexParametersType
+	{
+		DECLARE_INLINE_TYPE_LAYOUT(VertexParametersType, NonVirtual);
+	public:
+		void Bind(const FShaderParameterMap& ParameterMap) {}
+		void Serialize(FArchive& Ar) {}
+	};
+
+	class PixelParametersType
+	{
+		DECLARE_INLINE_TYPE_LAYOUT(PixelParametersType, NonVirtual);
+	public:
+		void Bind(const FShaderParameterMap& ParameterMap) {}
+		void Serialize(FArchive& Ar) {}
+	};
+
+
+	static RENDERER_API bool ShouldCompilePermutation(const FMeshMaterialShaderPermutationParameters& Parameters);
+	static RENDERER_API void ModifyCompilationEnvironment(const FMaterialShaderPermutationParameters& Parameters, FShaderCompilerEnvironment& OutEnvironment);
+
+	static RENDERER_API void GetVertexShaderBindings(
+		const FPrimitiveSceneProxy* PrimitiveSceneProxy,
+		const ElementDataType& ShaderElementData,
+		const VertexParametersType* VertexShaderParameters,
+		FMeshDrawSingleShaderBindings& ShaderBindings);
+
+	static RENDERER_API void GetPixelShaderBindings(
+		const FPrimitiveSceneProxy* PrimitiveSceneProxy,
+		const ElementDataType& ShaderElementData,
+		const PixelParametersType* PixelShaderParameters,
+		FMeshDrawSingleShaderBindings& ShaderBindings);
+
+	friend bool operator==(const FMobileSingleLayerWaterPolicy A, const FMobileSingleLayerWaterPolicy B)
+	{
+		return true;
+	}
+};
+
 enum ELightMapPolicyType
 {
 	LMP_NO_LIGHTMAP,
diff --git a/Engine/Source/Runtime/Renderer/Private/MobileBasePass.cpp b/Engine/Source/Runtime/Renderer/Private/MobileBasePass.cpp
--- a/Engine/Source/Runtime/Renderer/Private/MobileBasePass.cpp
+++ b/Engine/Source/Runtime/Renderer/Private/MobileBasePass.cpp
@@ -142,7 +142,7 @@ bool GetUniformMobileBasePassShaders(
 	using FPixelShaderType = TMobileBasePassPSPolicyParamType<FUniformLightMapPolicy>;
 
 	FMaterialShaderTypes ShaderTypes;
-	ShaderTypes.AddShaderType<TMobileBasePassVS<TUniformLightMapPolicy<Policy>>>();	
+	ShaderTypes.AddShaderType<TMobileBasePassVS<TUniformLightMapPolicy<Policy>>>();
 
 	switch (ColoredTransmittanceFallback)
 	{
@@ -168,7 +168,7 @@ bool GetUniformMobileBasePassShaders(
 
 template <EMobileLocalLightSetting LocalLightSetting>
 bool GetMobileBasePassShaders(
-	ELightMapPolicyType LightMapPolicyType, 
+	const FUniformLightMapPolicy& RESTRICT LightMapPolicy,
 	const FMaterial& Material, 
 	const FVertexFactoryType* VertexFactoryType, 
 	EMobileTranslucentColorTransmittanceMode ColoredTransmittanceFallback,
@@ -176,7 +176,7 @@ bool GetMobileBasePassShaders(
 	TShaderRef<TMobileBasePassPSPolicyParamType<FUniformLightMapPolicy>>& PixelShader
 	)
 {
-	switch (LightMapPolicyType)
+	switch (LightMapPolicy.GetIndirectPolicy())
 	{
 	case LMP_NO_LIGHTMAP:
 		return GetUniformMobileBasePassShaders<LMP_NO_LIGHTMAP, LocalLightSetting>(Material, VertexFactoryType, ColoredTransmittanceFallback, VertexShader, PixelShader);
@@ -200,13 +200,14 @@ bool GetMobileBasePassShaders(
 	}
 }
 
+template <typename LightMapPolicyType>
 bool MobileBasePass::GetShaders(
-	ELightMapPolicyType LightMapPolicyType,
+	const LightMapPolicyType& RESTRICT LightMapPolicy,
 	EMobileLocalLightSetting LocalLightSetting,
 	const FMaterial& MaterialResource,
 	const FVertexFactoryType* VertexFactoryType,
-	TShaderRef<TMobileBasePassVSPolicyParamType<FUniformLightMapPolicy>>& VertexShader,
-	TShaderRef<TMobileBasePassPSPolicyParamType<FUniformLightMapPolicy>>& PixelShader)
+	TShaderRef<TMobileBasePassVSPolicyParamType<LightMapPolicyType>>& VertexShader,
+	TShaderRef<TMobileBasePassPSPolicyParamType<LightMapPolicyType>>& PixelShader)
 {
 	EMobileTranslucentColorTransmittanceMode ColoredTransmittanceFallback = EMobileTranslucentColorTransmittanceMode::DEFAULT;
 	if (MaterialRequiresColorTransmittanceBlending(MaterialResource))
@@ -220,7 +221,7 @@ bool MobileBasePass::GetShaders(
 		case EMobileLocalLightSetting::LOCAL_LIGHTS_DISABLED:
 		{
 			return GetMobileBasePassShaders<EMobileLocalLightSetting::LOCAL_LIGHTS_DISABLED>(
-				LightMapPolicyType,
+				LightMapPolicy,
 				MaterialResource,
 				VertexFactoryType,
 				ColoredTransmittanceFallback,
@@ -231,7 +232,7 @@ bool MobileBasePass::GetShaders(
 		case EMobileLocalLightSetting::LOCAL_LIGHTS_ENABLED:
 		{
 			return GetMobileBasePassShaders<EMobileLocalLightSetting::LOCAL_LIGHTS_ENABLED>(
-				LightMapPolicyType,
+				LightMapPolicy,
 				MaterialResource,
 				VertexFactoryType,
 				ColoredTransmittanceFallback,
@@ -242,7 +243,7 @@ bool MobileBasePass::GetShaders(
 		case EMobileLocalLightSetting::LOCAL_LIGHTS_BUFFER:
 		{
 			return GetMobileBasePassShaders<EMobileLocalLightSetting::LOCAL_LIGHTS_BUFFER>(
-				LightMapPolicyType,
+				LightMapPolicy,
 				MaterialResource,
 				VertexFactoryType,
 				ColoredTransmittanceFallback,
@@ -257,19 +258,22 @@ bool MobileBasePass::GetShaders(
 	}
 }
 
+// Force instantination for FMobileSingleLayerWaterPolicy as it's used in a different translation unit
+template bool MobileBasePass::GetShaders<FMobileSingleLayerWaterPolicy>(
+	const FMobileSingleLayerWaterPolicy& RESTRICT LightMapPolicy,
+	EMobileLocalLightSetting LocalLightSetting,
+	const FMaterial& MaterialResource,
+	const FVertexFactoryType* VertexFactoryType,
+	TShaderRef<TMobileBasePassVSPolicyParamType<FMobileSingleLayerWaterPolicy>>& VertexShader,
+	TShaderRef<TMobileBasePassPSPolicyParamType<FMobileSingleLayerWaterPolicy>>& PixelShader);
+
 static bool UseSkyReflectionCapture(const FScene* RenderScene)
 {
 	return RenderScene
 		&& RenderScene->SkyLight
-		&& 
-		(
-			(
-				   RenderScene->SkyLight->ProcessedTexture
-				&& RenderScene->SkyLight->ProcessedTexture->TextureRHI
-			)
-			||
-			RenderScene->CanSampleSkyLightRealTimeCaptureData()
-		);
+		&& ((RenderScene->SkyLight->ProcessedTexture && 
+			RenderScene->SkyLight->ProcessedTexture->TextureRHI) 
+			|| RenderScene->CanSampleSkyLightRealTimeCaptureData());
 }
 
 const FLightSceneInfo* MobileBasePass::GetDirectionalLightInfo(const FScene* Scene, const FPrimitiveSceneProxy* PrimitiveSceneProxy)
@@ -638,7 +642,7 @@ void MobileBasePass::SetTranslucentRenderState(FMeshPassProcessorRenderState& Dr
 															CW_NONE, BO_Add, BF_One, BF_Zero, BO_Add, BF_One, BF_Zero>::GetRHI());
 			break;
 		default:
-			if (ShadingModels.HasShadingModel(MSM_SingleLayerWater) && !MobileSupportsSM5MaterialNodes(ShaderPlatform))
+			if (ShadingModels.HasShadingModel(MSM_SingleLayerWater))
 			{
 				// Single layer water is an opaque marerial rendered as translucent on Mobile. We force pre-multiplied alpha to achieve water depth based transmittance.
 				DrawRenderState.SetBlendState(TStaticBlendState<CW_RGBA, BO_Add, BF_One, BF_InverseSourceAlpha, BO_Add, BF_Zero, BF_InverseSourceAlpha,
@@ -680,19 +684,19 @@ static FMeshDrawCommandSortKey GetBasePassStaticSortKey(EDepthDrawingMode EarlyZ
 	return SortKey;
 }
 
-template<>
-void TMobileBasePassPSPolicyParamType<FUniformLightMapPolicy>::GetShaderBindings(
+template <typename LightMapPolicyType>
+void TMobileBasePassPSPolicyParamType<LightMapPolicyType>::GetShaderBindings(
 	const FScene* Scene,
 	const FStaticFeatureLevel FeatureLevel,
 	const FPrimitiveSceneProxy* PrimitiveSceneProxy,
 	const FMaterialRenderProxy& MaterialRenderProxy,
 	const FMaterial& Material,
-	const TMobileBasePassShaderElementData<FUniformLightMapPolicy>& ShaderElementData,
+	const TMobileBasePassShaderElementData<LightMapPolicyType>& ShaderElementData,
 	FMeshDrawSingleShaderBindings& ShaderBindings) const
 {
 	FMeshMaterialShader::GetShaderBindings(Scene, FeatureLevel, PrimitiveSceneProxy, MaterialRenderProxy, Material, ShaderElementData, ShaderBindings);
 
-	FUniformLightMapPolicy::GetPixelShaderBindings(
+	LightMapPolicyType::GetPixelShaderBindings(
 		PrimitiveSceneProxy,
 		ShaderElementData.LightMapPolicyElementData,
 		this,
@@ -734,6 +738,17 @@ void TMobileBasePassPSPolicyParamType<FUniformLightMapPolicy>::GetShaderBindings
 	}
 }
 
+// Force instantination for FMobileSingleLayerWaterPolicy as it's used in a different translation unit
+template void TMobileBasePassPSPolicyParamType<FMobileSingleLayerWaterPolicy>::GetShaderBindings(
+	const FScene* Scene,
+	const FStaticFeatureLevel FeatureLevel,
+	const FPrimitiveSceneProxy* PrimitiveSceneProxy,
+	const FMaterialRenderProxy& MaterialRenderProxy,
+	const FMaterial& Material,
+	const TMobileBasePassShaderElementData<FMobileSingleLayerWaterPolicy>& ShaderElementData,
+	FMeshDrawSingleShaderBindings& ShaderBindings) const;
+
+
 FMobileBasePassMeshProcessor::FMobileBasePassMeshProcessor(
 	EMeshPass::Type InMeshPassType,
 	const FScene* Scene,
@@ -758,7 +773,7 @@ bool FMobileBasePassMeshProcessor::ShouldDraw(const FMaterial& Material) const
 	bool bShouldDraw = ShouldIncludeDomainInMeshPass(Material.GetMaterialDomain());
 	const FMaterialShadingModelField ShadingModels = Material.GetShadingModels();
 	const bool bIsSLW = ShadingModels.HasShadingModel(MSM_SingleLayerWater);
-	const bool bIsTranslucentSLW = bIsSLW && !MobileSupportsSM5MaterialNodes(GetFeatureLevelShaderPlatform(ERHIFeatureLevel::ES3_1));
+	const bool bIsTranslucentSLW = bIsSLW && IsUsingSingleLayerWaterFallback(GetFeatureLevelShaderPlatform(ERHIFeatureLevel::ES3_1));
 	const bool bIsTranslucent = IsTranslucentBlendMode(Material.GetBlendMode()) || bIsTranslucentSLW; // Water goes into the translucent pass;
 	if (bTranslucentBasePass)
 	{
@@ -781,9 +796,9 @@ static bool MaterialCanReceiveCSM(const FMaterial& Material)
 	const EBlendMode BlendMode = Material.GetBlendMode();
 	const FMaterialShadingModelField ShadingModels = Material.GetShadingModels();
 	const bool bIsTranslucent = IsTranslucentBlendMode(BlendMode);
-	const bool bSingleLayerWater = ShadingModels.HasShadingModel(MSM_SingleLayerWater);
+	const bool bIsSLW = ShadingModels.HasShadingModel(MSM_SingleLayerWater);
 	const bool IsTranslucentPerPixelLighting = bIsTranslucent && Material.GetTranslucencyLightingMode() == ETranslucencyLightingMode::TLM_SurfacePerPixelLighting;
-	return !bIsTranslucent || bSingleLayerWater || IsTranslucentPerPixelLighting;
+	return !bIsTranslucent || bIsSLW || IsTranslucentPerPixelLighting;
 }
 
 bool FMobileBasePassMeshProcessor::TryAddMeshBatch(const FMeshBatch& RESTRICT MeshBatch, uint64 BatchElementMask, const FPrimitiveSceneProxy* RESTRICT PrimitiveSceneProxy, int32 StaticMeshId, const FMaterialRenderProxy& MaterialRenderProxy, const FMaterial& Material)
@@ -791,19 +806,34 @@ bool FMobileBasePassMeshProcessor::TryAddMeshBatch(const FMeshBatch& RESTRICT Me
 	if (ShouldDraw(Material))
 	{
 		FMaterialShadingModelField ShadingModels = Material.GetShadingModels();
+		const EShaderPlatform ShaderPlatform = GetFeatureLevelShaderPlatform(FeatureLevel);
 #if WITH_EDITOR
 		// Non-Editor builds filter out shading models on a material load
-		const EShaderPlatform ShaderPlatform = GetFeatureLevelShaderPlatform(FeatureLevel);
 		UMaterialInterface::FilterOutPlatformShadingModels(ShaderPlatform, ShadingModels);
 #endif
-		const bool bSingleLayerWater = ShadingModels.HasShadingModel(MSM_SingleLayerWater);
+		const bool bIsSLW = ShadingModels.HasShadingModel(MSM_SingleLayerWater);
+		// Only translucent SLW fallback takes this path. Opaque SLW must be handled by FMobileSingleLayerWaterPassMeshProcessor 
+		check(!bIsSLW || IsUsingSingleLayerWaterFallback(ShaderPlatform)); 
+
 		const EBlendMode BlendMode = Material.GetBlendMode();
 		const bool bIsLitMaterial = ShadingModels.IsLit();
-		const bool bIsTranslucent = IsTranslucentBlendMode(BlendMode) || bSingleLayerWater; // Water goes into the translucent pass;
+		const bool bIsTranslucent = IsTranslucentBlendMode(BlendMode) || bIsSLW; // Water fallback goes into the translucent pass;
 		const bool bIsMasked = IsMaskedBlendMode(Material);		
 		const bool bCanReceiveCSM = ((Flags & EFlags::CanReceiveCSM) == EFlags::CanReceiveCSM) && MaterialCanReceiveCSM(Material);
-		ELightMapPolicyType LightmapPolicyType = MobileBasePass::SelectMeshLightmapPolicy(Scene, MeshBatch, PrimitiveSceneProxy, bCanReceiveCSM, bPassUsesDeferredShading, bIsLitMaterial, bIsTranslucent);
-		return Process(MeshBatch, BatchElementMask, StaticMeshId, PrimitiveSceneProxy, MaterialRenderProxy, Material, bIsMasked, bIsTranslucent, ShadingModels, LightmapPolicyType, bCanReceiveCSM, MeshBatch.LCI);
+		ELightMapPolicyType Policy = MobileBasePass::SelectMeshLightmapPolicy(Scene, MeshBatch, PrimitiveSceneProxy, bCanReceiveCSM, bPassUsesDeferredShading, bIsLitMaterial, bIsTranslucent);
+
+		return Process<FUniformLightMapPolicy>(
+			MeshBatch,
+			BatchElementMask,
+			StaticMeshId,
+			PrimitiveSceneProxy,
+			MaterialRenderProxy,
+			Material,
+			bIsMasked,
+			ShadingModels,
+			FUniformLightMapPolicy(Policy),
+			MeshBatch.LCI,
+			bCanReceiveCSM);
 	}
 	return true;
 }
@@ -833,6 +863,7 @@ void FMobileBasePassMeshProcessor::AddMeshBatch(const FMeshBatch& RESTRICT MeshB
 	}
 }
 
+template <typename LightMapPolicyType>
 bool FMobileBasePassMeshProcessor::Process(
 		const FMeshBatch& RESTRICT MeshBatch,
 		uint64 BatchElementMask,
@@ -841,15 +872,14 @@ bool FMobileBasePassMeshProcessor::Process(
 		const FMaterialRenderProxy& RESTRICT MaterialRenderProxy,
 		const FMaterial& RESTRICT MaterialResource,
 		const bool bIsMasked,
-		const bool bIsTranslucent, //not used
 		FMaterialShadingModelField ShadingModels,
-		const ELightMapPolicyType LightMapPolicyType,
-		const bool bCanReceiveCSM,
-		const FUniformLightMapPolicy::ElementDataType& RESTRICT LightMapElementData)
+		const LightMapPolicyType& RESTRICT LightMapPolicy,
+		const typename LightMapPolicyType::ElementDataType& RESTRICT LightMapElementData,
+		const bool bCanReceiveCSM)
 {
 	TMeshProcessorShaders<
-		TMobileBasePassVSPolicyParamType<FUniformLightMapPolicy>,
-		TMobileBasePassPSPolicyParamType<FUniformLightMapPolicy>> BasePassShaders;
+		TMobileBasePassVSPolicyParamType<LightMapPolicyType>,
+		TMobileBasePassPSPolicyParamType<LightMapPolicyType>> BasePassShaders;
 
 	EMobileLocalLightSetting LocalLightSetting = EMobileLocalLightSetting::LOCAL_LIGHTS_DISABLED;
 	if (Scene && PrimitiveSceneProxy && ShadingModels.IsLit())
@@ -864,7 +894,7 @@ bool FMobileBasePassMeshProcessor::Process(
 	}
 
 	if (!MobileBasePass::GetShaders(
-		LightMapPolicyType,
+		LightMapPolicy,
 		LocalLightSetting,
 		MaterialResource,
 		MeshBatch.VertexFactory->GetType(),
@@ -917,7 +947,7 @@ bool FMobileBasePassMeshProcessor::Process(
 	ERasterizerFillMode MeshFillMode = ComputeMeshFillMode(MaterialResource, OverrideSettings);
 	ERasterizerCullMode MeshCullMode = ComputeMeshCullMode(MaterialResource, OverrideSettings);
 
-	TMobileBasePassShaderElementData<FUniformLightMapPolicy> ShaderElementData(LightMapElementData, bCanReceiveCSM);
+	TMobileBasePassShaderElementData<LightMapPolicyType> ShaderElementData(LightMapElementData, bCanReceiveCSM);
 	ShaderElementData.InitializeMeshMaterialData(ViewIfDynamicMeshCommand, PrimitiveSceneProxy, MeshBatch, StaticMeshId, false);
 
 	BuildMeshDrawCommands(
@@ -951,9 +981,11 @@ void FMobileBasePassMeshProcessor::CollectPSOInitializersForLMPolicy(
 	TMeshProcessorShaders<
 		TMobileBasePassVSPolicyParamType<FUniformLightMapPolicy>,
 		TMobileBasePassPSPolicyParamType<FUniformLightMapPolicy>> BasePassShaders;
+	
+	FUniformLightMapPolicy UniformLightmapPolicy(LightMapPolicyType);
 
 	if (!MobileBasePass::GetShaders(
-		LightMapPolicyType,
+		UniformLightmapPolicy,
 		LocalLightSetting,
 		MaterialResource,
 		VertexFactoryData.VertexFactoryType,
diff --git a/Engine/Source/Runtime/Renderer/Private/MobileBasePassRendering.cpp b/Engine/Source/Runtime/Renderer/Private/MobileBasePassRendering.cpp
--- a/Engine/Source/Runtime/Renderer/Private/MobileBasePassRendering.cpp
+++ b/Engine/Source/Runtime/Renderer/Private/MobileBasePassRendering.cpp
@@ -176,6 +176,7 @@ IMPLEMENT_MOBILE_SHADING_BASEPASS_LIGHTMAPPED_SHADER_TYPE(TUniformLightMapPolicy
 IMPLEMENT_MOBILE_SHADING_BASEPASS_LIGHTMAPPED_SHADER_TYPE(TUniformLightMapPolicy<LMP_MOBILE_DIRECTIONAL_LIGHT_AND_SH_INDIRECT>, FMobileDirectionalLightAndSHIndirectPolicy);
 IMPLEMENT_MOBILE_SHADING_BASEPASS_LIGHTMAPPED_SHADER_TYPE(TUniformLightMapPolicy<LMP_MOBILE_DIRECTIONAL_LIGHT_CSM_AND_SH_INDIRECT>, FMobileDirectionalLightCSMAndSHIndirectPolicy);
 IMPLEMENT_MOBILE_SHADING_BASEPASS_LIGHTMAPPED_SHADER_TYPE(TUniformLightMapPolicy<LMP_MOBILE_DIRECTIONAL_LIGHT_CSM>, FMobileDirectionalLightAndCSMPolicy);
+IMPLEMENT_MOBILE_SHADING_BASEPASS_LIGHTMAPPED_SHADER_TYPE(FMobileSingleLayerWaterPolicy, FMobileSingleLayerWaterPolicy);
 
 bool MaterialRequiresColorTransmittanceBlending(const FMaterial& MaterialResource)
 {
diff --git a/Engine/Source/Runtime/Renderer/Private/MobileBasePassRendering.h b/Engine/Source/Runtime/Renderer/Private/MobileBasePassRendering.h
--- a/Engine/Source/Runtime/Renderer/Private/MobileBasePassRendering.h
+++ b/Engine/Source/Runtime/Renderer/Private/MobileBasePassRendering.h
@@ -341,13 +341,14 @@ namespace MobileBasePass
 
 	FMobileLightMapPolicyTypeList GetUniformLightMapPolicyTypeForPSOCollection(bool bLitMaterial, bool bTranslucent/*not used*/, bool bUsesDeferredShading, bool bCanReceiveCSM, bool bMovable);
 
+	template <typename LightMapPolicyType>
 	bool GetShaders(
-		ELightMapPolicyType LightMapPolicyType,
+		const LightMapPolicyType& RESTRICT LightMapPolicy,
 		EMobileLocalLightSetting LocalLightSetting,
 		const FMaterial& MaterialResource,
 		const FVertexFactoryType* VertexFactoryType,
-		TShaderRef<TMobileBasePassVSPolicyParamType<FUniformLightMapPolicy>>& VertexShader,
-		TShaderRef<TMobileBasePassPSPolicyParamType<FUniformLightMapPolicy>>& PixelShader);
+		TShaderRef<TMobileBasePassVSPolicyParamType<LightMapPolicyType>>& VertexShader,
+		TShaderRef<TMobileBasePassPSPolicyParamType<LightMapPolicyType>>& PixelShader);
 
 	const FLightSceneInfo* GetDirectionalLightInfo(const FScene* Scene, const FPrimitiveSceneProxy* PrimitiveSceneProxy);
 
@@ -404,7 +405,6 @@ public:
 		// Translucent materials always support clustered shading on mobile deferred
 		const bool bForwardShading = !bDeferredShadingEnabled || bMaterialUsesForwardShading;
 				
-		TMobileBasePassPSBaseType<LightMapPolicyType>::ModifyCompilationEnvironment(Parameters, OutEnvironment);
 		// Only non-static skylights contribute into forward basepass
 		const bool bProjectSupportsNonStaticSkyLights = FReadOnlyCVARCache::EnableStationarySkylight() || !IsStaticLightingAllowed();
 		OutEnvironment.SetDefine(TEXT("ENABLE_SKY_LIGHT"), bIsLit && bForwardShading && bProjectSupportsNonStaticSkyLights);
@@ -438,10 +438,12 @@ public:
 		OutEnvironment.SetDefine(TEXT("MOBILE_TRANSLUCENT_COLOR_TRANSMITTANCE_DUAL_SRC_BLENDING"), TranslucentColorTransmittanceMode == EMobileTranslucentColorTransmittanceMode::DUAL_SRC_BLENDING ? 1u : 0u);
 		OutEnvironment.SetDefine(TEXT("MOBILE_TRANSLUCENT_COLOR_TRANSMITTANCE_PROGRAMMABLE_BLENDING"), TranslucentColorTransmittanceMode == EMobileTranslucentColorTransmittanceMode::PROGRAMMABLE_BLENDING ? 1u : 0u);
 		OutEnvironment.SetDefine(TEXT("MOBILE_TRANSLUCENT_COLOR_TRANSMITTANCE_SINGLE_SRC_BLENDING"), TranslucentColorTransmittanceMode == EMobileTranslucentColorTransmittanceMode::SINGLE_SRC_BLENDING ? 1u : 0u);
-		OutEnvironment.SetDefine(TEXT("MOBILE_SSR_ENABLED"), ((!bDeferredShadingEnabled && AreMobileScreenSpaceReflectionsEnabled(Parameters.Platform)) || (bIsSLW && MobileSupportsSM5MaterialNodes(Parameters.Platform))) ? 1u : 0u);
+		OutEnvironment.SetDefine(TEXT("MOBILE_SSR_ENABLED"), (!bDeferredShadingEnabled && AreMobileScreenSpaceReflectionsEnabled(Parameters.Platform)) ? 1u : 0u);
 
 		const bool bSupportCapsule = MobileSupportsSM5MaterialNodes(Parameters.Platform);
 		OutEnvironment.SetDefine(TEXT("MOBILE_SHADING_PATH_SUPPORT_CAPSULE_LIGHT"), bSupportCapsule ? 1u : 0u);
+
+		TMobileBasePassPSBaseType<LightMapPolicyType>::ModifyCompilationEnvironment(Parameters, OutEnvironment);
 	}
 	
 	/** Initialization constructor. */
@@ -505,6 +507,7 @@ private:
 
 	bool TryAddMeshBatch(const FMeshBatch& RESTRICT MeshBatch, uint64 BatchElementMask, const FPrimitiveSceneProxy* RESTRICT PrimitiveSceneProxy, int32 StaticMeshId, const FMaterialRenderProxy& MaterialRenderProxy, const FMaterial& Material);
 
+	template <typename LightMapPolicyType>
 	bool Process(
 		const FMeshBatch& RESTRICT MeshBatch,
 		uint64 BatchElementMask,
@@ -513,11 +516,10 @@ private:
 		const FMaterialRenderProxy& RESTRICT MaterialRenderProxy,
 		const FMaterial& RESTRICT MaterialResource,
 		const bool bIsMasked,
-		const bool bIsTranslucent,
 		FMaterialShadingModelField ShadingModels,
-		const ELightMapPolicyType LightMapPolicyType,
-		const bool bCanReceiveCSM,
-		const FUniformLightMapPolicy::ElementDataType& RESTRICT LightMapElementData);
+		const LightMapPolicyType& RESTRICT LightMapPolicy,
+		const typename LightMapPolicyType::ElementDataType& RESTRICT LightMapElementData,
+		const bool bCanReceiveCSM);
 
 	const ETranslucencyPass::Type TranslucencyPassType;
 	const EFlags Flags;
@@ -530,3 +532,42 @@ private:
 };
 
 ENUM_CLASS_FLAGS(FMobileBasePassMeshProcessor::EFlags);
+
+
+template <EMobileLocalLightSetting LocalLightSetting, typename LightMapPolicyType>
+bool GetMobileBasePassShaders(
+	const LightMapPolicyType& RESTRICT LightMapPolicy,
+	const FMaterial& Material,
+	const FVertexFactoryType* VertexFactoryType,
+	EMobileTranslucentColorTransmittanceMode ColoredTransmittanceFallback,
+	TShaderRef<TMobileBasePassVSPolicyParamType<LightMapPolicyType>>& VertexShader,
+	TShaderRef<TMobileBasePassPSPolicyParamType<LightMapPolicyType>>& PixelShader
+)
+{
+	using FVertexShaderType = TMobileBasePassVSPolicyParamType<LightMapPolicyType>;
+	using FPixelShaderType = TMobileBasePassPSPolicyParamType<LightMapPolicyType>;
+
+	FMaterialShaderTypes ShaderTypes;
+	ShaderTypes.AddShaderType<TMobileBasePassVS<LightMapPolicyType>>();
+
+	switch (ColoredTransmittanceFallback)
+	{
+	default:
+	case EMobileTranslucentColorTransmittanceMode::DEFAULT:
+		ShaderTypes.AddShaderType<TMobileBasePassPS<LightMapPolicyType, LocalLightSetting, EMobileTranslucentColorTransmittanceMode::DEFAULT>>();
+		break;
+	case EMobileTranslucentColorTransmittanceMode::SINGLE_SRC_BLENDING:
+		ShaderTypes.AddShaderType<TMobileBasePassPS<LightMapPolicyType, LocalLightSetting, EMobileTranslucentColorTransmittanceMode::SINGLE_SRC_BLENDING>>();
+		break;
+	}
+
+	FMaterialShaders Shaders;
+	if (!Material.TryGetShaders(ShaderTypes, VertexFactoryType, Shaders))
+	{
+		return false;
+	}
+
+	Shaders.TryGetVertexShader(VertexShader);
+	Shaders.TryGetPixelShader(PixelShader);
+	return true;
+}
diff --git a/Engine/Source/Runtime/Renderer/Private/MobileSingleLayerWaterRendering.cpp b/Engine/Source/Runtime/Renderer/Private/MobileSingleLayerWaterRendering.cpp
--- a/Engine/Source/Runtime/Renderer/Private/MobileSingleLayerWaterRendering.cpp
+++ b/Engine/Source/Runtime/Renderer/Private/MobileSingleLayerWaterRendering.cpp
@@ -56,18 +56,13 @@ private:
 		const FMaterialRenderProxy& RESTRICT MaterialRenderProxy,
 		const FMaterial& RESTRICT MaterialResource,
 		bool bIsMasked,
-		FMaterialShadingModelField ShadingModels,
-		ELightMapPolicyType LightMapPolicyType,
-		EMobileLocalLightSetting LocalLightSetting,
-		const FUniformLightMapPolicy::ElementDataType& RESTRICT LightMapElementData);
+		FMaterialShadingModelField ShadingModels);
 
 	void CollectPSOInitializersForLMPolicy(
 		const FPSOPrecacheVertexFactoryData& VertexFactoryData,
 		const FMeshPassProcessorRenderState& RESTRICT DrawRenderState,
 		const FGraphicsPipelineRenderTargetsInfo& RESTRICT RenderTargetsInfo,
 		const FMaterial& RESTRICT MaterialResource,
-		EMobileLocalLightSetting LocalLightSetting,
-		const ELightMapPolicyType LightMapPolicyType,
 		ERasterizerFillMode MeshFillMode,
 		ERasterizerCullMode MeshCullMode,
 		EPrimitiveType PrimitiveType,
@@ -86,7 +81,7 @@ void FMobileSingleLayerWaterPassMeshProcessor::AddMeshBatch(const FMeshBatch& RE
 {
 	// This MeshProcessor is only used on the mobile SM5 path, otherwise all SLW is handled by the translucent path.
 	EShaderPlatform ShaderPlatform = GetFeatureLevelShaderPlatform(FeatureLevel);
-	if (!MobileSupportsSM5MaterialNodes(ShaderPlatform))
+	if (IsUsingSingleLayerWaterFallback(ShaderPlatform))
 	{
 		return;
 	}
@@ -118,21 +113,7 @@ bool FMobileSingleLayerWaterPassMeshProcessor::TryAddMeshBatch(
 	if (ShadingModels.HasShadingModel(MSM_SingleLayerWater) && IsOpaqueOrMaskedBlendMode(Material))
 	{
 		const bool bIsMasked = IsMaskedBlendMode(Material);
-		ELightMapPolicyType LightmapPolicyType = MobileBasePass::SelectMeshLightmapPolicy(Scene, MeshBatch, PrimitiveSceneProxy,
-			true /*bCanReceiveCSM*/, false /*bPassUsesDeferredShading*/, true /*bIsLitMaterial*/, true /*bIsTranslucent*/);
-
-		EMobileLocalLightSetting LocalLightSetting = EMobileLocalLightSetting::LOCAL_LIGHTS_DISABLED;
-		if (Scene && PrimitiveSceneProxy)
-		{
-			// we can choose to use a single permutation regarless of local light state
-			// this is to avoid re-caching MDC on light state changes
-			if ((MobileLocalLightsUseSinglePermutation(Scene->GetShaderPlatform()) || PrimitiveSceneProxy->GetPrimitiveSceneInfo()->NumMobileDynamicLocalLights > 0))
-			{
-				LocalLightSetting = GetMobileForwardLocalLightSetting(Scene->GetShaderPlatform());
-			}
-		}
-
-		return Process(MeshBatch, BatchElementMask, PrimitiveSceneProxy, StaticMeshId, MaterialRenderProxy, Material, bIsMasked, ShadingModels, LightmapPolicyType, LocalLightSetting, MeshBatch.LCI);
+		return Process(MeshBatch, BatchElementMask, PrimitiveSceneProxy, StaticMeshId, MaterialRenderProxy, Material, bIsMasked, ShadingModels);
 	}
 
 	return true;
@@ -146,18 +127,15 @@ bool FMobileSingleLayerWaterPassMeshProcessor::Process(
 	const FMaterialRenderProxy& RESTRICT MaterialRenderProxy,
 	const FMaterial& RESTRICT MaterialResource,
 	bool bIsMasked,
-	FMaterialShadingModelField ShadingModels,
-	ELightMapPolicyType LightMapPolicyType,
-	EMobileLocalLightSetting LocalLightSetting,
-	const FUniformLightMapPolicy::ElementDataType& RESTRICT LightMapElementData)
+	FMaterialShadingModelField ShadingModels)
 {
 	TMeshProcessorShaders<
-		TMobileBasePassVSPolicyParamType<FUniformLightMapPolicy>,
-		TMobileBasePassPSPolicyParamType<FUniformLightMapPolicy>> BasePassShaders;
+		TMobileBasePassVSPolicyParamType<FMobileSingleLayerWaterPolicy>,
+		TMobileBasePassPSPolicyParamType<FMobileSingleLayerWaterPolicy>> BasePassShaders;
 
-	if (!MobileBasePass::GetShaders(
-		LightMapPolicyType,
-		LocalLightSetting,
+	if (!MobileBasePass::GetShaders<FMobileSingleLayerWaterPolicy>(
+		FMobileSingleLayerWaterPolicy(),
+		EMobileLocalLightSetting::LOCAL_LIGHTS_DISABLED,
 		MaterialResource,
 		MeshBatch.VertexFactory->GetType(),
 		BasePassShaders.VertexShader,
@@ -179,7 +157,7 @@ bool FMobileSingleLayerWaterPassMeshProcessor::Process(
 	ERasterizerFillMode MeshFillMode = ComputeMeshFillMode(MaterialResource, OverrideSettings);
 	ERasterizerCullMode MeshCullMode = ComputeMeshCullMode(MaterialResource, OverrideSettings);
 
-	TMobileBasePassShaderElementData<FUniformLightMapPolicy> ShaderElementData(MeshBatch.LCI, true /*bCanReceiveCSM*/);
+	TMobileBasePassShaderElementData<FMobileSingleLayerWaterPolicy> ShaderElementData(nullptr, true /*bCanReceiveCSM*/);
 	ShaderElementData.InitializeMeshMaterialData(ViewIfDynamicMeshCommand, PrimitiveSceneProxy, MeshBatch, StaticMeshId, false);
 
 	BuildMeshDrawCommands(
@@ -203,20 +181,18 @@ void FMobileSingleLayerWaterPassMeshProcessor::CollectPSOInitializersForLMPolicy
 	const FMeshPassProcessorRenderState& RESTRICT DrawRenderState,
 	const FGraphicsPipelineRenderTargetsInfo& RESTRICT RenderTargetsInfo,
 	const FMaterial& RESTRICT MaterialResource,
-	EMobileLocalLightSetting LocalLightSetting,
-	const ELightMapPolicyType LightMapPolicyType,
 	ERasterizerFillMode MeshFillMode,
 	ERasterizerCullMode MeshCullMode,
 	EPrimitiveType PrimitiveType,
 	TArray<FPSOPrecacheData>& PSOInitializers)
 {
 	TMeshProcessorShaders<
-		TMobileBasePassVSPolicyParamType<FUniformLightMapPolicy>,
-		TMobileBasePassPSPolicyParamType<FUniformLightMapPolicy>> BasePassShaders;
+		TMobileBasePassVSPolicyParamType<FMobileSingleLayerWaterPolicy>,
+		TMobileBasePassPSPolicyParamType<FMobileSingleLayerWaterPolicy>> BasePassShaders;
 
-	if (!MobileBasePass::GetShaders(
-		LightMapPolicyType,
-		LocalLightSetting,
+	if (!MobileBasePass::GetShaders<FMobileSingleLayerWaterPolicy>(
+		FMobileSingleLayerWaterPolicy(),
+		EMobileLocalLightSetting::LOCAL_LIGHTS_DISABLED,
 		MaterialResource,
 		VertexFactoryData.VertexFactoryType,
 		BasePassShaders.VertexShader,
@@ -248,9 +224,9 @@ void FMobileSingleLayerWaterPassMeshProcessor::CollectPSOInitializersForLMPolicy
 
 void FMobileSingleLayerWaterPassMeshProcessor::CollectPSOInitializers(const FSceneTexturesConfig& SceneTexturesConfig, const FMaterial& Material, const FPSOPrecacheVertexFactoryData& VertexFactoryData, const FPSOPrecacheParams& PreCacheParams, TArray<FPSOPrecacheData>& PSOInitializers)
 {
-	// This MeshProcessor is only used on the mobile SM5 path, otherwise all SLW is handled by the translucent path.
+	// Fallback SLW rendering is handled by the translucent path.
 	EShaderPlatform ShaderPlatform = GetFeatureLevelShaderPlatform(FeatureLevel);
-	if (!MobileSupportsSM5MaterialNodes(ShaderPlatform))
+	if (IsUsingSingleLayerWaterFallback(ShaderPlatform))
 	{
 		return;
 	}
@@ -290,20 +266,7 @@ void FMobileSingleLayerWaterPassMeshProcessor::CollectPSOInitializers(const FSce
 	}
 
 	SetMobileSingleLayerWaterRenderState(DrawRenderState, Material);
-
-	const EMobileLocalLightSetting LocalLightSetting = GetMobileForwardLocalLightSetting(ShaderPlatform);
-	const bool bUseLocalLightPermutation = (LocalLightSetting != EMobileLocalLightSetting::LOCAL_LIGHTS_DISABLED);
-
-	FMobileLightMapPolicyTypeList UniformLightMapPolicyTypes = MobileBasePass::GetUniformLightMapPolicyTypeForPSOCollection(bLitMaterial, true /*bTranslucent*/, false /*bUsesDeferredShading*/, true /*bCanReceiveCSM*/, bMovable);
-
-	for (ELightMapPolicyType LightMapPolicyType : UniformLightMapPolicyTypes)
-	{
-		CollectPSOInitializersForLMPolicy(VertexFactoryData, DrawRenderState, RenderTargetsInfo, Material, EMobileLocalLightSetting::LOCAL_LIGHTS_DISABLED, LightMapPolicyType, MeshFillMode, MeshCullMode, (EPrimitiveType)PreCacheParams.PrimitiveType, PSOInitializers);
-		if (bUseLocalLightPermutation)
-		{
-			CollectPSOInitializersForLMPolicy(VertexFactoryData, DrawRenderState, RenderTargetsInfo, Material, LocalLightSetting, LightMapPolicyType, MeshFillMode, MeshCullMode, (EPrimitiveType)PreCacheParams.PrimitiveType, PSOInitializers);
-		}
-	}
+	CollectPSOInitializersForLMPolicy(VertexFactoryData, DrawRenderState, RenderTargetsInfo, Material,  MeshFillMode, MeshCullMode, (EPrimitiveType)PreCacheParams.PrimitiveType, PSOInitializers);
 }
 
 FMeshPassProcessor* CreateMobileSingleLayerWaterPassProcessor(ERHIFeatureLevel::Type FeatureLevel, const FScene* Scene, const FSceneView* InViewIfDynamicMeshCommand, FMeshPassDrawListContext* InDrawListContext)
diff --git a/Engine/Source/Runtime/Renderer/Private/SceneVisibility.cpp b/Engine/Source/Runtime/Renderer/Private/SceneVisibility.cpp
--- a/Engine/Source/Runtime/Renderer/Private/SceneVisibility.cpp
+++ b/Engine/Source/Runtime/Renderer/Private/SceneVisibility.cpp
@@ -1429,7 +1429,7 @@ void FRelevancePacket::ComputeRelevance(FDynamicPrimitiveIndexList& DynamicPrimi
 	const bool bIsTranslucentHoldoutEnabled = IsPrimitiveAlphaHoldoutEnabled(View);
 	const bool bIsTranslucentClippedDepthEnabled = NeedTSRThinGeometryDetectionPass(View);
 	const bool bHLODActive = Scene.SceneLODHierarchy.IsActive();
-	const bool bMobileSupportsSM5MaterialNodes = MobileSupportsSM5MaterialNodes(View.GetShaderPlatform());
+	const bool bUsingSingleLayerWaterFallback = IsUsingSingleLayerWaterFallback(View.GetShaderPlatform());
 	const FHLODVisibilityState* const HLODState = bHLODActive && ViewState ? &ViewState->HLODVisibilityState : nullptr;
 	float MaxDrawDistanceScale = GetCachedScalabilityCVars().ViewDistanceScale;
 	MaxDrawDistanceScale *= GetCachedScalabilityCVars().CalculateFieldOfViewDistanceScale(View.DesiredFOV);
@@ -1690,7 +1690,7 @@ void FRelevancePacket::ComputeRelevance(FDynamicPrimitiveIndexList& DynamicPrimi
 								// Specific logic for mobile packets
 								if (ShadingPath == EShadingPath::Mobile)
 								{
-									if (StaticMeshRelevance.bUseSingleLayerWaterMaterial && bMobileSupportsSM5MaterialNodes)
+									if (StaticMeshRelevance.bUseSingleLayerWaterMaterial && !bUsingSingleLayerWaterFallback)
 									{
 										DrawCommandPacket.AddCommandsForMesh(PrimitiveIndex, PrimitiveSceneInfo, StaticMeshRelevance, StaticMesh, CullingPayloadFlags, Scene, bCanCache, EMeshPass::SingleLayerWaterPass);
 									}
diff --git a/Engine/Source/Runtime/Renderer/Private/SingleLayerWaterRendering.cpp b/Engine/Source/Runtime/Renderer/Private/SingleLayerWaterRendering.cpp
--- a/Engine/Source/Runtime/Renderer/Private/SingleLayerWaterRendering.cpp
+++ b/Engine/Source/Runtime/Renderer/Private/SingleLayerWaterRendering.cpp
@@ -24,6 +24,7 @@
 #include "UnrealEngine.h"
 #include "CustomDepthRendering.h"
 #include "HairStrands/HairStrandsData.h"
+#include "ComponentRecreateRenderStateContext.h"
 
 #include "DepthCopy.h"
 using namespace DepthCopy;
@@ -195,6 +196,15 @@ static TAutoConsoleVariable<int32> CVarWaterSingleLayerTiledSceneColorCopy(
 	TEXT("Use indirect draws and a list of water pixel tiles to copy only relevant parts of the SceneColor texture for refraction behind the water surface."),
 	ECVF_RenderThreadSafe);
 
+static TAutoConsoleVariable<int32> CVarWaterSingleLayerMobileFallback(
+	TEXT("r.Water.SingleLayer.MobileFallback"), 1,
+	TEXT("Whether to use a translucent rendering fallback for a single layer water on mobile platforms"),
+	FConsoleVariableDelegate::CreateLambda([](IConsoleVariable* InVariable) 
+		{ 
+			FGlobalComponentRecreateRenderStateContext Context; 
+		}),
+	ECVF_Scalability | ECVF_RenderThreadSafe);
+
 static int32 GetSingleLayerWaterReflectionTechnique()
 {
 	const int32 Value = CVarWaterSingleLayerReflection.GetValueOnRenderThread();
diff --git a/Engine/Source/Runtime/Renderer/Private/SkyPassRendering.cpp b/Engine/Source/Runtime/Renderer/Private/SkyPassRendering.cpp
--- a/Engine/Source/Runtime/Renderer/Private/SkyPassRendering.cpp
+++ b/Engine/Source/Runtime/Renderer/Private/SkyPassRendering.cpp
@@ -115,7 +115,7 @@ bool FSkyPassMeshProcessor::Process(
 			TMobileBasePassPSPolicyParamType<LightMapPolicyType>> SkyPassShaders;
 
 		if (!MobileBasePass::GetShaders(
-			LMP_NO_LIGHTMAP,
+			NoLightmapPolicy,
 			EMobileLocalLightSetting::LOCAL_LIGHTS_DISABLED,
 			MaterialResource,
 			VertexFactory->GetType(),
@@ -191,14 +191,14 @@ void FSkyPassMeshProcessor::CollectPSOInitializers(const FSceneTexturesConfig& S
 	}
 
 	typedef FUniformLightMapPolicy LightMapPolicyType;
+	FUniformLightMapPolicy NoLightmapPolicy(LMP_NO_LIGHTMAP);
 
 	if (GetFeatureLevelShadingPath(FeatureLevel) == EShadingPath::Deferred)
 	{
 		TMeshProcessorShaders<
 			TBasePassVertexShaderPolicyParamType<LightMapPolicyType>,
 			TBasePassPixelShaderPolicyParamType<LightMapPolicyType>> SkyPassShaders;
-
-		FUniformLightMapPolicy NoLightmapPolicy(LMP_NO_LIGHTMAP);
+				
 		const bool bRenderSkylight = false;
 		if (GetBasePassShaders<LightMapPolicyType>(
 			Material,
@@ -223,7 +223,7 @@ void FSkyPassMeshProcessor::CollectPSOInitializers(const FSceneTexturesConfig& S
 			TMobileBasePassPSPolicyParamType<LightMapPolicyType>> SkyPassShaders;
 
 		if (MobileBasePass::GetShaders(
-			LMP_NO_LIGHTMAP,
+			NoLightmapPolicy,
 			EMobileLocalLightSetting::LOCAL_LIGHTS_DISABLED,
 			Material,
 			VertexFactoryData.VertexFactoryType,
