Move definition of r.PrimitiveHasTileOffsetData to fix TSAN warning caused by bTrackFrequentCalls.

#rnx
#Synced-CL 49727255

diff --git a/Engine/Source/Runtime/Core/Private/Misc/LargeWorldRenderPosition.cpp b/Engine/Source/Runtime/Core/Private/Misc/LargeWorldRenderPosition.cpp
--- a/Engine/Source/Runtime/Core/Private/Misc/LargeWorldRenderPosition.cpp
+++ b/Engine/Source/Runtime/Core/Private/Misc/LargeWorldRenderPosition.cpp
@@ -6,6 +6,7 @@
 #include "Math/Matrix.inl"
 #include "Math/TranslationMatrix.h"
 #include "Misc/AssertionMacros.h"
+#include "HAL/IConsoleManager.h"
 
 #define UE_LWC_RENDER_TILE_SIZE_MIN (262144.0)	// LWC_TODO: Make this smaller?
 
@@ -137,4 +138,12 @@ CORE_API void TLargeWorldRenderScalar<double>::Validate(double InAbsolute)
 }
 
 template struct TLargeWorldRenderScalar<double>;
-template struct TLargeWorldRenderScalar<float>;
\ No newline at end of file
+template struct TLargeWorldRenderScalar<float>;
+
+CORE_API int32 GPrimitiveHasTileOffsetData = 1;
+static FAutoConsoleVariableRef CVarPrimitiveHasTileOffsetData(
+	TEXT("r.PrimitiveHasTileOffsetData"),
+	GPrimitiveHasTileOffsetData,
+	TEXT("1 to upload lower-precision tileoffset primitive data to gpu, 0 to use higher-precision double float.\n"),
+	ECVF_ReadOnly | ECVF_RenderThreadSafe
+);
diff --git a/Engine/Source/Runtime/Core/Public/Misc/LargeWorldRenderPosition.h b/Engine/Source/Runtime/Core/Public/Misc/LargeWorldRenderPosition.h
--- a/Engine/Source/Runtime/Core/Public/Misc/LargeWorldRenderPosition.h
+++ b/Engine/Source/Runtime/Core/Public/Misc/LargeWorldRenderPosition.h
@@ -156,3 +156,6 @@ private:
 	VectorType Offset;
 };
 using FLargeWorldRenderPosition = TLargeWorldRenderPosition<double>;
+
+/* See r.PrimitiveHasTileOffsetData */
+extern CORE_API int32 GPrimitiveHasTileOffsetData;
diff --git a/Engine/Source/Runtime/Engine/Private/ShaderCompiler/ShaderCompiler.cpp b/Engine/Source/Runtime/Engine/Private/ShaderCompiler/ShaderCompiler.cpp
--- a/Engine/Source/Runtime/Engine/Private/ShaderCompiler/ShaderCompiler.cpp
+++ b/Engine/Source/Runtime/Engine/Private/ShaderCompiler/ShaderCompiler.cpp
@@ -3526,9 +3526,7 @@ void GlobalBeginCompileShader(
 	const bool bViewHasTileOffsetData = CVarViewHasTileOffsetData->GetValueOnAnyThread() != 0;
 	SET_SHADER_DEFINE_AND_COMPILE_ARGUMENT(Input.Environment, VIEW_HAS_TILEOFFSET_DATA, bViewHasTileOffsetData);
 
-	static const auto CVarPrimitiveHasTileOffsetData = IConsoleManager::Get().FindTConsoleVariableDataInt(TEXT("r.PrimitiveHasTileOffsetData"));
-	const bool bPrimitiveHasTileOffsetData = CVarPrimitiveHasTileOffsetData->GetValueOnAnyThread() != 0;
-	SET_SHADER_DEFINE_AND_COMPILE_ARGUMENT(Input.Environment, PRIMITIVE_HAS_TILEOFFSET_DATA, bPrimitiveHasTileOffsetData);
+	SET_SHADER_DEFINE_AND_COMPILE_ARGUMENT(Input.Environment, PRIMITIVE_HAS_TILEOFFSET_DATA, GPrimitiveHasTileOffsetData);
 
 	// Set VR definitions
 	if (ShaderFrequencyNeedsInstancedStereoMods(ShaderType))
diff --git a/Engine/Source/Runtime/Engine/Public/PrimitiveUniformShaderParametersBuilder.h b/Engine/Source/Runtime/Engine/Public/PrimitiveUniformShaderParametersBuilder.h
--- a/Engine/Source/Runtime/Engine/Public/PrimitiveUniformShaderParametersBuilder.h
+++ b/Engine/Source/Runtime/Engine/Public/PrimitiveUniformShaderParametersBuilder.h
@@ -352,10 +352,7 @@ public:
 			}
 		}
 
-		static TConsoleVariableData<int32>* CVarPrimitiveHasTileOffsetData = nullptr;
-		CVarPrimitiveHasTileOffsetData = CVarPrimitiveHasTileOffsetData ? CVarPrimitiveHasTileOffsetData : IConsoleManager::Get().FindTConsoleVariableDataInt(TEXT("r.PrimitiveHasTileOffsetData")); // null at first
-		const bool bPrimitiveHasTileOffsetData = CVarPrimitiveHasTileOffsetData ? (CVarPrimitiveHasTileOffsetData->GetValueOnAnyThread() != 0) : false;
-		if (bPrimitiveHasTileOffsetData)
+		if (GPrimitiveHasTileOffsetData)
 		{
 			const FLargeWorldRenderPosition AbsoluteActorWorldPositionTO { AbsoluteActorWorldPosition };
 			Parameters.ActorWorldPositionHigh = AbsoluteActorWorldPositionTO.GetTile();
diff --git a/Engine/Source/Runtime/Renderer/Private/InstanceCulling/InstanceCullingContext.cpp b/Engine/Source/Runtime/Renderer/Private/InstanceCulling/InstanceCullingContext.cpp
--- a/Engine/Source/Runtime/Renderer/Private/InstanceCulling/InstanceCullingContext.cpp
+++ b/Engine/Source/Runtime/Renderer/Private/InstanceCulling/InstanceCullingContext.cpp
@@ -560,9 +560,7 @@ public:
 		EBatchProcessingMode ProcessingMode = (PermutationVector.Get<FSingleInstanceModeDim>() ? EBatchProcessingMode::UnCulled : EBatchProcessingMode::Generic);
 		OutEnvironment.SetDefine(TEXT("INSTANCE_DATA_STRIDE_ELEMENTS"), GetInstanceDataStrideElements(Parameters.Platform, ProcessingMode));
 
-		static const auto CVarPrimitiveHasTileOffsetData = IConsoleManager::Get().FindTConsoleVariableDataInt(TEXT("r.PrimitiveHasTileOffsetData"));
-		const bool bPrimitiveHasTileOffsetData = CVarPrimitiveHasTileOffsetData->GetValueOnAnyThread() != 0;
-		OutEnvironment.SetDefine(TEXT("PRIMITIVE_HAS_TILEOFFSET_DATA"), bPrimitiveHasTileOffsetData ? 1 : 0);
+		OutEnvironment.SetDefine(TEXT("PRIMITIVE_HAS_TILEOFFSET_DATA"), GPrimitiveHasTileOffsetData);
 
 		DynamicMeshBoundsModifyCompilationEnvironment(Parameters, OutEnvironment);
 	}
diff --git a/Engine/Source/Runtime/Renderer/Private/SceneRendering.cpp b/Engine/Source/Runtime/Renderer/Private/SceneRendering.cpp
--- a/Engine/Source/Runtime/Renderer/Private/SceneRendering.cpp
+++ b/Engine/Source/Runtime/Renderer/Private/SceneRendering.cpp
@@ -300,12 +300,6 @@ static TAutoConsoleVariable<int32> CVarViewHasTileOffsetData(
 	TEXT("1 to upload lower-precision tileoffset view data to gpu, 0 to use only higher-precision double float.\n"),
 	ECVF_ReadOnly | ECVF_RenderThreadSafe);
 
-static TAutoConsoleVariable<int32> CVarPrimitiveHasTileOffsetData(
-	TEXT("r.PrimitiveHasTileOffsetData"),
-	1,
-	TEXT("1 to upload lower-precision tileoffset primitive data to gpu, 0 to use higher-precision double float.\n"),
-	ECVF_ReadOnly | ECVF_RenderThreadSafe);
-
 static TAutoConsoleVariable<bool> CVarPrimitiveAlphaHoldoutSupport(
 	TEXT("r.Deferred.SupportPrimitiveAlphaHoldout"),
 	false,
