Modified r.SplitScreenDebug logic to store "second view state" in SourceViewState instead of using a global GSecondViewState.
- the global view state caused issues when we have multiple viewports active and all the "second view" created by this debug mode were referencing the same GSecondViewState.

#Synced-CL 49737718

diff --git a/Engine/Source/Runtime/Renderer/Private/SceneOcclusion.cpp b/Engine/Source/Runtime/Renderer/Private/SceneOcclusion.cpp
--- a/Engine/Source/Runtime/Renderer/Private/SceneOcclusion.cpp
+++ b/Engine/Source/Runtime/Renderer/Private/SceneOcclusion.cpp
@@ -307,6 +307,12 @@ bool FSceneViewState::IsShadowOccluded(FSceneViewState::FProjectedShadowKey Shad
 
 void FSceneViewState::Destroy()
 {
+	if (SplitScreenDebugViewState != nullptr)
+	{
+		SplitScreenDebugViewState->Destroy();
+		SplitScreenDebugViewState = nullptr;
+	}
+
 	FSceneViewState* self = this;
 	ENQUEUE_RENDER_COMMAND(FSceneViewState_Destroy)(
 	[self](FRHICommandListImmediate& RHICmdList)
diff --git a/Engine/Source/Runtime/Renderer/Private/ScenePrivate.h b/Engine/Source/Runtime/Renderer/Private/ScenePrivate.h
--- a/Engine/Source/Runtime/Renderer/Private/ScenePrivate.h
+++ b/Engine/Source/Runtime/Renderer/Private/ScenePrivate.h
@@ -1118,6 +1118,8 @@ public:
 
 	FTemporalLODState TemporalLODState;
 
+	FSceneViewState* SplitScreenDebugViewState = nullptr;
+
 	FVolumetricRenderTargetViewStateData VolumetricCloudRenderTarget;
 	FTemporalRenderTargetState VolumetricCloudShadowRenderTarget[NUM_ATMOSPHERE_LIGHTS];
 	FMatrix VolumetricCloudShadowmapPreviousTranslatedWorldToLightClipMatrix[NUM_ATMOSPHERE_LIGHTS];
diff --git a/Engine/Source/Runtime/Renderer/Private/SceneRenderBuilder.cpp b/Engine/Source/Runtime/Renderer/Private/SceneRenderBuilder.cpp
--- a/Engine/Source/Runtime/Renderer/Private/SceneRenderBuilder.cpp
+++ b/Engine/Source/Runtime/Renderer/Private/SceneRenderBuilder.cpp
@@ -222,27 +222,29 @@ static void CreateSplitScreenDebugViewFamilies(const FSceneViewFamily& InFamily,
 		// Use new static view state for second view
 		if (ViewIndex == 1)
 		{
-			static FSceneViewState* GSecondViewState = nullptr;
-			if (!GSecondViewState)
+			FSceneViewState* SourceViewState = InFamily.Views[0]->State->GetConcreteViewState();
+			
+			if (SourceViewState->SplitScreenDebugViewState == nullptr)
 			{
-				GSecondViewState = new FSceneViewState(InFamily.GetFeatureLevel(), nullptr);
+				SourceViewState->SplitScreenDebugViewState = new FSceneViewState(InFamily.GetFeatureLevel(), nullptr);
 			}
 
+			FSceneViewState* SecondViewState = SourceViewState->SplitScreenDebugViewState;
+
 			// Propagate this user writable field between FSceneViewState
-			const FSceneViewState* SourceViewState = InFamily.Views[0]->State->GetConcreteViewState();
-			GSecondViewState->SequencerState = SourceViewState->SequencerState;
+			SecondViewState->SequencerState = SourceViewState->SequencerState;
 
 			// Add or remove optional Lumen scene for second view state
 			if (CVarSplitScreenDebugLumenScene.GetValueOnGameThread())
 			{
-				GSecondViewState->AddLumenSceneData(ViewParents[ViewIndex]->Scene, 1.0f);
+				SecondViewState->AddLumenSceneData(ViewParents[ViewIndex]->Scene, 1.0f);
 			}
 			else
 			{
-				GSecondViewState->RemoveLumenSceneData(ViewParents[ViewIndex]->Scene);
+				SecondViewState->RemoveLumenSceneData(ViewParents[ViewIndex]->Scene);
 			}
 
-			InitOptions.SceneViewStateInterface = GSecondViewState;
+			InitOptions.SceneViewStateInterface = SecondViewState;
 		}
 
 		FSceneView* View = new FSceneView(InitOptions);
