Removed dead FDistortionPrimSet data structure.

#rnx
#Synced-CL 49761124

diff --git a/Engine/Source/Runtime/Renderer/Private/DistortionRendering.h b/Engine/Source/Runtime/Renderer/Private/DistortionRendering.h
--- a/Engine/Source/Runtime/Renderer/Private/DistortionRendering.h
+++ b/Engine/Source/Runtime/Renderer/Private/DistortionRendering.h
@@ -13,52 +13,6 @@
 class FPrimitiveSceneProxy;
 struct FMeshPassProcessorRenderState;
 
-/** 
-* Set of distortion scene prims  
-*/
-class FDistortionPrimSet
-{
-public:
-
-	/** 
-	* Iterate over the distortion prims and draw their accumulated offsets
-	* @param ViewInfo - current view used to draw items
-	* @param DPGIndex - current DPG used to draw items
-	* @return true if anything was drawn
-	*/
-	bool DrawAccumulatedOffsets(FRHICommandListImmediate& RHICmdList, const class FViewInfo& View, const FMeshPassProcessorRenderState& DrawRenderState, bool bInitializeOffsets);
-
-	/**
-	* Adds a new primitives to the list of distortion prims
-	* @param PrimitiveSceneProxies - primitive info to add.
-	*/
-	void Append(FPrimitiveSceneProxy** PrimitiveSceneProxies, int32 NumProxies)
-	{
-		Prims.Append(PrimitiveSceneProxies, NumProxies);
-	}
-
-	/** 
-	* @return number of prims to render
-	*/
-	int32 NumPrims() const
-	{
-		return Prims.Num();
-	}
-
-	/** 
-	* @return a prim currently set to render
-	*/
-	const FPrimitiveSceneProxy* GetPrim(int32 i)const
-	{
-		check(i>=0 && i<NumPrims());
-		return Prims[i];
-	}
-
-private:
-	/** list of distortion prims added from the scene */
-	TArray<FPrimitiveSceneProxy*, SceneRenderingAllocator> Prims;
-};
-
 extern void SetupDistortionParams(FVector4f& DistortionParams, const FViewInfo& View);
 
 class FDistortionMeshProcessor : public FSceneRenderingAllocatorObject<FDistortionMeshProcessor>, public FMeshPassProcessor
