Add ubershader permutation to FDeferredLightPS.

The ubershader allows to start rendering earlier (with a more expensive shader) while requested PSO get asynced compile in background.

Work done jointly with Kenzo.

#Synced-CL 49758930

diff --git a/Engine/Shaders/Private/DeferredLightPixelShaders.usf b/Engine/Shaders/Private/DeferredLightPixelShaders.usf
--- a/Engine/Shaders/Private/DeferredLightPixelShaders.usf
+++ b/Engine/Shaders/Private/DeferredLightPixelShaders.usf
@@ -112,10 +112,11 @@ FDerivedParams GetDerivedParams(in FInputParams Input, in float SceneDepth)
 
 Texture2D<uint> LightingChannelsTexture;
 
+uint bLightingChannelEnabled;
 uint GetLightingChannelMask(float2 UV)
 {
 	uint2 IntegerUV = UV * View.BufferSizeAndInvSize.xy;
-	return LightingChannelsTexture.Load(uint3(IntegerUV, 0)).x;
+	return bLightingChannelEnabled ? LightingChannelsTexture.Load(uint3(IntegerUV, 0)).x : 0xFFu;
 }
 
 float GetExposure()
@@ -171,17 +172,25 @@ float3 GetAtmosphericLightTransmittance(in float3 StableTranslatedLargeWorldPosi
 
 #endif // USE_ATMOSPHERE_TRANSMITTANCE
 
+uint bCloudShadowEnabled;
+uint bAtmosphereTransmittanceEnabled;
 float3 GetAtmosphereAndCloudAttenuation(float3 StableTranslatedLargeWorldPosition, float2 ScreenUV)
 {
 	float3 Out = 1;
 
 	#if USE_ATMOSPHERE_TRANSMITTANCE
-	Out *= GetAtmosphericLightTransmittance(StableTranslatedLargeWorldPosition, ScreenUV, DeferredLightUniforms.Direction.xyz);
+	if (bAtmosphereTransmittanceEnabled)
+	{
+		Out *= GetAtmosphericLightTransmittance(StableTranslatedLargeWorldPosition, ScreenUV, DeferredLightUniforms.Direction.xyz);
+	}
 	#endif
 
 	#if USE_CLOUD_TRANSMITTANCE
-	float OutOpticalDepth = 0.0f;
-	Out *= lerp(1.0f, GetCloudVolumetricShadow(StableTranslatedLargeWorldPosition, CloudShadowmapTranslatedWorldToLightClipMatrix, CloudShadowmapFarDepthKm, CloudShadowmapTexture, CloudShadowmapSampler, OutOpticalDepth), CloudShadowmapStrength);
+	if (bCloudShadowEnabled)
+	{
+		float OutOpticalDepth = 0.0f;
+		Out *= lerp(1.0f, GetCloudVolumetricShadow(StableTranslatedLargeWorldPosition, CloudShadowmapTranslatedWorldToLightClipMatrix, CloudShadowmapFarDepthKm, CloudShadowmapTexture, CloudShadowmapSampler, OutOpticalDepth), CloudShadowmapStrength);
+	}
 	#endif
 
 	return Out;
diff --git a/Engine/Source/Runtime/Renderer/Private/LightRendering.cpp b/Engine/Source/Runtime/Renderer/Private/LightRendering.cpp
--- a/Engine/Source/Runtime/Renderer/Private/LightRendering.cpp
+++ b/Engine/Source/Runtime/Renderer/Private/LightRendering.cpp
@@ -915,6 +915,7 @@ class FDeferredLightPS : public FGlobalShader
 	class FCloudTransmittance 	: SHADER_PERMUTATION_BOOL("USE_CLOUD_TRANSMITTANCE");
 	class FAnistropicMaterials 	: SHADER_PERMUTATION_BOOL("SUPPORTS_ANISOTROPIC_MATERIALS");
 	class FVirtualShadowMapMask : SHADER_PERMUTATION_BOOL("USE_VIRTUAL_SHADOW_MAP_MASK");
+	class FUbershader 			: SHADER_PERMUTATION_BOOL("USE_UBERSHADER");
 
 	using FPermutationDomain = TShaderPermutationDomain<
 		FSourceShapeDim,
@@ -930,7 +931,8 @@ class FDeferredLightPS : public FGlobalShader
 		FCloudTransmittance,
 		FAnistropicMaterials,
 		Substrate::FSubstrateTileType,
-		FVirtualShadowMapMask>;
+		FVirtualShadowMapMask,
+		FUbershader>;
 
 	BEGIN_SHADER_PARAMETER_STRUCT(FParameters, )
 		SHADER_PARAMETER_RDG_UNIFORM_BUFFER(FSceneTextureUniformParameters, SceneTextures)
@@ -939,7 +941,9 @@ class FDeferredLightPS : public FGlobalShader
 		SHADER_PARAMETER_STRUCT_INCLUDE(FVolumetricCloudShadowAOParameters, CloudShadowAO)
 		SHADER_PARAMETER_STRUCT_INCLUDE(FLightCloudTransmittanceParameters, CloudShadow)
 		SHADER_PARAMETER_RDG_UNIFORM_BUFFER(FForwardLightUniformParameters, ForwardLightStruct)
-		SHADER_PARAMETER(uint32, CloudShadowEnabled)
+		SHADER_PARAMETER(uint32, bLightingChannelEnabled)
+		SHADER_PARAMETER(uint32, bCloudShadowEnabled)
+		SHADER_PARAMETER(uint32, bAtmosphereTransmittanceEnabled)
 		SHADER_PARAMETER(uint32, HairTransmittanceBufferMaxCount)
 		SHADER_PARAMETER(uint32, HairShadowMaskValid)
 		SHADER_PARAMETER(FVector4f, ShadowChannelMask)
@@ -986,6 +990,30 @@ class FDeferredLightPS : public FGlobalShader
 			return false;
 		}
 
+		// Uber permutation - Enabled all the features with dynamic branching - 6 Permutations
+		if (PermutationVector.Get< FUbershader >())
+		{
+			const ESubstrateTileType UberShaderTileType = Substrate::IsSubstrateEnabled() ? (Substrate::ShouldCompileSubstrateTileTypePermutations(ESubstrateTileType::EComplexSpecial, Parameters.Platform) ? ESubstrateTileType::EComplexSpecial : ESubstrateTileType::EComplex) : ESubstrateTileType::ESimple;
+			if (PermutationVector.Get< FSourceTextureDim >() &&
+				PermutationVector.Get< FIESProfileDim >() &&
+				PermutationVector.Get< FLightFunctionAtlasDim >() &&
+				PermutationVector.Get< FLightingChannelsDim >() &&
+				PermutationVector.Get< FTransmissionDim >() &&
+				PermutationVector.Get< FHairComplexTransmittance >() &&
+				PermutationVector.Get< FAtmosphereTransmittance >() &&
+				PermutationVector.Get< FCloudTransmittance >() &&
+				PermutationVector.Get< FAnistropicMaterials >() &&
+				PermutationVector.Get< FVirtualShadowMapMask >() &&
+				PermutationVector.Get< Substrate::FSubstrateTileType >() == UberShaderTileType)
+			{
+				return IsFeatureLevelSupported(Parameters.Platform, ERHIFeatureLevel::SM5);
+			}
+			else
+			{
+				return false;
+			}
+		}
+
 		if (PermutationVector.Get< FSourceShapeDim >() == ELightSourceShape::Directional && PermutationVector.Get< FIESProfileDim >())
 		{
 			return false;
@@ -1085,7 +1113,7 @@ class FDeferredLightPS : public FGlobalShader
 		return IsFeatureLevelSupported(Parameters.Platform, ERHIFeatureLevel::SM5);
 	}
 
-	static FPermutationDomain RemapPermutation(FPermutationDomain PermutationVector)
+	static FPermutationDomain RemapPermutation(FPermutationDomain PermutationVector, EShaderPlatform Platform)
 	{
 		// Build FVisualizeCullingDim permutation only for a restricted number of case, as they don't impact the 'estimated cost' of lighting
 		if (PermutationVector.Get< FVisualizeCullingDim >())
@@ -1117,6 +1145,26 @@ class FDeferredLightPS : public FGlobalShader
 			}
 		}
 
+		if (PermutationVector.Get< FUbershader>())
+		{
+			const ESubstrateTileType UberShaderTileType = Substrate::IsSubstrateEnabled() ? (Substrate::ShouldCompileSubstrateTileTypePermutations(ESubstrateTileType::EComplexSpecial, Platform) ? ESubstrateTileType::EComplexSpecial : ESubstrateTileType::EComplex) : ESubstrateTileType::ESimple;
+			// Remaining permutation (i.e. not dynamic branching)
+			//  - FSourceShapeDim = 3
+			//  - FHairLighting  = 2
+			//  - FVisualizeCullingDim = 0
+			PermutationVector.Set< FSourceTextureDim >(true);
+			PermutationVector.Set< FIESProfileDim >(true);
+			PermutationVector.Set< FLightFunctionAtlasDim >(true);
+			PermutationVector.Set< FLightingChannelsDim >(true);
+			PermutationVector.Set< FTransmissionDim >(true);
+			PermutationVector.Set< FHairComplexTransmittance >(true);
+			PermutationVector.Set< FAtmosphereTransmittance >(true);
+			PermutationVector.Set< FCloudTransmittance >(true);
+			PermutationVector.Set< FAnistropicMaterials >(true);
+			PermutationVector.Set< FVirtualShadowMapMask >(true);
+			PermutationVector.Set< Substrate::FSubstrateTileType >(UberShaderTileType);
+		}
+
 		return PermutationVector;
 	}
 
@@ -2598,10 +2646,12 @@ static FDeferredLightPS::FParameters GetDeferredLightPSParameters(
 	Out.HairStrands = HairStrandsUniformBuffer;
 	Out.ForwardLightStruct = View.ForwardLightingResources.ForwardLightUniformBuffer;
 	Out.Substrate = Substrate::BindSubstrateGlobalUniformParameters(View);
+	Out.bLightingChannelEnabled = LightingChannelsTexture ? 1u : 0u;
 	Out.LightingChannelsTexture = LightingChannelsTexture ? LightingChannelsTexture : WhiteDummy;
 	Out.LightingChannelsSampler = TStaticSamplerState<SF_Point, AM_Clamp, AM_Clamp, AM_Clamp>::GetRHI();
 	Out.CloudShadowAO = GetCloudShadowAOParameters(GraphBuilder, View, CloudInfo);
-	Out.CloudShadowEnabled = SetupLightCloudTransmittanceParameters(GraphBuilder, Scene, View, LightSceneInfo, Out.CloudShadow) ? 1 : 0;
+	Out.bCloudShadowEnabled = SetupLightCloudTransmittanceParameters(GraphBuilder, Scene, View, LightSceneInfo, Out.CloudShadow) ? 1 : 0;
+	Out.bAtmosphereTransmittanceEnabled = IsLightAtmospherePerPixelTransmittanceEnabled(Scene, View, LightSceneInfo) ? 1 : 0;
 	Out.LightAttenuationTexture = ShadowMaskTexture ? ShadowMaskTexture : WhiteDummy;
 	Out.LightAttenuationTextureSampler = TStaticSamplerState<SF_Point, AM_Wrap, AM_Wrap, AM_Wrap>::GetRHI();
 	Out.View = View.ViewUniformBuffer;
@@ -2633,6 +2683,12 @@ static FDeferredLightPS::FParameters GetDeferredLightPSParameters(
 		Out.RenderTargets.DepthStencil = FDepthStencilBinding(SceneDepthTexture, ERenderTargetLoadAction::ELoad, ERenderTargetLoadAction::ELoad, FExclusiveDepthStencil::DepthRead_StencilWrite);
 	}
 
+	if (!VirtualShadowMapUniformBuffer)
+	{
+		Out.VirtualShadowMap = FVirtualShadowMapArray::CreateDummySamplingParameters(GraphBuilder).VirtualShadowMap;
+		Out.VirtualShadowMapId = INDEX_NONE;
+	}
+
 	return Out;
 }
 
@@ -2688,6 +2744,27 @@ uint32 SetupLightGraphicsPSOState(
 	return StencilRef;
 }
 
+TShaderMapRef<FDeferredLightPS> GetUberDeferredLightPS(FGlobalShaderMap* GlobalShaderMap, EShaderPlatform ShaderPlatform, const FDeferredLightPS::FPermutationDomain& PermutationVector)
+{
+	FDeferredLightPS::FPermutationDomain UberPermutationVector = PermutationVector;
+	UberPermutationVector.Set < FDeferredLightPS::FUbershader > (true);
+	UberPermutationVector = FDeferredLightPS::RemapPermutation(UberPermutationVector, ShaderPlatform);
+	TShaderMapRef<FDeferredLightPS> UberPixelShader(GlobalShaderMap, UberPermutationVector);
+	check(UberPixelShader.IsValid());
+	return UberPixelShader;
+}
+
+void ValidatePSOPrecacheState(const FGraphicsPipelineStateInitializer& GraphicsPSOInit)
+{
+#if PSO_PRECACHING_VALIDATE
+	if (PSOCollectorStats::IsFullPrecachingValidationEnabled())
+	{
+		static const int32 GlobalPSOCollectorIndex = FGlobalPSOCollectorManager::GetIndex(DeferredLightGlobalPSOCollectorName);
+		PSOCollectorStats::CheckGlobalGraphicsPipelineStateInCache(GraphicsPSOInit, GlobalPSOCollectorIndex);
+	}
+#endif // PSO_PRECACHING_VALIDATE
+}
+
 // Used by RenderLights to render a light to the scene color buffer.
 template<typename TShaderType, typename TParametersType>
 static void InternalRenderLight(
@@ -2696,6 +2773,7 @@ static void InternalRenderLight(
 	const FViewInfo& View,
 	const FLightSceneInfo* LightSceneInfo,
 	TShaderType& PixelShader,
+	TShaderType& UberPixelShader,
 	TParametersType* PassParameters,
 	ESubstrateTileType SubstrateTileMaterialType,
 	const TCHAR* ShaderName)
@@ -2709,9 +2787,8 @@ static void InternalRenderLight(
 		RDG_EVENT_NAME("%s: %s", ShaderName, *LightProxy->GetOwnerNameOrLabel()),
 		PassParameters,
 		ERDGPassFlags::Raster,
-		[Scene, &View, PixelShader, LightSceneInfo, PassParameters, LightBounds, LightType, SubstrateTileMaterialType](FRDGAsyncTask, FRHICommandList& RHICmdList)
+		[Scene, &View, PixelShader, UberPixelShader, LightSceneInfo, PassParameters, LightBounds, LightType, SubstrateTileMaterialType](FRDGAsyncTask, FRHICommandList& RHICmdList)
 	{
-
 		const bool bIsRadial = LightType != LightType_Directional;
 		const bool bEnableSubstrateStencilTest = SubstrateTileMaterialType != ESubstrateTileType::ECount && bIsRadial;
 
@@ -2746,13 +2823,13 @@ static void InternalRenderLight(
 			GraphicsPSOInit.BoundShaderState.VertexShaderRHI = bEnableSubstrateTiledPass ? TileVertexShader.GetVertexShader() : VertexShader.GetVertexShader();
 			GraphicsPSOInit.BoundShaderState.PixelShaderRHI = PixelShader.GetPixelShader();
 
-#if PSO_PRECACHING_VALIDATE
-			if (PSOCollectorStats::IsFullPrecachingValidationEnabled())
+			TShaderType UsedPixelShader = PixelShader;
+			if (!PrecachePSO(GraphicsPSOInit))
 			{
-				static const int32 GlobalPSOCollectorIndex = FGlobalPSOCollectorManager::GetIndex(DeferredLightGlobalPSOCollectorName);
-				PSOCollectorStats::CheckGlobalGraphicsPipelineStateInCache(GraphicsPSOInit, GlobalPSOCollectorIndex);
+				UsedPixelShader = UberPixelShader;
+				GraphicsPSOInit.BoundShaderState.PixelShaderRHI = UsedPixelShader.GetPixelShader();
+				ValidatePSOPrecacheState(GraphicsPSOInit);
 			}
-#endif // PSO_PRECACHING_VALIDATE
 
 			Substrate::FSubstrateTilePassVS::FParameters VSParameters;
 			if (Substrate::IsSubstrateEnabled())
@@ -2761,7 +2838,7 @@ static void InternalRenderLight(
 			}
 
 			SetGraphicsPipelineState(RHICmdList, GraphicsPSOInit, StencilRef);
-			SetShaderParameters(RHICmdList, PixelShader, PixelShader.GetPixelShader(), PassParameters->PS);
+			SetShaderParameters(RHICmdList, UsedPixelShader, UsedPixelShader.GetPixelShader(), PassParameters->PS);
 
 			if (bEnableSubstrateTiledPass)
 			{
@@ -2796,16 +2873,16 @@ static void InternalRenderLight(
 			GraphicsPSOInit.BoundShaderState.VertexShaderRHI = VertexShader.GetVertexShader();
 			GraphicsPSOInit.BoundShaderState.PixelShaderRHI = PixelShader.GetPixelShader();
 
-#if PSO_PRECACHING_VALIDATE
-			if (PSOCollectorStats::IsFullPrecachingValidationEnabled())
+			TShaderType UsedPixelShader = PixelShader;
+			if (!PrecachePSO(GraphicsPSOInit))
 			{
-				static const int32 GlobalPSOCollectorIndex = FGlobalPSOCollectorManager::GetIndex(DeferredLightGlobalPSOCollectorName);
-				PSOCollectorStats::CheckGlobalGraphicsPipelineStateInCache(GraphicsPSOInit, GlobalPSOCollectorIndex);
+				UsedPixelShader = UberPixelShader;
+				GraphicsPSOInit.BoundShaderState.PixelShaderRHI = UsedPixelShader.GetPixelShader();
+				ValidatePSOPrecacheState(GraphicsPSOInit);
 			}
-#endif // PSO_PRECACHING_VALIDATE
 
 			SetGraphicsPipelineState(RHICmdList, GraphicsPSOInit, StencilRef);
-			SetShaderParameters(RHICmdList, PixelShader, PixelShader.GetPixelShader(), PassParameters->PS);
+			SetShaderParameters(RHICmdList, UsedPixelShader, UsedPixelShader.GetPixelShader(), PassParameters->PS);
 
 			FDeferredLightVS::FParameters VSParameters2 = FDeferredLightVS::GetParameters(View, LightSceneInfo);
 			SetShaderParameters(RHICmdList, VertexShader, VertexShader.GetVertexShader(), VSParameters2);
@@ -2931,7 +3008,7 @@ static void RenderLight(
 		FDeferredLightOverlapPS::FPermutationDomain PermutationVector;
 		PermutationVector.Set<FDeferredLightOverlapPS::FRadialAttenuation>(bIsRadial);
 		TShaderMapRef<FDeferredLightOverlapPS> PixelShader(View.ShaderMap, PermutationVector);		
-		InternalRenderLight(GraphBuilder, Scene, View, LightSceneInfo, PixelShader, PassParameters, ESubstrateTileType::ECount, TEXT("Light::StandardDeferred(Overlap)"));
+		InternalRenderLight(GraphBuilder, Scene, View, LightSceneInfo, PixelShader, PixelShader, PassParameters, ESubstrateTileType::ECount, TEXT("Light::StandardDeferred(Overlap)"));
 	}
 	// Lighting shader
 	else
@@ -2989,8 +3066,8 @@ static void RenderLight(
 			PermutationVector.Set< FDeferredLightPS::FIESProfileDim >(false);
 			PermutationVector.Set< FDeferredLightPS::FAnistropicMaterials >(bSupportAnisotropyPermutation);
 			// Only directional lights are rendered in this path, so we only need to check if it is use to light the atmosphere
-			PermutationVector.Set< FDeferredLightPS::FAtmosphereTransmittance >(IsLightAtmospherePerPixelTransmittanceEnabled(Scene, View, LightSceneInfo));
-			PermutationVector.Set< FDeferredLightPS::FCloudTransmittance >(PassParameters->PS.CloudShadowEnabled > 0);
+			PermutationVector.Set< FDeferredLightPS::FAtmosphereTransmittance >(PassParameters->PS.bAtmosphereTransmittanceEnabled > 0);
+			PermutationVector.Set< FDeferredLightPS::FCloudTransmittance >(PassParameters->PS.bCloudShadowEnabled > 0);
 		}
 
 		// Substrate tile rendering: 
@@ -3008,41 +3085,51 @@ static void RenderLight(
 			{
 				const ESubstrateTileType TileType = ESubstrateTileType::EComplexSpecial;
 				PermutationVector.Set<Substrate::FSubstrateTileType>(TileType);
-				FDeferredLightPS::FPermutationDomain PermutationVectorRemap = FDeferredLightPS::RemapPermutation(PermutationVector);
+				FDeferredLightPS::FPermutationDomain PermutationVectorRemap = FDeferredLightPS::RemapPermutation(PermutationVector, View.GetShaderPlatform());
 				TShaderMapRef< FDeferredLightPS > PixelShader(View.ShaderMap, PermutationVectorRemap);
-				InternalRenderLight(GraphBuilder, Scene, View, LightSceneInfo, PixelShader, PassParameters, TileType, TEXT("Light::StandardDeferred(ComplexSpecial)"));
+				TShaderMapRef< FDeferredLightPS > UberPixelShader = GetUberDeferredLightPS(View.ShaderMap, View.GetShaderPlatform(), PermutationVectorRemap);
+
+				InternalRenderLight(GraphBuilder, Scene, View, LightSceneInfo, PixelShader, UberPixelShader, PassParameters, TileType, TEXT("Light::StandardDeferred(ComplexSpecial)"));
 			}
 			if (ShouldRenderTileType(ESubstrateTileType::EComplex))
 			{
 				const ESubstrateTileType TileType = ESubstrateTileType::EComplex;
 				PermutationVector.Set<Substrate::FSubstrateTileType>(TileType);
-				FDeferredLightPS::FPermutationDomain PermutationVectorRemap = FDeferredLightPS::RemapPermutation(PermutationVector);
+				FDeferredLightPS::FPermutationDomain PermutationVectorRemap = FDeferredLightPS::RemapPermutation(PermutationVector, View.GetShaderPlatform());
 				TShaderMapRef< FDeferredLightPS > PixelShader(View.ShaderMap, PermutationVectorRemap);
-				InternalRenderLight(GraphBuilder, Scene, View, LightSceneInfo, PixelShader, PassParameters, TileType, TEXT("Light::StandardDeferred(Complex)"));
+				TShaderMapRef< FDeferredLightPS > UberPixelShader = GetUberDeferredLightPS(View.ShaderMap, View.GetShaderPlatform(), PermutationVectorRemap);
+
+				InternalRenderLight(GraphBuilder, Scene, View, LightSceneInfo, PixelShader, UberPixelShader, PassParameters, TileType, TEXT("Light::StandardDeferred(Complex)"));
 			}
 			if (ShouldRenderTileType(ESubstrateTileType::ESingle))
 			{
 				const ESubstrateTileType TileType = ESubstrateTileType::ESingle;
 				PermutationVector.Set<Substrate::FSubstrateTileType>(TileType);
-				FDeferredLightPS::FPermutationDomain PermutationVectorRemap = FDeferredLightPS::RemapPermutation(PermutationVector);
+				FDeferredLightPS::FPermutationDomain PermutationVectorRemap = FDeferredLightPS::RemapPermutation(PermutationVector, View.GetShaderPlatform());
 				TShaderMapRef< FDeferredLightPS > PixelShader(View.ShaderMap, PermutationVectorRemap);
-				InternalRenderLight(GraphBuilder, Scene, View, LightSceneInfo, PixelShader, PassParameters, TileType, TEXT("Light::StandardDeferred(Single)"));
+				TShaderMapRef< FDeferredLightPS > UberPixelShader = GetUberDeferredLightPS(View.ShaderMap, View.GetShaderPlatform(), PermutationVectorRemap);
+
+				InternalRenderLight(GraphBuilder, Scene, View, LightSceneInfo, PixelShader, UberPixelShader, PassParameters, TileType, TEXT("Light::StandardDeferred(Single)"));
 			}
 			if (ShouldRenderTileType(ESubstrateTileType::ESimple))
 			{
 				const ESubstrateTileType TileType = ESubstrateTileType::ESimple;
 				PermutationVector.Set<Substrate::FSubstrateTileType>(TileType);
-				FDeferredLightPS::FPermutationDomain PermutationVectorRemap = FDeferredLightPS::RemapPermutation(PermutationVector);
+				FDeferredLightPS::FPermutationDomain PermutationVectorRemap = FDeferredLightPS::RemapPermutation(PermutationVector, View.GetShaderPlatform());
 				TShaderMapRef< FDeferredLightPS > PixelShader(View.ShaderMap, PermutationVectorRemap);
-				InternalRenderLight(GraphBuilder, Scene, View, LightSceneInfo, PixelShader, PassParameters, TileType, TEXT("Light::StandardDeferred(Simple)"));
+				TShaderMapRef< FDeferredLightPS > UberPixelShader = GetUberDeferredLightPS(View.ShaderMap, View.GetShaderPlatform(), PermutationVectorRemap);
+
+				InternalRenderLight(GraphBuilder, Scene, View, LightSceneInfo, PixelShader, UberPixelShader, PassParameters, TileType, TEXT("Light::StandardDeferred(Simple)"));
 			}
 		}
 		else
 		{
 			PermutationVector.Set< Substrate::FSubstrateTileType>(0);
-			PermutationVector = FDeferredLightPS::RemapPermutation(PermutationVector);
+			PermutationVector = FDeferredLightPS::RemapPermutation(PermutationVector, View.GetShaderPlatform());
 			TShaderMapRef< FDeferredLightPS > PixelShader(View.ShaderMap, PermutationVector);
-			InternalRenderLight(GraphBuilder, Scene, View, LightSceneInfo, PixelShader, PassParameters, ESubstrateTileType::ECount, TEXT("Light::StandardDeferred"));
+			TShaderMapRef< FDeferredLightPS > UberPixelShader = GetUberDeferredLightPS(View.ShaderMap, View.GetShaderPlatform(), PermutationVector);
+
+			InternalRenderLight(GraphBuilder, Scene, View, LightSceneInfo, PixelShader, UberPixelShader, PassParameters, ESubstrateTileType::ECount, TEXT("Light::StandardDeferred"));
 		}
 	}
 }
@@ -3153,8 +3240,8 @@ void FDeferredShadingSceneRenderer::RenderLightForHair(
 		PermutationVector.Set< FDeferredLightPS::FSourceShapeDim >(ELightSourceShape::Directional);
 		PermutationVector.Set< FDeferredLightPS::FSourceTextureDim >(false);
 		PermutationVector.Set< FDeferredLightPS::FIESProfileDim >(false);
-		PermutationVector.Set< FDeferredLightPS::FAtmosphereTransmittance >(IsLightAtmospherePerPixelTransmittanceEnabled(Scene, View, LightSceneInfo));
-		PermutationVector.Set< FDeferredLightPS::FCloudTransmittance >(PassParameters->PS.CloudShadowEnabled > 0.f);
+		PermutationVector.Set< FDeferredLightPS::FAtmosphereTransmittance >(PassParameters->PS.bAtmosphereTransmittanceEnabled > 0.f);
+		PermutationVector.Set< FDeferredLightPS::FCloudTransmittance >(PassParameters->PS.bCloudShadowEnabled > 0.f);
 	}
 	else
 	{
@@ -3166,7 +3253,7 @@ void FDeferredShadingSceneRenderer::RenderLightForHair(
 		PermutationVector.Set< FDeferredLightPS::FCloudTransmittance >(false);
 	}
 
-	PermutationVector = FDeferredLightPS::RemapPermutation(PermutationVector);
+	PermutationVector = FDeferredLightPS::RemapPermutation(PermutationVector, View.GetShaderPlatform());
 	TShaderMapRef<FDeferredLightHairVS> VertexShader(View.ShaderMap);
 	TShaderMapRef<FDeferredLightPS> PixelShader(View.ShaderMap, PermutationVector);
 
@@ -3174,7 +3261,7 @@ void FDeferredShadingSceneRenderer::RenderLightForHair(
 		{},
 		PassParameters,
 		ERDGPassFlags::Raster,
-		[this, VertexShader, PixelShader, PassParameters, SampleLightingViewportResolution](FRDGAsyncTask, FRHICommandList& RHICmdList)
+		[this, VertexShader, PixelShader, PermutationVector, &View, PassParameters, SampleLightingViewportResolution](FRDGAsyncTask, FRHICommandList& RHICmdList)
 	{
 		RHICmdList.SetViewport(0, 0, 0.0f, SampleLightingViewportResolution.X, SampleLightingViewportResolution.Y, 1.0f);
 
@@ -3185,19 +3272,19 @@ void FDeferredShadingSceneRenderer::RenderLightForHair(
 
 		GraphicsPSOInit.BoundShaderState.VertexShaderRHI = VertexShader.GetVertexShader();
 		GraphicsPSOInit.BoundShaderState.PixelShaderRHI = PixelShader.GetPixelShader();
-				
-#if PSO_PRECACHING_VALIDATE
-		if (PSOCollectorStats::IsFullPrecachingValidationEnabled())
+
+		TShaderMapRef< FDeferredLightPS > UsedPixelShader = PixelShader;
+		if (!PrecachePSO(GraphicsPSOInit))
 		{
-			static const int32 GlobalPSOCollectorIndex = FGlobalPSOCollectorManager::GetIndex(DeferredLightGlobalPSOCollectorName);
-			PSOCollectorStats::CheckGlobalGraphicsPipelineStateInCache(GraphicsPSOInit, GlobalPSOCollectorIndex);
+			UsedPixelShader = GetUberDeferredLightPS(View.ShaderMap, View.GetShaderPlatform(), PermutationVector);
+			GraphicsPSOInit.BoundShaderState.PixelShaderRHI = UsedPixelShader.GetPixelShader();
+			ValidatePSOPrecacheState(GraphicsPSOInit);
 		}
-#endif // PSO_PRECACHING_VALIDATE
 
 		SetGraphicsPipelineState(RHICmdList, GraphicsPSOInit, 0);
 
 		SetShaderParameters(RHICmdList, VertexShader, VertexShader.GetVertexShader(), PassParameters->VS);
-		SetShaderParameters(RHICmdList, PixelShader, PixelShader.GetPixelShader(), PassParameters->PS);
+		SetShaderParameters(RHICmdList, UsedPixelShader, UsedPixelShader.GetPixelShader(), PassParameters->PS);
 		
 		RHICmdList.SetStreamSource(0, nullptr, 0);
 		RHICmdList.DrawPrimitive(0, 1, 1);
@@ -3286,12 +3373,14 @@ static FDeferredLightPS::FParameters GetRenderLightSimplePSParameters(
 	// PS - General parameters
 	Out.SceneTextures = SceneTextures.UniformBuffer;
 	Out.HairStrands = View.HairStrandsViewData.UniformBuffer;
+	Out.ForwardLightStruct = View.ForwardLightingResources.ForwardLightUniformBuffer;
 	Out.Substrate = Substrate::BindSubstrateGlobalUniformParameters(View);
+	Out.bLightingChannelEnabled = 0u;
 	Out.LightingChannelsTexture = WhiteDummy;
 	Out.LightingChannelsSampler = TStaticSamplerState<SF_Point, AM_Clamp, AM_Clamp, AM_Clamp>::GetRHI();
 	Out.CloudShadowAO = GetCloudShadowAOParameters(GraphBuilder, View, nullptr);
-	Out.CloudShadowEnabled = 0;
-	SetupLightCloudTransmittanceParameters(GraphBuilder, nullptr, View, nullptr, Out.CloudShadow);
+	Out.bCloudShadowEnabled = SetupLightCloudTransmittanceParameters(GraphBuilder, nullptr, View, nullptr, Out.CloudShadow);	
+	Out.bAtmosphereTransmittanceEnabled = 0;
 	Out.LightAttenuationTexture = WhiteDummy;
 	Out.LightAttenuationTextureSampler = TStaticSamplerState<SF_Point, AM_Wrap, AM_Wrap, AM_Wrap>::GetRHI();
 	Out.View = View.ViewUniformBuffer;
@@ -3302,6 +3391,14 @@ static FDeferredLightPS::FParameters GetRenderLightSimplePSParameters(
 	Out.HairTransmittanceBufferMaxCount = 0;
 	Out.HairShadowMaskValid = false;
 	Out.ShadowChannelMask = FVector4f(1, 1, 1, 1);
+	// PS - One pass projection
+	Out.VirtualShadowMap = FVirtualShadowMapArray::CreateDummySamplingParameters(GraphBuilder).VirtualShadowMap;
+	Out.VirtualShadowMapId = INDEX_NONE;
+	Out.ShadowMaskBits = GSystemTextures.GetZeroUIntDummy(GraphBuilder);
+
+	// For ubershader permutation
+	Out.LightFunctionAtlas = LightFunctionAtlas::BindGlobalParameters(GraphBuilder, View);
+
 	// PS - RT/Depth
 	Out.RenderTargets[0] = FRenderTargetBinding(SceneTextures.Color.Target, ERenderTargetLoadAction::ELoad);
 	if (SceneTextures.Depth.Target)
@@ -3387,7 +3484,7 @@ static void InternalRenderSimpleLightsForHair(
 	PermutationVector.Set< FDeferredLightPS::FCloudTransmittance >(false);
 	PermutationVector.Set< Substrate::FSubstrateTileType>(Substrate::IsSubstrateEnabled() ? ESubstrateTileType::EComplex : 0);
 
-	PermutationVector = FDeferredLightPS::RemapPermutation(PermutationVector);
+	PermutationVector = FDeferredLightPS::RemapPermutation(PermutationVector, View.GetShaderPlatform());
 	TShaderMapRef<FDeferredLightPS> PixelShader(View.ShaderMap, PermutationVector);
 	TShaderMapRef<FDeferredLightHairVS> VertexShader(View.ShaderMap);
 
@@ -3395,7 +3492,7 @@ static void InternalRenderSimpleLightsForHair(
 		RDG_EVENT_NAME("Light::DeferredSimpleLights(Hair)"),
 		PassParameters,
 		ERDGPassFlags::Raster,
-		[&View, &SortedLightSet, ViewIndex, NumViews, PassParameters, PixelShader, VertexShader](FRDGAsyncTask, FRHICommandList& RHICmdList)
+		[&View, &SortedLightSet, ViewIndex, NumViews, PassParameters, PixelShader, VertexShader, PermutationVector](FRDGAsyncTask, FRHICommandList& RHICmdList)
 	{
 		FGraphicsPipelineStateInitializer GraphicsPSOInit;
 		RHICmdList.ApplyCachedRenderTargets(GraphicsPSOInit);
@@ -3421,19 +3518,20 @@ static void InternalRenderSimpleLightsForHair(
 			GraphicsPSOInit.BoundShaderState.VertexShaderRHI = VertexShader.GetVertexShader();
 			GraphicsPSOInit.BoundShaderState.PixelShaderRHI = PixelShader.GetPixelShader();
 						
-			#if PSO_PRECACHING_VALIDATE
-			if (PSOCollectorStats::IsFullPrecachingValidationEnabled())
+			TShaderMapRef< FDeferredLightPS > UsedPixelShader = PixelShader;
+			if (!PrecachePSO(GraphicsPSOInit))
 			{
-				PSOCollectorStats::CheckGlobalGraphicsPipelineStateInCache(GraphicsPSOInit, GlobalPSOCollectorIndex);
+				UsedPixelShader = GetUberDeferredLightPS(View.ShaderMap, View.GetShaderPlatform(), PermutationVector);
+				GraphicsPSOInit.BoundShaderState.PixelShaderRHI = UsedPixelShader.GetPixelShader();
+				ValidatePSOPrecacheState(GraphicsPSOInit);
 			}
-			#endif // PSO_PRECACHING_VALIDATE
 
 			SetGraphicsPipelineState(RHICmdList, GraphicsPSOInit, 0/*StencilRef*/);
 
 			// Update the light parameters with a custom uniform buffer
 			FDeferredLightUniformStruct DeferredLightUniformsValue = GetSimpleDeferredLightParameters(View, SimpleLight, SimpleLightPerViewData);
 
-			SetShaderParametersMixed(RHICmdList, PixelShader, PixelShader.GetPixelShader(), PassParameters->PS, DeferredLightUniformsValue);
+			SetShaderParametersMixed(RHICmdList, UsedPixelShader, UsedPixelShader.GetPixelShader(), PassParameters->PS, DeferredLightUniformsValue);
 			SetShaderParameters(RHICmdList, VertexShader, VertexShader.GetVertexShader(), PassParameters->VS);
 
 			RHICmdList.SetStreamSource(0, nullptr, 0);
@@ -3477,7 +3575,7 @@ static void InternalRenderSimpleLightsStandardDeferred(
 	PermutationVector.Set< FDeferredLightPS::FCloudTransmittance >(false);
 	PermutationVector.Set< Substrate::FSubstrateTileType>(TileType != ESubstrateTileType::ECount ? TileType : 0);
 
-	PermutationVector = FDeferredLightPS::RemapPermutation(PermutationVector);
+	PermutationVector = FDeferredLightPS::RemapPermutation(PermutationVector, View.GetShaderPlatform());
 	TShaderMapRef<FDeferredLightPS> PixelShader(View.ShaderMap, PermutationVector);
 
 	FDeferredLightVS::FPermutationDomain PermutationVectorVS;
@@ -3488,7 +3586,7 @@ static void InternalRenderSimpleLightsStandardDeferred(
 		RDG_EVENT_NAME("Light::DeferredSimpleLights(Substrate:%s,Tile:%s)", Substrate::IsSubstrateEnabled() ? TEXT("True") : TEXT("False"), Substrate::IsSubstrateEnabled() ? ToString(TileType) : TEXT("None")),
 		PassParameters,
 		ERDGPassFlags::Raster,
-		[&View, &SortedLightSet, ViewIndex, NumViews, PassParameters, PixelShader, VertexShader, TileType](FRDGAsyncTask, FRHICommandList& RHICmdList)
+		[&View, &SortedLightSet, ViewIndex, NumViews, PassParameters, PixelShader, VertexShader, PermutationVector, TileType](FRDGAsyncTask, FRHICommandList& RHICmdList)
 	{
 		FGraphicsPipelineStateInitializer GraphicsPSOInit;
 		RHICmdList.ApplyCachedRenderTargets(GraphicsPSOInit);
@@ -3497,8 +3595,6 @@ static void InternalRenderSimpleLightsStandardDeferred(
 		GraphicsPSOInit.BlendState = TStaticBlendState<CW_RGBA, BO_Add, BF_One, BF_One, BO_Add, BF_One, BF_One>::GetRHI();
 		GraphicsPSOInit.PrimitiveType = PT_TriangleList;
 
-		const int32 GlobalPSOCollectorIndex = FGlobalPSOCollectorManager::GetIndex(DeferredLightGlobalPSOCollectorName);
-
 		for (int32 SortedLightIndex = 0; SortedLightIndex < SortedLightSet.SimpleLightsEnd; SortedLightIndex++)
 		{
 			const FSortedLightSceneInfo& SortedLightInfo = SortedLightSet.SortedLights[SortedLightIndex];
@@ -3523,19 +3619,20 @@ static void InternalRenderSimpleLightsStandardDeferred(
 			GraphicsPSOInit.BoundShaderState.VertexShaderRHI = VertexShader.GetVertexShader();
 			GraphicsPSOInit.BoundShaderState.PixelShaderRHI = PixelShader.GetPixelShader();
 						
-#if PSO_PRECACHING_VALIDATE
-			if (PSOCollectorStats::IsFullPrecachingValidationEnabled())
+			TShaderMapRef< FDeferredLightPS > UsedPixelShader = PixelShader;
+			if (!PrecachePSO(GraphicsPSOInit))
 			{
-				PSOCollectorStats::CheckGlobalGraphicsPipelineStateInCache(GraphicsPSOInit, GlobalPSOCollectorIndex);
+				UsedPixelShader = GetUberDeferredLightPS(View.ShaderMap, View.GetShaderPlatform(), PermutationVector);
+				GraphicsPSOInit.BoundShaderState.PixelShaderRHI = UsedPixelShader.GetPixelShader();
+				ValidatePSOPrecacheState(GraphicsPSOInit);
 			}
-#endif // PSO_PRECACHING_VALIDATE
 
 			SetGraphicsPipelineState(RHICmdList, GraphicsPSOInit, StencilRef);
 
 			// Update the light parameters with a custom uniform buffer
 			FDeferredLightUniformStruct DeferredLightUniformsValue = GetSimpleDeferredLightParameters(View, SimpleLight, SimpleLightPerViewData);
 
-			SetShaderParametersMixed(RHICmdList, PixelShader, PixelShader.GetPixelShader(), PassParameters->PS, DeferredLightUniformsValue);
+			SetShaderParametersMixed(RHICmdList, UsedPixelShader, UsedPixelShader.GetPixelShader(), PassParameters->PS, DeferredLightUniformsValue);
 
 			// Update vertex shader parameters with custom parameters/uniform buffer
 			FDeferredLightVS::FParameters ParametersVS = FDeferredLightVS::GetParameters(View, LightBounds);
@@ -3786,9 +3883,6 @@ void DeferredLightGlobalPSOCollector(const FSceneTexturesConfig& SceneTexturesCo
 		PSOInitializers.Add(MoveTemp(PSOPrecacheData));
 	};
 
-	// Create variations for given render & depth stencil compare states - influence raster state culling for non directional only
-	const bool bReverseCulling = false;
-
 	// Precache PSOs are never required
 	const bool bRequired = false;
 
@@ -3808,15 +3902,23 @@ void DeferredLightGlobalPSOCollector(const FSceneTexturesConfig& SceneTexturesCo
 			{
 				FDeferredLightPS::FPermutationDomain PermutationVector(PermutationId);
 
+				// Only precache the uber shader versions
+				if (!PermutationVector.Get < FDeferredLightPS::FUbershader > ())
+				{
+					continue;
+				}
+
 				// Extract useful information from the permutation vector
 				bool bIsDirectional = (PermutationVector.Get< FDeferredLightPS::FSourceShapeDim >() == ELightSourceShape::Directional);
 				ESubstrateTileType SubstrateTileMaterialType = Substrate::IsSubstrateEnabled() ? (ESubstrateTileType) PermutationVector.Get< Substrate::FSubstrateTileType >() : ESubstrateTileType::ECount;
 				bool bIsHairShader = PermutationVector.Get<FDeferredLightPS::FHairLighting>() > 0;
 
-				AddPSOInitializer(bIsHairShader, true /*bCameraInsideLightGeometry*/, bReverseCulling, bIsDirectional, SubstrateTileMaterialType, RHIPixelShader);
+				AddPSOInitializer(bIsHairShader, true /*bCameraInsideLightGeometry*/, true /*bReverCulling*/, bIsDirectional, SubstrateTileMaterialType, RHIPixelShader);
+				AddPSOInitializer(bIsHairShader, true /*bCameraInsideLightGeometry*/, false /*bReverCulling*/, bIsDirectional, SubstrateTileMaterialType, RHIPixelShader);
 				if (!bIsDirectional)
 				{
-					AddPSOInitializer(bIsHairShader, false /*bCameraInsideLightGeometry*/, bReverseCulling, bIsDirectional, SubstrateTileMaterialType, RHIPixelShader);
+					AddPSOInitializer(bIsHairShader, false /*bCameraInsideLightGeometry*/, true /*bReverCulling*/, bIsDirectional, SubstrateTileMaterialType, RHIPixelShader);
+					AddPSOInitializer(bIsHairShader, true /*bCameraInsideLightGeometry*/, false /*bReverCulling*/, bIsDirectional, SubstrateTileMaterialType, RHIPixelShader);
 				}
 			}
 		}
