Skip virtual texture feedback compaction and update if there are no allocated virtual textures.
#rnx
#Synced-CL 49758170

diff --git a/Engine/Source/Runtime/Renderer/Private/VT/VirtualTextureFeedbackResource.cpp b/Engine/Source/Runtime/Renderer/Private/VT/VirtualTextureFeedbackResource.cpp
--- a/Engine/Source/Runtime/Renderer/Private/VT/VirtualTextureFeedbackResource.cpp
+++ b/Engine/Source/Runtime/Renderer/Private/VT/VirtualTextureFeedbackResource.cpp
@@ -120,6 +120,13 @@ void FVirtualTextureFeedbackBufferResource::End(FRDGBuilder& GraphBuilder)
 		RHICmdList.Transition(FRHITransitionInfo(UAV, ERHIAccess::UAVMask, ERHIAccess::SRVCompute));
 	});
 
+	// If there are no allocated VTs then we can assume that no GPU feedback was generated.
+	// So early out here to avoid paying any GPU cost compacting the empty feedback buffer.
+	if (!FVirtualTextureSystem::Get().HasAllocatedVirtualTextures())
+	{
+		return;
+	}
+
 	FRDGBufferRef FeedbackBuffer = GraphBuilder.RegisterExternalBuffer(PooledBuffer, ERDGBufferFlags::SkipTracking);
 	FRDGBufferSRVRef FeedbackBufferSRV = GraphBuilder.CreateSRV(FeedbackBuffer);
 
diff --git a/Engine/Source/Runtime/Renderer/Private/VT/VirtualTextureSystem.h b/Engine/Source/Runtime/Renderer/Private/VT/VirtualTextureSystem.h
--- a/Engine/Source/Runtime/Renderer/Private/VT/VirtualTextureSystem.h
+++ b/Engine/Source/Runtime/Renderer/Private/VT/VirtualTextureSystem.h
@@ -171,6 +171,8 @@ public:
 
 	bool IsPendingRootPageMap(IAllocatedVirtualTexture* AllocatedVT) const;
 
+	bool HasAllocatedVirtualTextures() const { return AllocatedVTs.Num() > 0; }
+
 private:
 	friend class FFeedbackAnalysisTask;
 	friend class FAddRequestedTilesTask;
