[TSR] TSR flickering detection logic covers more cases. Change r.TSR.ShadingRejection.Mode from 0 to 1. This makes the Flickering period more sensitive to detect noise. It helps VSM flickering when antiflickering logic is enabled.

#tests test on level TSR shared across different teams and the city sample.
#Synced-CL 49128105

diff --git a/Engine/Source/Runtime/Renderer/Private/PostProcess/TemporalSuperResolution.cpp b/Engine/Source/Runtime/Renderer/Private/PostProcess/TemporalSuperResolution.cpp
--- a/Engine/Source/Runtime/Renderer/Private/PostProcess/TemporalSuperResolution.cpp
+++ b/Engine/Source/Runtime/Renderer/Private/PostProcess/TemporalSuperResolution.cpp
@@ -120,9 +120,9 @@ TAutoConsoleVariable<int32> CVarTSR16BitVALUOnNvidia(
 #endif // PLATFORM_DESKTOP
 
 TAutoConsoleVariable<int32> CVarTSRShadingRejectionMode(
-	TEXT("r.TSR.ShadingRejection.Mode"), 0,
+	TEXT("r.TSR.ShadingRejection.Mode"), 1,
 	TEXT(".\n")
-	TEXT("  0: More responsiveness and less ghosting but tends to have more blocky artifacts.(Default)\n")
+	TEXT("  0: More responsiveness and less ghosting but tends to have more blocky artifacts.\n")
 	TEXT("  1: Better stability and shading rejection control but tends to have ghosting.\n"),
 	ECVF_RenderThreadSafe);
 
