[HWRT] Fix bug in SBT allocations without deduplication.
- Match the logic used for FRefCountedAllocations.
- MarkSet(Allocation, false) was not being called.
- And also defer the free on the StaticRangeAllocator.

#Synced-CL 49716836

diff --git a/Engine/Source/Runtime/Renderer/Private/RayTracing/RayTracingShaderBindingTable.cpp b/Engine/Source/Runtime/Renderer/Private/RayTracing/RayTracingShaderBindingTable.cpp
--- a/Engine/Source/Runtime/Renderer/Private/RayTracing/RayTracingShaderBindingTable.cpp
+++ b/Engine/Source/Runtime/Renderer/Private/RayTracing/RayTracingShaderBindingTable.cpp
@@ -496,7 +496,7 @@ FRayTracingSBTAllocation* FRayTracingShaderBindingTable::AllocateStaticRange(uin
 	return Allocation.Allocation;
 }
 
-void FRayTracingShaderBindingTable::FreeStaticRange(const FRayTracingSBTAllocation* Allocation)
+void FRayTracingShaderBindingTable::FreeStaticRange(FRayTracingSBTAllocation* Allocation)
 {
 	if (Allocation == nullptr)
 	{
@@ -507,7 +507,7 @@ void FRayTracingShaderBindingTable::FreeStaticRange(const FRayTracingSBTAllocati
 
 	TotalStaticAllocationCount--;
 
-	// If geometry is stored then it could have been deduplicatedf and we can build the allocation key again
+	// If geometry is stored then it could have been deduplicated and we can build the allocation key again
 	if (Allocation->Geometry)
 	{
 		FAllocationKey Key;
@@ -515,7 +515,7 @@ void FRayTracingShaderBindingTable::FreeStaticRange(const FRayTracingSBTAllocati
 		Key.Flags = Allocation->Flags;
 
 		FRefCountedAllocation* RefAllocation = TrackedAllocationMap.Find(Key);
-		check(Allocation);
+		check(RefAllocation);
 		RefAllocation->RefCount--;
 
 		if (RefAllocation->RefCount == 0)
@@ -523,7 +523,7 @@ void FRayTracingShaderBindingTable::FreeStaticRange(const FRayTracingSBTAllocati
 			AllocatedStaticSegmentCount -= (RefAllocation->Allocation->NumRecords / NumShaderSlotsPerGeometrySegment);
 
 			// Already mark unused but only defer the free on the free range allocator until the clear op is actually processed
-			// so the free can happen free threaded while the static allocations are still locked			
+			// so the free can happen free threaded while the static allocations are still locked
 			MarkSet(RefAllocation->Allocation, false);
 			PendingStaticAllocationsToFree.Add(RefAllocation->Allocation);
 
@@ -532,9 +532,12 @@ void FRayTracingShaderBindingTable::FreeStaticRange(const FRayTracingSBTAllocati
 	}
 	else
 	{
-		StaticRangeAllocator.Free(Allocation->BaseRecordIndex, Allocation->NumRecords);
-		AllocatedStaticSegmentCount -= (Allocation->NumRecords / NumShaderSlotsPerGeometrySegment);		
-		delete Allocation;
+		AllocatedStaticSegmentCount -= (Allocation->NumRecords / NumShaderSlotsPerGeometrySegment);
+
+		// Already mark unused but only defer the free on the free range allocator until the clear op is actually processed
+		// so the free can happen free threaded while the static allocations are still locked
+		MarkSet(Allocation, false);
+		PendingStaticAllocationsToFree.Add(Allocation);
 	}
 }
 
@@ -633,7 +636,7 @@ void FRayTracingShaderBindingTable::FlushAllocationsToClear(FRHICommandList& RHI
 
 	PersistentRecordsToClear.Init(false, PersistentRecordsToClear.Num());
 
-	// Only makr ranges as free now on the SBT after the clear ops have been processed by the RHI SBT
+	// Only make ranges as free now on the SBT after the clear ops have been processed by the RHI SBT
 	// to make sure we don't reuse records which haven't been freed yet (due to out of order freeing of records)
 	uint32 FreedRecords = 0;
 	for (FRayTracingSBTAllocation* AllocationToFree : PendingStaticAllocationsToFree)
diff --git a/Engine/Source/Runtime/Renderer/Private/RayTracing/RayTracingShaderBindingTable.h b/Engine/Source/Runtime/Renderer/Private/RayTracing/RayTracingShaderBindingTable.h
--- a/Engine/Source/Runtime/Renderer/Private/RayTracing/RayTracingShaderBindingTable.h
+++ b/Engine/Source/Runtime/Renderer/Private/RayTracing/RayTracingShaderBindingTable.h
@@ -167,7 +167,7 @@ public:
 	 * Allocate or share static allocation range - sharing can happen if geometry and cached RT MDC flags are the same (will result in exactly the same binding data written in the SBT)
 	 */
 	RENDERER_API FRayTracingSBTAllocation* AllocateStaticRange(uint32 SegmentCount, const FRHIRayTracingGeometry* Geometry, FRayTracingCachedMeshCommandFlags Flags);
-	RENDERER_API void FreeStaticRange(const FRayTracingSBTAllocation* Allocation);
+	RENDERER_API void FreeStaticRange(FRayTracingSBTAllocation* Allocation);
 	 
 	/**
 	 * Allocate dynamic SBT range which can be reused again when ResetDynamicAllocationData is called
