Add experimental SetDefaultIfNull to SystemTextures in order to get default texture if a shader resources is null.

#Synced-CL 49800659

diff --git a/Engine/Source/Runtime/Renderer/Private/SystemTextures.cpp b/Engine/Source/Runtime/Renderer/Private/SystemTextures.cpp
--- a/Engine/Source/Runtime/Renderer/Private/SystemTextures.cpp
+++ b/Engine/Source/Runtime/Renderer/Private/SystemTextures.cpp
@@ -1851,3 +1851,55 @@ FRDGBufferRef FSystemTextures::GetDefaultStructuredBuffer(FRDGBuilder& GraphBuil
 FRDGBufferRef FSystemTextures::GetDefaultStructuredBuffer(FRDGBuilder& GraphBuilder, uint32 NumBytePerElement, const FVector3f& Value)		{ return GetInternalDefaultBuffer(GraphBuilder, DefaultBuffers, HashDefaultBuffers, NumBytePerElement, EDefaultBufferType::StructuredBuffer, &Value); }
 FRDGBufferRef FSystemTextures::GetDefaultStructuredBuffer(FRDGBuilder& GraphBuilder, uint32 NumBytePerElement, const FVector4f& Value)		{ return GetInternalDefaultBuffer(GraphBuilder, DefaultBuffers, HashDefaultBuffers, NumBytePerElement, EDefaultBufferType::StructuredBuffer, &Value); }
 FRDGBufferRef FSystemTextures::GetDefaultStructuredBuffer(FRDGBuilder& GraphBuilder, uint32 NumBytePerElement, const FUintVector4& Value)	{ return GetInternalDefaultBuffer(GraphBuilder, DefaultBuffers, HashDefaultBuffers, NumBytePerElement, EDefaultBufferType::StructuredBuffer, &Value); }
+
+//////////////////////////////////////////////////////////////////////////////////////////////////////////////////////////////////////////////////////////////////////////////////////////////////////////////////////////////////////////////////////////////////////////////////////////////////////////////////
+// Default resources
+
+void FSystemTextures::SetDefaultIfNull(FRDGBuilder& GraphBuilder, FRDGTextureRef& Out)
+{
+	if (!Out)
+	{
+		Out = GetBlackDummy(GraphBuilder);
+	}
+}
+
+void FSystemTextures::SetDefaultIfNull(FRDGBuilder& GraphBuilder, FRDGTextureUAVRef& Out)
+{
+	if (!Out)
+	{
+		const FRDGTextureRef DummyTexture = GraphBuilder.CreateTexture(FRDGTextureDesc::Create2D(FIntPoint(1, 1), PF_R32_UINT, FClearValueBinding::None, TexCreate_ShaderResource | TexCreate_UAV), TEXT("DummyUAV"));
+		Out = GraphBuilder.CreateUAV(DummyTexture);
+	}
+}
+
+void FSystemTextures::SetDefaultIfNull(FRDGBuilder& GraphBuilder, TRDGBufferAccess<ERHIAccess::IndirectArgs>& Out)
+{
+	if (!Out)
+	{
+		Out = GetDefaultBuffer(GraphBuilder, 4u, 0u);
+	}
+}
+
+void FSystemTextures::SetDefaultIfNull(FRDGBuilder& GraphBuilder, FRDGBufferRef& Out)
+{
+	if (!Out)
+	{
+		Out = GetDefaultBuffer(GraphBuilder, 4u, 0u);
+	}
+}
+
+void FSystemTextures::SetDefaultIfNull(FRDGBuilder& GraphBuilder, FRDGBufferSRVRef& Out)
+{
+	if (!Out)
+	{
+		Out = GraphBuilder.CreateSRV(GetDefaultBuffer(GraphBuilder, 4u, 0u), PF_R32_UINT);
+	}
+}
+
+void FSystemTextures::SetDefaultIfNull(FRDGBuilder& GraphBuilder, FRDGBufferUAVRef& Out)
+{
+	if (!Out)
+	{
+		Out = GraphBuilder.CreateUAV(GetDefaultBuffer(GraphBuilder, 4u, 0u), PF_R32_UINT);
+	}
+}
diff --git a/Engine/Source/Runtime/Renderer/Public/SystemTextures.h b/Engine/Source/Runtime/Renderer/Public/SystemTextures.h
--- a/Engine/Source/Runtime/Renderer/Public/SystemTextures.h
+++ b/Engine/Source/Runtime/Renderer/Public/SystemTextures.h
@@ -241,6 +241,15 @@ public:
 	FRDGBufferRef RENDERER_API GetDefaultStructuredBuffer(FRDGBuilder& GraphBuilder, uint32 NumBytePerElement, const FVector4f& Value);
 	FRDGBufferRef RENDERER_API GetDefaultStructuredBuffer(FRDGBuilder& GraphBuilder, uint32 NumBytePerElement, const FUintVector4& Value);
 
+
+	// Experimental functions (non-public)
+	void SetDefaultIfNull(FRDGBuilder& GraphBuilder, FRDGTextureRef& Out);
+	void SetDefaultIfNull(FRDGBuilder& GraphBuilder, FRDGTextureUAVRef& Out);
+	void SetDefaultIfNull(FRDGBuilder& GraphBuilder, TRDGBufferAccess<ERHIAccess::IndirectArgs>& Out);
+	void SetDefaultIfNull(FRDGBuilder& GraphBuilder, FRDGBufferRef& Out);
+	void SetDefaultIfNull(FRDGBuilder& GraphBuilder, FRDGBufferSRVRef& Out);
+	void SetDefaultIfNull(FRDGBuilder& GraphBuilder, FRDGBufferUAVRef& Out);
+
 protected:
 	/** Maximum feature level that the textures have been initialized up to */
 	ERHIFeatureLevel::Type FeatureLevelInitializedTo;
