Moved SkyAtmosphereOnRenderBegin in PostStaticMeshUpdate to make sure it does not race with renderer command adding/removing lights.
Must be part of the task scheduling with depenedencies.

#tests PC editor dx12
#rnx
#Synced-CL 49800085

diff --git a/Engine/Source/Runtime/Renderer/Private/SceneRendering.cpp b/Engine/Source/Runtime/Renderer/Private/SceneRendering.cpp
--- a/Engine/Source/Runtime/Renderer/Private/SceneRendering.cpp
+++ b/Engine/Source/Runtime/Renderer/Private/SceneRendering.cpp
@@ -4008,8 +4008,6 @@ IVisibilityTaskData* FSceneRenderer::OnRenderBegin(FRDGBuilder& GraphBuilder, co
 
 	// This is called prior to scene update
 	OIT::OnRenderBegin(Scene->OITSceneData);
-	// Notify the atmosphere system to update lights that are needed for this scene and all its renders
-	SkyAtmosphereOnRenderBegin(*Scene, ViewFamily.EngineShowFlags);
 
 	const bool bIsMobilePlatform = IsMobilePlatform(ShaderPlatform);
 
@@ -4146,6 +4144,9 @@ IVisibilityTaskData* FSceneRenderer::OnRenderBegin(FRDGBuilder& GraphBuilder, co
 				}
 			}
 		}
+
+		// Notify the atmosphere system to update lights that are needed for this scene and all its renders.
+		SkyAtmosphereOnRenderBegin(*Scene, ViewFamily.EngineShowFlags);
 	
 		// Lighting is skipped when running ERendererOutput::DepthPrepassOnly
 		if (GetRendererOutput() == ERendererOutput::FinalSceneColor)
