Pull in PR #14221 to fix missing VRS support for ReflectionEnvironmentAndSky pass and clean up VRS per-pass CVar descriptions.

#Synced-CL 49782381

diff --git a/Engine/Source/Runtime/Renderer/Private/IndirectLightRendering.cpp b/Engine/Source/Runtime/Renderer/Private/IndirectLightRendering.cpp
--- a/Engine/Source/Runtime/Renderer/Private/IndirectLightRendering.cpp
+++ b/Engine/Source/Runtime/Renderer/Private/IndirectLightRendering.cpp
@@ -33,6 +33,7 @@
 #include "Lumen/LumenReflections.h"
 #include "Lumen/LumenShortRangeAO.h"
 #include "MegaLights/MegaLights.h"
+#include "VariableRateShadingImageManager.h"
 
 // This is the project default dynamic global illumination, NOT the scalability setting (see r.Lumen.DiffuseIndirect.Allow for scalability)
 // Must match EDynamicGlobalIlluminationMethod
@@ -1920,6 +1921,8 @@ static void AddSkyReflectionPass(
 		PassParameters->RenderTargets.DepthStencil = FDepthStencilBinding(SceneTextures.Depth.Target,
 			ERenderTargetLoadAction::ELoad, ERenderTargetLoadAction::ELoad, FExclusiveDepthStencil::DepthRead_StencilRead);
 	}
+	
+	PassParameters->RenderTargets.ShadingRateTexture = GVRSImageManager.GetVariableRateShadingImage(GraphBuilder, View, FVariableRateShadingImageManager::EVRSPassType::ReflectionEnvironmentAndSky);
 
 	// Bind hair data
 	const bool bCheckerboardSubsurfaceRendering = IsSubsurfaceCheckerboardFormat(SceneColorTexture.Target->Desc.Format, View);
diff --git a/Engine/Source/Runtime/Renderer/Private/VariableRateShading/VariableRateShadingImageManager.cpp b/Engine/Source/Runtime/Renderer/Private/VariableRateShading/VariableRateShadingImageManager.cpp
--- a/Engine/Source/Runtime/Renderer/Private/VariableRateShading/VariableRateShadingImageManager.cpp
+++ b/Engine/Source/Runtime/Renderer/Private/VariableRateShading/VariableRateShadingImageManager.cpp
@@ -116,72 +116,79 @@ static void CVarVRSImagePassTypeCallback(IConsoleVariable* Var)
 TAutoConsoleVariable<int32> CVarVRSBasePass(
 	TEXT("r.VRS.BasePass"),
 	2,
-	TEXT("Enables Variable Rate Shading for the base pass\n")
-	TEXT("0: Disabled")
-	TEXT("1: Full")
+	TEXT("Whether to apply screen space VRS to the base pass.\n")
+	TEXT("0: Disabled\n")
+	TEXT("1: Full\n")
 	TEXT("2: Conservative (default)"),
 	FConsoleVariableDelegate::CreateStatic(&CVarVRSImagePassTypeCallback),
 	ECVF_RenderThreadSafe);
+
 TAutoConsoleVariable<int32> CVarVRSTranslucency(
 	TEXT("r.VRS.Translucency"),
 	1,
-	TEXT("Enable VRS with translucency rendering.\n")
-	TEXT("0: Disabled")
-	TEXT("1: Full (default)")
+	TEXT("Whether to apply screen space VRS to the translucency pass.\n")
+	TEXT("0: Disabled\n")
+	TEXT("1: Full (default)\n")
 	TEXT("2: Conservative"),
 	FConsoleVariableDelegate::CreateStatic(&CVarVRSImagePassTypeCallback),
 	ECVF_RenderThreadSafe);
+
 TAutoConsoleVariable<int32> CVarVRSNaniteEmitGBuffer(
 	TEXT("r.VRS.NaniteEmitGBuffer"),
 	2,
-	TEXT("Enable VRS with Nanite EmitGBuffer rendering.\n")
-	TEXT("0: Disabled")
-	TEXT("1: Full")
+	TEXT("Whether to apply screen space VRS to Nanite.\n")
+	TEXT("0: Disabled\n")
+	TEXT("1: Full\n")
 	TEXT("2: Conservative (default)"),
 	FConsoleVariableDelegate::CreateStatic(&CVarVRSImagePassTypeCallback),
 	ECVF_RenderThreadSafe);
+
 TAutoConsoleVariable<int32> CVarVRS_SSAO(
 	TEXT("r.VRS.SSAO"),
 	0,
-	TEXT("Enable VRS with SSAO rendering.\n")
-	TEXT("0: Disabled")
-	TEXT("1: Full")
-	TEXT("2: Conservative (default)"),
+	TEXT("Whether to apply screen space VRS to SSAO (not currently recommended).\n")
+	TEXT("0: Disabled (default)\n")
+	TEXT("1: Full\n")
+	TEXT("2: Conservative"),
 	FConsoleVariableDelegate::CreateStatic(&CVarVRSImagePassTypeCallback),
 	ECVF_RenderThreadSafe);
+
 TAutoConsoleVariable<int32> CVarVRS_SSR(
 	TEXT("r.VRS.SSR"),
 	2,
-	TEXT("Enable VRS with SSR (PS) rendering.\n")
-	TEXT("0: Disabled")
-	TEXT("1: Full")
+	TEXT("Whether to apply screen space VRS to screen space reflections (only supported when r.SSR.Compute=0).\n")
+	TEXT("0: Disabled\n")
+	TEXT("1: Full\n")
 	TEXT("2: Conservative (default)"),
 	FConsoleVariableDelegate::CreateStatic(&CVarVRSImagePassTypeCallback),
 	ECVF_RenderThreadSafe);
+
 TAutoConsoleVariable<int32> CVarVRSReflectionEnvironmentSky(
 	TEXT("r.VRS.ReflectionEnvironmentSky"),
 	2,
-	TEXT("Enable VRS with ReflectionEnvironmentAndSky (PS) rendering.\n")
-	TEXT("0: Disabled")
-	TEXT("1: Full")
+	TEXT("Whether to apply screen space VRS to reflections and indirect sky lighting (only supported when Lumen is disabled).\n")
+	TEXT("0: Disabled\n")
+	TEXT("1: Full\n")
 	TEXT("2: Conservative (default)"),
 	FConsoleVariableDelegate::CreateStatic(&CVarVRSImagePassTypeCallback),
 	ECVF_RenderThreadSafe);
+
 TAutoConsoleVariable<int32> CVarVRSLightFunctions(
 	TEXT("r.VRS.LightFunctions"),
 	1,
-	TEXT("Enables Variable Rate Shading for light functions\n")
-	TEXT("0: Disabled")
-	TEXT("1: Full (default)")
+	TEXT("Whether to apply screen space VRS to light functions.\n")
+	TEXT("0: Disabled\n")
+	TEXT("1: Full (default)\n")
 	TEXT("2: Conservative"),
 	FConsoleVariableDelegate::CreateStatic(&CVarVRSImagePassTypeCallback),
 	ECVF_RenderThreadSafe);
+
 TAutoConsoleVariable<int32> CVarVRSDecals(
 	TEXT("r.VRS.Decals"),
 	2,
-	TEXT("Enables Variable Rate Shading for decals\n")
-	TEXT("0: Disabled")
-	TEXT("1: Full")
+	TEXT("Whether to apply screen space VRS to decals.\n")
+	TEXT("0: Disabled\n")
+	TEXT("1: Full\n")
 	TEXT("2: Conservative (default)"),
 	FConsoleVariableDelegate::CreateStatic(&CVarVRSImagePassTypeCallback),
 	ECVF_RenderThreadSafe);
