LightFunctionAtlas - Fixed assert with view parameters that are not required since view UB is already setup.

#tests PC editor dx12
#rnx
#Synced-CL 49800085

diff --git a/Engine/Source/Runtime/Renderer/Private/LightFunctionAtlas.cpp b/Engine/Source/Runtime/Renderer/Private/LightFunctionAtlas.cpp
--- a/Engine/Source/Runtime/Renderer/Private/LightFunctionAtlas.cpp
+++ b/Engine/Source/Runtime/Renderer/Private/LightFunctionAtlas.cpp
@@ -193,7 +193,6 @@ class FLightFunctionAtlasSlotPS : public FMaterialShader
 	void SetParameters(FRHIBatchedShaderParameters& BatchedParameters, const FSceneView& View, const FMaterialRenderProxy* MaterialProxy)
 	{
 		const FMaterial& Material = MaterialProxy->GetMaterialWithFallback(View.GetFeatureLevel(), MaterialProxy);
-		FMaterialShader::SetViewParameters(BatchedParameters, View, View.ViewUniformBuffer);
 		FMaterialShader::SetParameters(BatchedParameters, MaterialProxy, Material, View);
 	}
 
