Mobile: Add support for a "Lit Wireframe" debug view mode

#Synced-CL 49835704

diff --git a/Engine/Config/VulkanPC/DataDrivenPlatformInfo.ini b/Engine/Config/VulkanPC/DataDrivenPlatformInfo.ini
--- a/Engine/Config/VulkanPC/DataDrivenPlatformInfo.ini
+++ b/Engine/Config/VulkanPC/DataDrivenPlatformInfo.ini
@@ -269,6 +269,7 @@ MaxFeatureLevel=ES3_1
 ShaderFormat=SF_VULKAN_ES31
 bIsPC=true
 bIsMobile=true
+bSupportsDebugViewShaders=true
 ; Enables RDG pass merging(like on Android), helps with debugging
 bTargetsTiledGPU=true
 bSupportsMobileMultiView=true
diff --git a/Engine/Source/Runtime/Renderer/Private/MobileShadingRenderer.cpp b/Engine/Source/Runtime/Renderer/Private/MobileShadingRenderer.cpp
--- a/Engine/Source/Runtime/Renderer/Private/MobileShadingRenderer.cpp
+++ b/Engine/Source/Runtime/Renderer/Private/MobileShadingRenderer.cpp
@@ -1358,6 +1358,7 @@ void FMobileSceneRenderer::Render(FRDGBuilder& GraphBuilder, const FSceneRenderU
 	}
 
 	FRDGTextureRef ViewFamilyTexture = TryCreateViewFamilyTexture(GraphBuilder, ViewFamily);
+	FRDGTextureRef ViewFamilyDepthTexture = TryCreateViewFamilyDepthTexture(GraphBuilder, ViewFamily);
 
 	const FRDGSystemTextures& SystemTextures = FRDGSystemTextures::Get(GraphBuilder);
 
@@ -1661,6 +1662,8 @@ void FMobileSceneRenderer::Render(FRDGBuilder& GraphBuilder, const FSceneRenderU
 		FRendererModule& RendererModule = static_cast<FRendererModule&>(GetRendererModule());
 		RendererModule.RenderPostOpaqueExtensions(GraphBuilder, Views, SceneTextures);
 
+		CopySceneCaptureComponentToTarget(GraphBuilder, SceneTextures, ViewFamilyTexture, ViewFamilyDepthTexture, ViewFamily, Views);
+
 		RenderOpaqueFX(GraphBuilder, GetSceneViews(), GetSceneUniforms(), FXSystem, SceneTextures.MobileUniformBuffer);
 
 		if (ViewFamily.EngineShowFlags.VisualizeMeshDistanceFields || ViewFamily.EngineShowFlags.VisualizeGlobalDistanceField)
@@ -1681,6 +1684,7 @@ void FMobileSceneRenderer::Render(FRDGBuilder& GraphBuilder, const FSceneRenderU
 
 					FMobilePostProcessingInputs PostProcessingInputs;
 					PostProcessingInputs.ViewFamilyTexture = ViewFamilyTexture;
+					PostProcessingInputs.ViewFamilyDepthTexture = ViewFamilyDepthTexture;
 					PostProcessingInputs.SceneTextures = CreateMobileSceneTextureUniformBuffer(GraphBuilder, &SceneTextures, EMobileSceneTextureSetupMode::All);
 
 					for (int32 ViewExt = 0; ViewExt < ViewFamily.ViewExtensions.Num(); ++ViewExt)
@@ -1878,6 +1882,16 @@ FRenderTargetBindingSlots FMobileSceneRenderer::InitRenderTargetBindings_Forward
 	BasePassRenderTargets.SubpassHint = ESubpassHint::None;
 	BasePassRenderTargets.NumOcclusionQueries = 0u;
 
+#if WITH_EDITOR
+	if (ViewFamily.EngineShowFlags.Wireframe)
+	{
+		BasePassRenderTargets[0].SetTexture(SceneTextures.EditorPrimitiveColor);
+		BasePassRenderTargets[0].SetResolveTexture(nullptr);
+		BasePassRenderTargets.DepthStencil.SetTexture(SceneTextures.EditorPrimitiveDepth);
+		BasePassRenderTargets.DepthStencil.SetResolveTexture(nullptr);
+	}
+#endif
+
 	return BasePassRenderTargets;
 }
 
@@ -2357,6 +2371,16 @@ FRenderTargetBindingSlots FMobileSceneRenderer::InitRenderTargetBindings_Deferre
 	//if the scenecolor isn't multiview but the app is, need to render as a single-view multiview due to shaders
 	BasePassRenderTargets.MultiViewCount = (Views[0].bIsMobileMultiViewEnabled) ? 2 : (Views[0].Aspects.IsMobileMultiViewEnabled() ? 1 : 0);
 
+#if WITH_EDITOR
+	if (ViewFamily.EngineShowFlags.Wireframe)
+	{
+		BasePassRenderTargets[0].SetTexture(SceneTextures.EditorPrimitiveColor);
+		BasePassRenderTargets[0].SetResolveTexture(nullptr);
+		BasePassRenderTargets.DepthStencil.SetTexture(SceneTextures.EditorPrimitiveDepth);
+		BasePassRenderTargets.DepthStencil.SetResolveTexture(nullptr);
+	}
+#endif
+
 	return BasePassRenderTargets;
 }
 
@@ -3146,6 +3170,11 @@ void FMobileSceneRenderer::RenderHZB(FRDGBuilder& GraphBuilder, FRDGTextureRef S
 {
 	RDG_EVENT_SCOPE_STAT(GraphBuilder, HZB, "HZB");
 
+	if (!HasBeenProduced(SceneDepthTexture))
+	{
+		return;
+	}
+
 	for (int32 ViewIndex = 0; ViewIndex < Views.Num(); ViewIndex++)
 	{
 		FViewInfo& View = Views[ViewIndex];
diff --git a/Engine/Source/Runtime/Renderer/Private/PostProcess/PostProcessing.cpp b/Engine/Source/Runtime/Renderer/Private/PostProcess/PostProcessing.cpp
--- a/Engine/Source/Runtime/Renderer/Private/PostProcess/PostProcessing.cpp
+++ b/Engine/Source/Runtime/Renderer/Private/PostProcess/PostProcessing.cpp
@@ -3490,8 +3490,10 @@ void AddMobilePostProcessingPasses(FRDGBuilder& GraphBuilder, FScene* Scene, con
 	if (PassSequence.IsEnabled(EPass::EditorPrimitive))
 	{
 		FCompositePrimitiveInputs PassInputs;
-		PassSequence.AcceptOverrideIfLastPass(EPass::EditorPrimitive, PassInputs.OverrideOutput);
-		PassInputs.OverrideDepthOutput = ViewFamilyDepthOutput;
+		if (PassSequence.AcceptOverrideIfLastPass(EPass::EditorPrimitive, PassInputs.OverrideOutput))
+		{
+			PassInputs.OverrideDepthOutput = ViewFamilyDepthOutput;
+		}
 		PassInputs.SceneColor = SceneColor;
 		PassInputs.SceneDepth = SceneDepth;
 		PassInputs.BasePassType = FCompositePrimitiveInputs::EBasePassType::Mobile;
