Fix crash when enabling light complexity view mode.

#Synced-CL 49835472

diff --git a/Engine/Source/Runtime/Renderer/Private/LightRendering.cpp b/Engine/Source/Runtime/Renderer/Private/LightRendering.cpp
--- a/Engine/Source/Runtime/Renderer/Private/LightRendering.cpp
+++ b/Engine/Source/Runtime/Renderer/Private/LightRendering.cpp
@@ -985,7 +985,8 @@ class FDeferredLightPS : public FGlobalShader
 			PermutationVector.Get< FAtmosphereTransmittance >() ||
 			PermutationVector.Get< FCloudTransmittance >() ||
 			PermutationVector.Get< FAnistropicMaterials >() ||
-			PermutationVector.Get< FVirtualShadowMapMask >()))
+			PermutationVector.Get< FVirtualShadowMapMask >() || 
+			PermutationVector.Get< FUbershader >()))
 		{
 			return false;
 		}
@@ -1128,6 +1129,7 @@ class FDeferredLightPS : public FGlobalShader
 			PermutationVector.Set< FAnistropicMaterials >(false);
 			PermutationVector.Set< FVirtualShadowMapMask >(false);
 			PermutationVector.Set< FHairComplexTransmittance >(false);
+			PermutationVector.Set< FUbershader >(false);
 		}
 
 		if (Substrate::IsSubstrateEnabled())
