Custom Render Passes:  Temporarily disable decal rendering to fix crash bug while investigating console specific crash.

#rnx

#Synced-CL 49859170

diff --git a/Engine/Source/Runtime/Renderer/Private/DeferredShadingRenderer.cpp b/Engine/Source/Runtime/Renderer/Private/DeferredShadingRenderer.cpp
--- a/Engine/Source/Runtime/Renderer/Private/DeferredShadingRenderer.cpp
+++ b/Engine/Source/Runtime/Renderer/Private/DeferredShadingRenderer.cpp
@@ -2428,13 +2428,17 @@ void FDeferredShadingSceneRenderer::Render(FRDGBuilder& GraphBuilder, const FSce
 						Nanite::BuildShadingCommands(GraphBuilder, *Scene, ENaniteMeshPass::BasePass, NaniteBasePassShadingCommands, Nanite::EBuildShadingCommandsMode::Custom);
 					}
 
+#if 0	// TODO:  temporary fix for PLAY-109743 while investigating, comment out decal rendering in Custom Render Passes
 					// This class handles decal rendering
 					FCompositionLighting CustomRenderPassCompositionLighting(InitViewTaskDatas.Decals, CustomRenderPassViews, *CustomRenderPassSceneTextures, [this](int32 ViewIndex) { return false; });
 					CustomRenderPassCompositionLighting.ProcessBeforeBasePass(GraphBuilder, *CustomRenderPassSceneTextures->DBufferTextures, InstanceCullingManager, *CustomRenderPassSceneTextures->SubstrateSceneData);
+#endif
 
 					RenderBasePass(*this, GraphBuilder, CustomRenderPassViews, *CustomRenderPassSceneTextures, BasePassDepthStencilAccess, /*ForwardScreenSpaceShadowMaskTexture=*/nullptr, InstanceCullingManager, bNaniteEnabled, NaniteBasePassShadingCommands, NaniteRasterResults);
 
+#if 0	// TODO:  temporary fix for PLAY-109743
 					CustomRenderPassCompositionLighting.ProcessAfterBasePass(GraphBuilder, InstanceCullingManager, FCompositionLighting::EProcessAfterBasePassMode::All, *CustomRenderPassSceneTextures->SubstrateSceneData);
+#endif
 
 					if (ShouldRenderSingleLayerWater(CustomRenderPassViews))
 					{
