[Engine/Renderer]
* Fixed code paths when statestream is enabled. SkeletalMeshUpdater needs to be destroyed render side when statestream is enabled.

#Synced-CL 49872240

diff --git a/Engine/Source/Runtime/Engine/Private/World.cpp b/Engine/Source/Runtime/Engine/Private/World.cpp
--- a/Engine/Source/Runtime/Engine/Private/World.cpp
+++ b/Engine/Source/Runtime/Engine/Private/World.cpp
@@ -6315,7 +6315,10 @@ void UWorld::CleanupWorldInternal(bool bSessionEnded, bool bCleanupResources, bo
 	ClearWorldComponents();
 
 #if WITH_STATE_STREAM
-	StateStreamManager->Game_EndTick(RealTimeSeconds);
+	if (StateStreamManager)
+	{
+		StateStreamManager->Game_EndTick(RealTimeSeconds);
+	}
 #endif
 
 	if (bCleanupResources && PersistentLevel)
diff --git a/Engine/Source/Runtime/Renderer/Private/RendererScene.cpp b/Engine/Source/Runtime/Renderer/Private/RendererScene.cpp
--- a/Engine/Source/Runtime/Renderer/Private/RendererScene.cpp
+++ b/Engine/Source/Runtime/Renderer/Private/RendererScene.cpp
@@ -4272,10 +4272,12 @@ void FScene::ReleaseUnusedPersistentSBTs()
 void FScene::EndFrame(FRHICommandListImmediate& RHICmdList)
 {
 	// Performs a final update of any queued scene primitives in the case where the scene wasn't rendered to avoid a build-up of queued data that is never flushed.
+#if !WITH_STATE_STREAM
 	if (LastUpdateFrameCounter != GFrameCounterRenderThread)
 	{
 		UpdateAllPrimitiveSceneInfos(RHICmdList);
 	}
+#endif
 
 #if RHI_RAYTRACING
 	ReleaseUnusedPersistentSBTs();
@@ -4319,29 +4321,34 @@ void FScene::Release()
 
 	UE::RenderCommandPipe::FSyncScope SyncScope;
 
+	FScene* Scene = this;
+
 #if WITH_STATE_STREAM
-	FStateStreamManagerImpl* StateStreamManager = static_cast<FStateStreamManagerImpl*>(World->GetStateStreamManager());
+	if (FStateStreamManagerImpl* StateStreamManager = static_cast<FStateStreamManagerImpl*>(World->GetStateStreamManager()))
+	{
+		ENQUEUE_RENDER_COMMAND(FReleaseStateStream)(
+			[Scene, StateStreamManager](FRHICommandListImmediate& RHICmdList)
+			{
+				StateStreamManager->Render_Exit();
+				GetRendererModule().FreeStateStream(StateStreamManager);
+				if (FSkeletalMeshUpdater* Smu = Scene->SkeletalMeshUpdater)
+				{
+					Smu->Shutdown();
+				}
+			});
+	}
+	else
 #endif
-
 	if (SkeletalMeshUpdater)
 	{
 		SkeletalMeshUpdater->Shutdown();
 	}
 
 	// Send a command to the rendering thread to release the scene.
-	FScene* Scene = this;
 	ENQUEUE_RENDER_COMMAND(FReleaseCommand)(
 		[Scene
-#if WITH_STATE_STREAM
-		, StateStreamManager
-#endif
 		](FRHICommandListImmediate& RHICmdList)
 		{
-#if WITH_STATE_STREAM
-			StateStreamManager->Render_Exit();
-			GetRendererModule().FreeStateStream(StateStreamManager);
-#endif
-
 			// Flush any remaining batched primitive update commands before deleting the scene.
 			FUpdateParameters UpdateParameters;
 			UpdateParameters.bDestruction = true;
