Path Tracing: Remove rarely used EnableBackfaceCulling cvar and default to on

This cvar was mainly for debugging in the early days of the path tracer, but it is rarely used in practice. The value is now hardcoded to always on.

#Synced-CL 49848588

diff --git a/Engine/Plugins/Experimental/GPULightmass/Shaders/Private/LightmapPathTracing.usf b/Engine/Plugins/Experimental/GPULightmass/Shaders/Private/LightmapPathTracing.usf
--- a/Engine/Plugins/Experimental/GPULightmass/Shaders/Private/LightmapPathTracing.usf
+++ b/Engine/Plugins/Experimental/GPULightmass/Shaders/Private/LightmapPathTracing.usf
@@ -129,7 +129,6 @@ static const uint MaxBounces = 32;
 static const uint MISMode = 2;
 
 static const uint ApproximateCaustics = 1;
-static const uint EnableCameraBackfaceCulling = 0;
 static const uint EnableEmissive = 1;
 static const float MaxNormalBias = 0.1;
 
@@ -159,7 +158,7 @@ bool IsInvalidSurfaceHit(FPathTracingPayload Payload)
 
 FPathTracingPayload TraceTransparentRay(FRayDesc Ray, bool IsCameraRay, bool LastBounce, bool IncludeEmission, uint2 LaunchIndex, uint NumLights, inout RandomSequence RandSequence, inout float3 PathThroughput, inout float3 Radiance)
 {
-	const uint RayFlags = IsCameraRay && EnableCameraBackfaceCulling ? RAY_FLAG_CULL_BACK_FACING_TRIANGLES : 0;
+	const uint RayFlags = 0;
 	const uint MissShaderIndex = 0;
 	float SelectionWeightSum = 0;
 	float3 PayloadThroughput = PathThroughput;
diff --git a/Engine/Shaders/Private/PathTracing/PathTracingCore.ush b/Engine/Shaders/Private/PathTracing/PathTracingCore.ush
--- a/Engine/Shaders/Private/PathTracing/PathTracingCore.ush
+++ b/Engine/Shaders/Private/PathTracing/PathTracingCore.ush
@@ -71,7 +71,6 @@ uint MISMode;
 uint VolumeMISMode;
 
 uint ApproximateCaustics;
-uint EnableCameraBackfaceCulling;
 uint SamplerType;
 uint EnableDBuffer;
 uint ApplyDiffuseSpecularOverrides;
@@ -474,7 +473,7 @@ FPackedPathTracingPayload TraceTransparentRay(inout FPathState PathState, int Bo
 {
 	const bool bIsCameraRay = Bounce == 0;
 	const bool bLastBounce = Bounce == MaxBounces;
-	const uint RayFlags = bIsCameraRay && EnableCameraBackfaceCulling ? RAY_FLAG_CULL_BACK_FACING_TRIANGLES : 0;
+	const uint RayFlags = bIsCameraRay ? RAY_FLAG_CULL_BACK_FACING_TRIANGLES : 0;
 	const bool bUseCameraRayType = !PathState.HasMadeContributionScatter();
 	uint InstanceInclusionMask = bUseCameraRayType ? PATHTRACER_MASK_CAMERA : PATHTRACER_MASK_INDIRECT;
 	if (length2(PathState.Ray.Origin) < Pow2(PATH_TRACING_SKIP_HAIR_DISTANCE))
diff --git a/Engine/Source/Runtime/Renderer/Private/PathTracing.cpp b/Engine/Source/Runtime/Renderer/Private/PathTracing.cpp
--- a/Engine/Source/Runtime/Renderer/Private/PathTracing.cpp
+++ b/Engine/Source/Runtime/Renderer/Private/PathTracing.cpp
@@ -317,13 +317,6 @@ TAutoConsoleVariable<int32> CVarPathTracingEnableEmissive(
 	ECVF_RenderThreadSafe
 );
 
-TAutoConsoleVariable<int32> CVarPathTracingEnableCameraBackfaceCulling(
-	TEXT("r.PathTracing.EnableCameraBackfaceCulling"),
-	1,
-	TEXT("When non-zero, the path tracer will skip over backfacing triangles when tracing primary rays from the camera."),
-	ECVF_RenderThreadSafe
-);
-
 static int32 GEnableReferenceDOF = -1;
 FAutoConsoleVariableRef CVarPathTracingEnableReferenceDOF(
 	TEXT("r.PathTracing.EnableReferenceDOF"),
@@ -661,7 +654,6 @@ BEGIN_SHADER_PARAMETER_STRUCT(FPathTracingData, )
 	SHADER_PARAMETER(uint32, MISMode)
 	SHADER_PARAMETER(uint32, VolumeMISMode)
 	SHADER_PARAMETER(uint32, ApproximateCaustics)
-	SHADER_PARAMETER(uint32, EnableCameraBackfaceCulling)
 	SHADER_PARAMETER(uint32, SamplerType)
 	SHADER_PARAMETER(uint32, EnableDBuffer)
 	SHADER_PARAMETER(uint32, VolumeFlags)
@@ -719,7 +711,6 @@ struct FPathTracingConfig
 			PathTracingData.VolumeMISMode != Other.PathTracingData.VolumeMISMode ||
 			PathTracingData.SamplerType != Other.PathTracingData.SamplerType ||
 			PathTracingData.ApproximateCaustics != Other.PathTracingData.ApproximateCaustics ||
-			PathTracingData.EnableCameraBackfaceCulling != Other.PathTracingData.EnableCameraBackfaceCulling ||
 			PathTracingData.EnableDBuffer != Other.PathTracingData.EnableDBuffer ||
 			PathTracingData.MaxPathIntensity != Other.PathTracingData.MaxPathIntensity ||
 			PathTracingData.FilterWidth != Other.PathTracingData.FilterWidth ||
@@ -933,7 +924,6 @@ static void PreparePathTracingData(const FScene* Scene, const FViewInfo& View, F
 	}
 	PathTracingData.MaxPathIntensity = FFloat16(PathTracingData.MaxPathIntensity).GetClampedNonNegativeAndFinite().GetFloat(); // Clip to half precision
 	PathTracingData.ApproximateCaustics = CVarPathTracingApproximateCaustics.GetValueOnRenderThread();
-	PathTracingData.EnableCameraBackfaceCulling = CVarPathTracingEnableCameraBackfaceCulling.GetValueOnRenderThread();
 	PathTracingData.SamplerType = CVarPathTracingSamplerType.GetValueOnRenderThread();
 	PathTracingData.FilterWidth = CVarPathTracingFilterWidth.GetValueOnRenderThread();
 	PathTracingData.CameraFocusDistance = 0;
