[PostProcessing] Add more vignette options. Vignette can now be rendered as an analytical circle, or from a texture. Vignette can now be moved, scaled, colorized and feathered.

#Synced-CL 49769684

diff --git a/Engine/Plugins/Cameras/GameplayCameras/Source/GameplayCameras/Private/Compat/PostProcessUtils.cpp b/Engine/Plugins/Cameras/GameplayCameras/Source/GameplayCameras/Private/Compat/PostProcessUtils.cpp
--- a/Engine/Plugins/Cameras/GameplayCameras/Source/GameplayCameras/Private/Compat/PostProcessUtils.cpp
+++ b/Engine/Plugins/Cameras/GameplayCameras/Source/GameplayCameras/Private/Compat/PostProcessUtils.cpp
@@ -117,6 +117,12 @@ bool FPostProcessUtils::OverridePostProcessSettings(FPostProcessSettings& ThisFr
 		UE_SET_PP(LensFlareBokehSize);
 		UE_SET_PP(LensFlareThreshold);
 		UE_SET_PP(VignetteIntensity);
+		UE_SET_PP(VignetteType);
+		UE_SET_PP(VignetteCenter);
+		UE_SET_PP(VignetteColor);
+		UE_SET_PP(VignetteSize);
+		UE_SET_PP(VignetteSoftness);
+		UE_SET_PP(VignetteTexture);
 		UE_SET_PP(Sharpen);
 		UE_SET_PP(FilmGrainIntensity);
 		UE_SET_PP(FilmGrainIntensityShadows);
@@ -401,6 +407,12 @@ bool FPostProcessUtils::BlendPostProcessSettings(FPostProcessSettings& ThisFrom,
 		UE_LERP_PP(LensFlareBokehSize);
 		UE_LERP_PP(LensFlareThreshold);
 		UE_LERP_PP(VignetteIntensity);
+		UE_SET_PP(VignetteType);
+		UE_LERP_PP(VignetteCenter);
+		UE_LERP_PP(VignetteColor);
+		UE_LERP_PP(VignetteSize);
+		UE_LERP_PP(VignetteSoftness);
+		UE_SET_PP(VignetteTexture);
 		UE_LERP_PP(Sharpen);
 		UE_LERP_PP(FilmGrainIntensity);
 		UE_LERP_PP(FilmGrainIntensityShadows);
diff --git a/Engine/Plugins/Cameras/GameplayCameras/Source/GameplayCameras/Private/Debug/PostProcessSettingsDebugBlock.cpp b/Engine/Plugins/Cameras/GameplayCameras/Source/GameplayCameras/Private/Debug/PostProcessSettingsDebugBlock.cpp
--- a/Engine/Plugins/Cameras/GameplayCameras/Source/GameplayCameras/Private/Debug/PostProcessSettingsDebugBlock.cpp
+++ b/Engine/Plugins/Cameras/GameplayCameras/Source/GameplayCameras/Private/Debug/PostProcessSettingsDebugBlock.cpp
@@ -152,6 +152,11 @@ void FPostProcessSettingsDebugBlock::OnDebugDraw(const FCameraDebugBlockDrawPara
 		UE_DRAW_PP(LensFlareBokehSize);
 		UE_DRAW_PP(LensFlareThreshold);
 		UE_DRAW_PP(VignetteIntensity);
+		UE_DRAW_PP(VignetteType);
+		UE_DRAW_PP(VignetteCenter);
+		UE_DRAW_PP(VignetteColor);
+		UE_DRAW_PP(VignetteSize);
+		UE_DRAW_PP(VignetteSoftness);
 		UE_DRAW_PP(Sharpen);
 		UE_DRAW_PP(FilmGrainIntensity);
 		UE_DRAW_PP(FilmGrainIntensityShadows);
diff --git a/Engine/Plugins/Compositing/Composite/Shaders/Private/CompositeColorKeyerCommon.ush b/Engine/Plugins/Compositing/Composite/Shaders/Private/CompositeColorKeyerCommon.ush
--- a/Engine/Plugins/Compositing/Composite/Shaders/Private/CompositeColorKeyerCommon.ush
+++ b/Engine/Plugins/Compositing/Composite/Shaders/Private/CompositeColorKeyerCommon.ush
@@ -161,8 +161,13 @@ float3 BlueDespillAvg(float3 Color)
 float Vignette(float2 ScreenPos, float AspectRatio, float Intensity)
 {
 	// Scale vignette to always be a circle with consistent corner intensity.
-	float2 Vignette = VignetteSpace(ScreenPos, AspectRatio);
-	return ComputeVignetteMask(Vignette, Intensity);
+	float2 Vignette = VignetteSpace(ScreenPos, VIGNETTE_TYPE_COSINE_FOURTH_LAW, AspectRatio);
+	
+	float4 VignetteParams0, VignetteParams1;
+	float VignetteParams2;
+	PackVignetteParameters(VignetteParams0, VignetteParams1, VignetteParams2, Intensity);
+
+	return ComputeVignetteMask(Vignette, VIGNETTE_TYPE_COSINE_FOURTH_LAW, VignetteParams0, VignetteParams1, VignetteParams2);
 }
 
 float4 GreenScreenKeyer(
diff --git a/Engine/Shaders/Private/PostProcessCommon.ush b/Engine/Shaders/Private/PostProcessCommon.ush
--- a/Engine/Shaders/Private/PostProcessCommon.ush
+++ b/Engine/Shaders/Private/PostProcessCommon.ush
@@ -135,33 +135,102 @@ half3 ColorCorrection(half3 InLDRColor)
 	return MappingPolynomial.x * (InLDRColor * InLDRColor) + MappingPolynomial.y * InLDRColor + MappingPolynomial.z;
 }
 
+// See EVignetteType
+#define VIGNETTE_TYPE_COSINE_FOURTH_LAW (0)
+#define VIGNETTE_TYPE_TEXTURE			(1)
+#define VIGNETTE_TYPE_OVAL				(2)
+
+void PackVignetteParameters(out float4 OutVignetteParams0, out float4 OutVignetteParams1, out float OutVignetteParams2, float Intensity)
+{
+	float3 Color	= float3(1.0, 1.0, 1.0);
+	float2 Center	= float2(0.5, 0.5);
+	float2 InvSize	= float2(1.0, 1.0);
+	float Softness	= 1.0;
+
+	OutVignetteParams0.x	= Intensity;
+	OutVignetteParams0.yzw	= Color;
+	OutVignetteParams1.xy	= Center;
+	OutVignetteParams1.zw	= InvSize;
+	OutVignetteParams2		= Softness;
+}
 
 // Generate a mask to darken the screen borders.
 // (Camera artifact and artistic effect)
 // @param VignetteCircleSpacePos from VignetteSpace()
-float ComputeVignetteMask(float2 VignetteCircleSpacePos, float Intensity)
+float3 ComputeVignetteMask(
+	float2 VignetteCircleSpacePos, 
+	uint VignetteType, 
+	float4 VignetteParams0, 
+	float4 VignetteParams1, 
+	float  VignetteParams2
+#if USE_VIGNETTE_TEXTURE
+	, Texture2D VignetteTexture
+	, SamplerState TextureSampler
+#endif
+)
 {
-	// Natural vignetting
-	// cosine-fourth law
-	VignetteCircleSpacePos *= Intensity;
-	float Tan2Angle = dot( VignetteCircleSpacePos, VignetteCircleSpacePos );
-	float Cos4Angle = Square( rcp( Tan2Angle + 1 ) );
-	return Cos4Angle;
+	float Intensity = VignetteParams0.x;
+	float3 Color	= VignetteParams0.yzw;
+	float2 Center	= VignetteParams1.xy - 0.5;
+	float2 InvSize	= VignetteParams1.zw;
+	float Softness	= VignetteParams2;
+
+	float Weight = 1.0;
+	BRANCH
+	if (Intensity <= 0)
+	{ }
+	else if (VignetteType == VIGNETTE_TYPE_COSINE_FOURTH_LAW)
+	{
+		VignetteCircleSpacePos -= Center;
+		VignetteCircleSpacePos *= Intensity;
+		float Tan2Angle = dot(VignetteCircleSpacePos, VignetteCircleSpacePos);
+		float Cos4Angle = Square(rcp(Tan2Angle + 1));
+		Weight = Cos4Angle;
+	}
+#if USE_VIGNETTE_TEXTURE
+	else if (VignetteType == VIGNETTE_TYPE_TEXTURE)
+	{
+		Weight = VignetteTexture.SampleLevel(TextureSampler, (InvSize * VignetteCircleSpacePos - Center) * 0.5 + 0.5, 0).r;
+		Weight = lerp(1.0, Weight, Intensity);
+	}
+#endif
+	else
+	{
+		float2 Position = InvSize * VignetteCircleSpacePos;
+
+		float Distance = 0;
+		if (VignetteType == VIGNETTE_TYPE_OVAL)
+		{
+			Distance = length2(Position - Center);
+		}
+
+		Weight = pow(saturate(1.0 - Distance), Softness);
+		Weight = lerp(1.0, Weight, Intensity);
+	}
+
+	return lerp(Color, (float3)1, Weight);
 }
 
 // Scale {-1 to 1} of the viewport space to vignette circle space.
 // Vignette space is scaled such that regardless of viewport aspect ratio, 
 // corners are at the same brightness on a circle.
-float2 VignetteSpace(float2 Pos, float AspectRatio)
+float2 VignetteSpace(float2 Pos, uint VignetteType, float AspectRatio)
 {
-	// could be optimized but this computation should be done in the vertex shader (3 or 4 vertices)
-	float Scale = sqrt(2.0) / sqrt(1.0 + AspectRatio * AspectRatio);
-	return Pos * float2(1.0, AspectRatio) * Scale;
+	if (VignetteType == VIGNETTE_TYPE_TEXTURE)
+	{
+		float InvAspectRatio = rcp(AspectRatio);
+		return Pos * float2(saturate(InvAspectRatio), saturate(AspectRatio));
+	}
+	else
+	{
+		float Scale = sqrt(2.0) / sqrt(1.0 + AspectRatio * AspectRatio);
+		return Pos * float2(1.0, AspectRatio) * Scale;
+	}
 }
 
-float2 VignetteSpace(float2 Pos)
+float2 VignetteSpace(float2 Pos, uint VignetteType)
 {
-	return VignetteSpace(Pos, ViewportSize.y * ViewportSize.z);
+	return VignetteSpace(Pos, VignetteType, ViewportSize.y * ViewportSize.z);
 }
 
 half4 UnwrappedTexture3DSample( Texture2D Texture, SamplerState Sampler, float3 UVW, float Size, float InvSize )
diff --git a/Engine/Shaders/Private/PostProcessMobile.usf b/Engine/Shaders/Private/PostProcessMobile.usf
--- a/Engine/Shaders/Private/PostProcessMobile.usf
+++ b/Engine/Shaders/Private/PostProcessMobile.usf
@@ -1150,7 +1150,6 @@ void SunMergeVS_Mobile(
 	in FStereoVSInput StereoInput,
 	out float2 OutTexCoord : TEXCOORD0,
 	out float4 OutSunShaftUVs[4] : TEXCOORD1,
-	out float2 OutVignetteSpace : TEXCOORD5,
 	out FStereoVSOutput StereoOutput,
 	out float4 OutPosition : SV_POSITION
 	)
@@ -1166,8 +1165,6 @@ void SunMergeVS_Mobile(
 	OutSunShaftUVs[2].zw = SunShaftRect(OutPosition.xy, 1.0 - SUN_P5 * SUN_M);
 	OutSunShaftUVs[3].xy = SunShaftRect(OutPosition.xy, 1.0 - SUN_P6 * SUN_M);
 	OutSunShaftUVs[3].zw = float2(0.0, 0.0);
-
-    OutVignetteSpace = VignetteSpace(OutPosition.xy);
 }
 
 #undef SUN_M
@@ -1187,7 +1184,6 @@ float4 SunMergeBloomUpTextureSizeAndInvSize;
 void SunMergePS_Mobile(
 	float2 InUV : TEXCOORD0,
 	float4 InSunShaftUVs[4] : TEXCOORD1,
-	float2 InVignetteSpace : TEXCOORD5,
 	in FStereoPSInput StereoInput,
 	out HALF4_TYPE OutColor : SV_Target0
 	)
diff --git a/Engine/Shaders/Private/PostProcessTonemap.usf b/Engine/Shaders/Private/PostProcessTonemap.usf
--- a/Engine/Shaders/Private/PostProcessTonemap.usf
+++ b/Engine/Shaders/Private/PostProcessTonemap.usf
@@ -14,6 +14,10 @@
 #define LOCAL_EXPOSURE_MODE_BILATERAL 1
 #define LOCAL_EXPOSURE_MODE_FUSION 2
 
+#ifndef USE_VIGNETTE_TEXTURE
+#define USE_VIGNETTE_TEXTURE 0
+#endif
+
 #include "Common.ush"
 #include "PostProcessCommon.ush"
 #include "TonemapCommon.ush"
@@ -68,8 +72,16 @@ float4 ColorScale0;
 // xyz:Bloom1Tint.rgb, w:unused
 float4 ColorScale1;
 
-// from the postprocess settings, x:VignetteIntensity, y:SharpenDiv6
-float4 TonemapperParams;
+float SharpenDiv6;
+
+float4 VignetteParams0;
+float4 VignetteParams1;
+float  VignetteParams2;
+uint   VignetteType;
+#if USE_VIGNETTE_TEXTURE
+Texture2D VignetteTexture;
+SamplerState VignetteTextureSampler;
+#endif
 
 // Grain quantization
 float3 GrainRandomFull;
@@ -157,6 +169,11 @@ float ComputeFilmGrainIntensity(float3 LinearColorRGBGamut)
 	return FilmGrainIntensity;
 }
 
+bool IsFilmGrainEnabled()
+{
+	return FilmGrainIntensityShadows != 0 || FilmGrainIntensityMidtones != 0 || FilmGrainIntensityHighlights != 0;
+}
+
 
 // LUT for color grading
 Texture3D ColorGradingLUT;
@@ -224,7 +241,7 @@ void TonemapCommonVS(
 
 	// Scale vignette to always be a circle with consistent corner intensity.
 	float2 LensViewportPos = ConvertScreenViewportSpaceToLensViewportSpace(ColorViewportPos);
-	OutVignette = VignetteSpace(LensViewportPos, AspectRatio);
+	OutVignette = VignetteSpace(LensViewportPos, VignetteType, AspectRatio);
 
 	// Grain
 	OutGrainUV.xy = TexCoord + Color_ExtentInverse * float2(-0.5,0.5);
@@ -422,7 +439,7 @@ float4 TonemapCommonPS(
 	{
 		// 0..x, 0:no effect .. 1:strong, from r.Tonemapper.Sharpen
 		// Div6 is to save one instruction
-		float SharpenMultiplierDiv6 = TonemapperParams.y;
+		float SharpenMultiplierDiv6 = SharpenDiv6;
 
 		half A0 = Luminance(SceneColor.rgb);
 
@@ -490,13 +507,14 @@ float4 TonemapCommonPS(
 #else // !USE_GAMMA_ONLY
 	{
 		// Compute vignette.
-		float VignetteMask = MobileVignette;
+		half3 VignetteMask = MobileVignette;
 		
-		BRANCH
-		if (TonemapperParams.x > 0.0f)
-		{
-			VignetteMask *= ComputeVignetteMask(Vignette, TonemapperParams.x);
-		}
+		VignetteMask *= ComputeVignetteMask(
+			Vignette, VignetteType, VignetteParams0, VignetteParams1, VignetteParams2
+#if USE_VIGNETTE_TEXTURE
+			, VignetteTexture, VignetteTextureSampler
+#endif
+		);
 		
 		// Compute the final linear color space with explicitly all the terms that should affect the pre exposure.
 		half3 FinalLinearColor = SceneColor.rgb * SceneColorTint * (OneOverPreExposure * GlobalExposure * VignetteMask * LocalExposure);
@@ -506,6 +524,7 @@ float4 TonemapCommonPS(
 		
 		// Film grain is energy conservative so doesn't need to be taken into PreExposure
 		#if USE_FILM_GRAIN
+		if (IsFilmGrainEnabled())
 		{
 			float3 FilmGrainDecodeMultiply = FilmGrainTextureConstants[0].rgb;
 
diff --git a/Engine/Source/Runtime/Engine/Classes/Engine/Scene.h b/Engine/Source/Runtime/Engine/Classes/Engine/Scene.h
--- a/Engine/Source/Runtime/Engine/Classes/Engine/Scene.h
+++ b/Engine/Source/Runtime/Engine/Classes/Engine/Scene.h
@@ -127,6 +127,14 @@ enum class EReflectedAndRefractedRayTracedShadows : uint8
 	Area_shadows	UMETA(DisplayName = "Area Shadows"),
 };
 
+UENUM()
+enum class EVignetteType : uint8
+{
+	CosineFourthLaw	= 0 UMETA(DisplayName = "Cosine Fourth Law (Default)"),
+	Texture			= 1 UMETA(DisplayName = "Texture"),
+	Oval			= 2 UMETA(DisplayName = "Oval")
+};
+
 USTRUCT(BlueprintType)
 struct FColorGradePerRangeSettings
 {
@@ -1003,6 +1011,24 @@ struct FPostProcessSettings
 	UPROPERTY(EditAnywhere, BlueprintReadWrite, Category=Overrides, meta=(PinHiddenByDefault, InlineEditConditionToggle))
 	uint8 bOverride_VignetteIntensity:1;
 
+	UPROPERTY(EditAnywhere, BlueprintReadWrite, Category=Overrides, meta=(PinHiddenByDefault, InlineEditConditionToggle))
+	uint8 bOverride_VignetteType:1;
+
+	UPROPERTY(EditAnywhere, BlueprintReadWrite, Category=Overrides, meta=(PinHiddenByDefault, InlineEditConditionToggle))
+	uint8 bOverride_VignetteCenter:1;
+
+	UPROPERTY(EditAnywhere, BlueprintReadWrite, Category=Overrides, meta=(PinHiddenByDefault, InlineEditConditionToggle))
+	uint8 bOverride_VignetteColor:1;
+
+	UPROPERTY(EditAnywhere, BlueprintReadWrite, Category=Overrides, meta=(PinHiddenByDefault, InlineEditConditionToggle))
+	uint8 bOverride_VignetteSize:1;
+
+	UPROPERTY(EditAnywhere, BlueprintReadWrite, Category=Overrides, meta=(PinHiddenByDefault, InlineEditConditionToggle))
+	uint8 bOverride_VignetteSoftness:1;
+
+	UPROPERTY(EditAnywhere, BlueprintReadWrite, Category=Overrides, meta=(PinHiddenByDefault, InlineEditConditionToggle))
+	uint8 bOverride_VignetteTexture:1;
+
 	UPROPERTY(EditAnywhere, BlueprintReadWrite, Category=Overrides, meta=(PinHiddenByDefault, InlineEditConditionToggle))
 	uint8 bOverride_Sharpen:1;
 
@@ -2085,6 +2111,30 @@ struct FPostProcessSettings
 	UPROPERTY(interp, BlueprintReadWrite, Category="Lens|Image Effects", meta=(UIMin = "0.0", UIMax = "1.0", editcondition = "bOverride_VignetteIntensity"))
 	float VignetteIntensity;
 
+	/** The shape type of the vignette */
+	UPROPERTY(EditAnywhere, BlueprintReadWrite, Category="Lens|Image Effects", AdvancedDisplay, meta=(editcondition = "bOverride_VignetteType"))
+	EVignetteType VignetteType;
+
+	/** Position of the vignette center in UV space.  */
+	UPROPERTY(interp, BlueprintReadWrite, Category="Lens|Image Effects", AdvancedDisplay, meta=(UIMin = "0.0", UIMax = "1.0", editcondition = "bOverride_VignetteCenter"))
+	FVector2f VignetteCenter;
+
+	/** Color multiplied with areas that are masked by the vignette. Default is black. */
+	UPROPERTY(interp, BlueprintReadWrite, Category="Lens|Image Effects", AdvancedDisplay, meta=(editcondition = "bOverride_VignetteColor"))
+	FLinearColor VignetteColor;
+
+	/** Scale of the vignette shape in screenspace. Can be used to scale vignette anamorphically. Default is 1. */
+	UPROPERTY(interp, BlueprintReadWrite, Category="Lens|Image Effects", AdvancedDisplay, meta=(ClampMin = "0.0", UIMin = "0.1", UIMax = "2.0", editcondition = "bOverride_VignetteSize"))
+	FVector2f VignetteSize;
+
+	/** The softness of the transition from opaque to transparant on the edge of the vignette shape. */
+	UPROPERTY(interp, BlueprintReadWrite, Category="Lens|Image Effects", AdvancedDisplay, meta=(ClampMin = "0.001", UIMax = "10.0", editcondition = "bOverride_VignetteSoftness"))
+	float VignetteSoftness;
+
+	/** Defines the shape of the vignette when VignetteType is set to Texture */
+	UPROPERTY(EditAnywhere, BlueprintReadWrite, Category="Lens|Image Effects", AdvancedDisplay, meta=(editcondition = "bOverride_VignetteTexture", DisplayName = "Vignette Texture"))
+	TObjectPtr<class UTexture> VignetteTexture;
+
 	/** Controls the strength of image sharpening applied during tonemapping. */
 	UPROPERTY(interp, BlueprintReadWrite, Category = "Lens|Image Effects", meta = (ClampMin = "0.0", ClampMax = "10.0", UIMin = "0.0", UIMax = "2.0", editcondition = "bOverride_Sharpen"))
 	float Sharpen;
diff --git a/Engine/Source/Runtime/Engine/Private/PostProcessUtils.cpp b/Engine/Source/Runtime/Engine/Private/PostProcessUtils.cpp
--- a/Engine/Source/Runtime/Engine/Private/PostProcessUtils.cpp
+++ b/Engine/Source/Runtime/Engine/Private/PostProcessUtils.cpp
@@ -118,6 +118,12 @@ bool FPostProcessUtils::OverridePostProcessSettings(FPostProcessSettings& ThisFr
 		UE_SET_PP(LensFlareBokehSize);
 		UE_SET_PP(LensFlareThreshold);
 		UE_SET_PP(VignetteIntensity);
+		UE_SET_PP(VignetteType);
+		UE_SET_PP(VignetteCenter);
+		UE_SET_PP(VignetteColor);
+		UE_SET_PP(VignetteSize);
+		UE_SET_PP(VignetteSoftness);
+		UE_SET_PP(VignetteTexture);
 		UE_SET_PP(Sharpen);
 		UE_SET_PP(FilmGrainIntensity);
 		UE_SET_PP(FilmGrainIntensityShadows);
@@ -424,6 +430,12 @@ bool FPostProcessUtils::BlendPostProcessSettings(FPostProcessSettings& ThisFrom,
 		UE_LERP_PP(LensFlareBokehSize);
 		UE_LERP_PP(LensFlareThreshold);
 		UE_LERP_PP(VignetteIntensity);
+		UE_SET_PP(VignetteType);
+		UE_LERP_PP(VignetteCenter);
+		UE_LERP_PP(VignetteColor);
+		UE_LERP_PP(VignetteSize);
+		UE_LERP_PP(VignetteSoftness);
+		UE_SET_PP(VignetteTexture);
 		UE_LERP_PP(Sharpen);
 		UE_LERP_PP(FilmGrainIntensity);
 		UE_LERP_PP(FilmGrainIntensityShadows);
diff --git a/Engine/Source/Runtime/Engine/Private/Scene.cpp b/Engine/Source/Runtime/Engine/Private/Scene.cpp
--- a/Engine/Source/Runtime/Engine/Private/Scene.cpp
+++ b/Engine/Source/Runtime/Engine/Private/Scene.cpp
@@ -538,6 +538,11 @@ FPostProcessSettings::FPostProcessSettings()
 	LensFlareBokehSize = 3.0f;
 	LensFlareThreshold = 8.0f;
 	VignetteIntensity = 0.4f;
+	VignetteType = EVignetteType::CosineFourthLaw;
+	VignetteCenter = FVector2f(0.5f, 0.5f);
+	VignetteColor = FLinearColor::Black;
+	VignetteSize = FVector2f(1.0f, 1.0f);
+	VignetteSoftness = 5.0f;
 	Sharpen = 0.0f;
 	GrainIntensity_DEPRECATED = 0.0f;
 	GrainJitter_DEPRECATED = 0.0f;
diff --git a/Engine/Source/Runtime/Engine/Private/SceneView.cpp b/Engine/Source/Runtime/Engine/Private/SceneView.cpp
--- a/Engine/Source/Runtime/Engine/Private/SceneView.cpp
+++ b/Engine/Source/Runtime/Engine/Private/SceneView.cpp
@@ -1689,6 +1689,12 @@ void FSceneView::OverridePostProcessSettings(const FPostProcessSettings& Src, fl
 		LERP_PP(LensFlareBokehSize);
 		LERP_PP(LensFlareThreshold);
 		LERP_PP(VignetteIntensity);
+		SET_PP(VignetteType);
+		LERP_PP(VignetteCenter);
+		LERP_PP(VignetteColor);
+		LERP_PP(VignetteSize);
+		LERP_PP(VignetteSoftness);
+		SET_PP(VignetteTexture);
 		LERP_PP(Sharpen);
 		LERP_PP(FilmGrainIntensity);
 		LERP_PP(FilmGrainIntensityShadows);
diff --git a/Engine/Source/Runtime/Renderer/Private/PostProcess/PostProcessTonemap.cpp b/Engine/Source/Runtime/Renderer/Private/PostProcess/PostProcessTonemap.cpp
--- a/Engine/Source/Runtime/Renderer/Private/PostProcess/PostProcessTonemap.cpp
+++ b/Engine/Source/Runtime/Renderer/Private/PostProcess/PostProcessTonemap.cpp
@@ -83,17 +83,19 @@ struct FOutputLuminance
 namespace TonemapperPermutation
 {
 // Shared permutation dimensions between deferred and mobile renderer.
-class FTonemapperGammaOnlyDim      : SHADER_PERMUTATION_BOOL("USE_GAMMA_ONLY");
-class FTonemapperLocalExposureDim  : SHADER_PERMUTATION_INT("LOCAL_EXPOSURE_MODE", 3);
-class FTonemapperSharpenDim        : SHADER_PERMUTATION_BOOL("USE_SHARPEN");
-class FTonemapperFilmGrainDim      : SHADER_PERMUTATION_BOOL("USE_FILM_GRAIN");
-class FTonemapperMsaaDim           : SHADER_PERMUTATION_BOOL("METAL_MSAA_HDR_DECODE");
-class FTonemapperAlphaChannelDim   : SHADER_PERMUTATION_BOOL("DIM_ALPHA_CHANNEL");
+class FTonemapperGammaOnlyDim      	: SHADER_PERMUTATION_BOOL("USE_GAMMA_ONLY");
+class FTonemapperLocalExposureDim  	: SHADER_PERMUTATION_INT("LOCAL_EXPOSURE_MODE", 3);
+class FTonemapperSharpenDim        	: SHADER_PERMUTATION_BOOL("USE_SHARPEN");
+class FTonemapperFilmGrainDim      	: SHADER_PERMUTATION_BOOL("USE_FILM_GRAIN");
+class FTonemapperVignetteTextureDim	: SHADER_PERMUTATION_BOOL("USE_VIGNETTE_TEXTURE");
+class FTonemapperMsaaDim           	: SHADER_PERMUTATION_BOOL("METAL_MSAA_HDR_DECODE");
+class FTonemapperAlphaChannelDim   	: SHADER_PERMUTATION_BOOL("DIM_ALPHA_CHANNEL");
 using FCommonDomain = TShaderPermutationDomain<
 	FTonemapperGammaOnlyDim,
 	FTonemapperLocalExposureDim,
 	FTonemapperSharpenDim,
 	FTonemapperFilmGrainDim,
+	FTonemapperVignetteTextureDim,
 	FTonemapperMsaaDim,
 	FTonemapperAlphaChannelDim>;
 
@@ -111,8 +113,10 @@ bool ShouldCompileCommonPermutation(const FGlobalShaderPermutationParameters& Pa
 		return !PermutationVector.Get<FTonemapperLocalExposureDim>() &&
 			!PermutationVector.Get<FTonemapperSharpenDim>() &&
 			!PermutationVector.Get<FTonemapperFilmGrainDim>() &&
+			!PermutationVector.Get<FTonemapperVignetteTextureDim>() &&
 			!PermutationVector.Get<FTonemapperMsaaDim>();
 	}
+
 	return true;
 }
 
@@ -144,6 +148,7 @@ FCommonDomain BuildCommonPermutationDomain(const FViewInfo& View, bool bGammaOnl
 	PermutationVector.Set<FTonemapperFilmGrainDim>(View.FinalPostProcessSettings.FilmGrainIntensity > 0.0f && View.FilmGrainTexture != nullptr);
 	PermutationVector.Set<FTonemapperSharpenDim>(GetSharpenSetting(Settings) > 0.0f);
 	PermutationVector.Set<FTonemapperMsaaDim>(bMetalMSAAHDRDecode);
+	PermutationVector.Set<FTonemapperVignetteTextureDim>(View.FinalPostProcessSettings.VignetteType == EVignetteType::Texture);
 
 	if (bLocalExposure)
 	{
@@ -178,15 +183,21 @@ FDesktopDomain RemapPermutation(FDesktopDomain PermutationVector, ERHIFeatureLev
 		return PermutationVector;
 	}
 
-	bool bFallbackToSlowest = CommonPermutationVector.Get<FTonemapperFilmGrainDim>();
+	bool bFallbackToSlowest = CommonPermutationVector.Get<FTonemapperFilmGrainDim>() || CommonPermutationVector.Get<FTonemapperVignetteTextureDim>();
 	if (bFallbackToSlowest)
 	{
 		CommonPermutationVector.Set<FTonemapperFilmGrainDim>(true);
+		CommonPermutationVector.Set<FTonemapperVignetteTextureDim>(true);
 		CommonPermutationVector.Set<FTonemapperSharpenDim>(true);
 
 		PermutationVector.Set<FTonemapperColorFringeDim>(true);
 	}
 
+	if (FeatureLevel <= ERHIFeatureLevel::SM5)
+	{
+		CommonPermutationVector.Set<FTonemapperVignetteTextureDim>(false);
+	}
+
 	if (FeatureLevel < ERHIFeatureLevel::SM5)
 	{
 		// Mobile doesn't support film grain.
@@ -337,15 +348,22 @@ BEGIN_SHADER_PARAMETER_STRUCT(FTonemapParameters, )
 	SHADER_PARAMETER(FVector4f, ColorScale0)
 	SHADER_PARAMETER(FVector4f, BloomDirtMaskTint)
 	SHADER_PARAMETER(FVector4f, ChromaticAberrationParams)
-	SHADER_PARAMETER(FVector4f, TonemapperParams)
 	SHADER_PARAMETER(FVector4f, LensPrincipalPointOffsetScale)
 	SHADER_PARAMETER(FVector4f, LensPrincipalPointOffsetScaleInverse)
+	SHADER_PARAMETER(float, SharpenDiv6)
 	SHADER_PARAMETER(float, LUTSize)
 	SHADER_PARAMETER(float, InvLUTSize)
 	SHADER_PARAMETER(float, LUTScale)
 	SHADER_PARAMETER(float, LUTOffset)
 	SHADER_PARAMETER(float, EditorNITLevel)
 	SHADER_PARAMETER(float, BackbufferQuantizationDithering)
+
+	SHADER_PARAMETER(FVector4f, VignetteParams0)
+	SHADER_PARAMETER(FVector4f, VignetteParams1)
+	SHADER_PARAMETER(float, 	VignetteParams2)
+	SHADER_PARAMETER(uint32,	VignetteType)
+	SHADER_PARAMETER_TEXTURE(Texture2D, VignetteTexture)
+	SHADER_PARAMETER_SAMPLER(SamplerState, VignetteTextureSampler)
 END_SHADER_PARAMETER_STRUCT()
 
 class FFilmGrainReduceCS : public FGlobalShader
@@ -833,9 +851,29 @@ FScreenPassTexture AddTonemapPass(FRDGBuilder& GraphBuilder, const FViewInfo& Vi
 	CommonParameters.ColorGradingLUTSampler = BilinearClampSampler;
 	CommonParameters.ColorScale0 = PostProcessSettings.SceneColorTint;
 	CommonParameters.ChromaticAberrationParams = ChromaticAberrationParams;
-	CommonParameters.TonemapperParams = FVector4f(PostProcessSettings.VignetteIntensity, SharpenDiv6, 0.0f, 0.0f);
+	CommonParameters.SharpenDiv6 = SharpenDiv6;
 	CommonParameters.EditorNITLevel = EditorNITLevel;
 
+	CommonParameters.VignetteType = (uint32)PostProcessSettings.VignetteType;
+	CommonParameters.VignetteParams0 = FVector4f(PostProcessSettings.VignetteIntensity, PostProcessSettings.VignetteColor.R, PostProcessSettings.VignetteColor.G, PostProcessSettings.VignetteColor.B);
+	CommonParameters.VignetteParams1 = FVector4f(
+		PostProcessSettings.VignetteCenter.X, 
+		PostProcessSettings.VignetteCenter.Y, 
+		1.0f / FMath::Max(0.001, PostProcessSettings.VignetteSize.X), 
+		1.0f / FMath::Max(0.001, PostProcessSettings.VignetteSize.Y));
+	CommonParameters.VignetteParams2 = PostProcessSettings.VignetteSoftness;
+	CommonParameters.VignetteTexture = GWhiteTexture->TextureRHI;
+	CommonParameters.VignetteTextureSampler = TStaticSamplerState<SF_Bilinear>::GetRHI();
+	if (PostProcessSettings.VignetteType == EVignetteType::Texture && PostProcessSettings.VignetteTexture)
+	{
+		FTextureResource* VignetteTextureResource = PostProcessSettings.VignetteTexture->GetResource();
+
+		if (VignetteTextureResource && VignetteTextureResource->TextureRHI)
+		{
+			CommonParameters.VignetteTexture = VignetteTextureResource->TextureRHI;
+		}
+	}
+
 	CommonParameters.FilmSlope = PostProcessSettings.FilmSlope;
 	CommonParameters.FilmToe = PostProcessSettings.FilmToe;
 	CommonParameters.FilmShoulder = PostProcessSettings.FilmShoulder;
