[Log] Adding more info to logs

#tests Verified that the log output included the asset name

#Synced-CL 49835867

diff --git a/Engine/Source/Runtime/Renderer/Private/PostProcess/PostProcessMaterial.cpp b/Engine/Source/Runtime/Renderer/Private/PostProcess/PostProcessMaterial.cpp
--- a/Engine/Source/Runtime/Renderer/Private/PostProcess/PostProcessMaterial.cpp
+++ b/Engine/Source/Runtime/Renderer/Private/PostProcess/PostProcessMaterial.cpp
@@ -84,12 +84,12 @@ static EMaterialCustomDepthPolicy GetMaterialCustomDepthPolicy(const FMaterialRe
 		// Custom stencil texture allocated and available.
 		if (GetCustomDepthMode() != ECustomDepthMode::EnabledWithStencil)
 		{
-			UE_LOG(LogRenderer, Warning, TEXT("PostProcessMaterial uses stencil test, but stencil not allocated. Set r.CustomDepth to 3 to allocate custom stencil."));
+			UE_LOG(LogRenderer, Warning, TEXT("PostProcessMaterial %s uses stencil test, but stencil not allocated. Set r.CustomDepth to 3 to allocate custom stencil."), *Material->GetAssetName());
 		}
 		else if (MaterialRenderProxy->GetBlendableLocation(Material) == BL_SceneColorAfterTonemapping)
 		{
 			// We can't support custom stencil after tonemapping due to target size differences
-			UE_LOG(LogRenderer, Warning, TEXT("PostProcessMaterial uses stencil test, but is set to blend After Tonemapping. This is not supported."));
+			UE_LOG(LogRenderer, Warning, TEXT("PostProcessMaterial %s uses stencil test, but is set to blend After Tonemapping. This is not supported."), *Material->GetAssetName());
 		}
 		else
 		{
