Mobile: Remove HQ SLW rendering shader permutation from GL platforms, as there are not enough samplers to support it

#Synced-CL 49917250

#Synced-CL 49917250

diff --git a/Engine/Source/Runtime/RenderCore/Private/RenderUtils.cpp b/Engine/Source/Runtime/RenderCore/Private/RenderUtils.cpp
--- a/Engine/Source/Runtime/RenderCore/Private/RenderUtils.cpp
+++ b/Engine/Source/Runtime/RenderCore/Private/RenderUtils.cpp
@@ -1563,8 +1563,10 @@ bool IsUsingSingleLayerWaterFallback(FStaticShaderPlatform Platform)
 	}
 
 	// Only multi-pass deferred supports SLW rendering
+	// Disabled for GL, as there is not enough samplers
 	if (IsMobileDeferredShadingEnabled(Platform) 
-		&& !FReadOnlyCVARCache::MobileAllowFramebufferFetch(Platform))
+		&& !FReadOnlyCVARCache::MobileAllowFramebufferFetch(Platform)
+		&& !IsOpenGLPlatform(Platform))
 	{
 		return (!CVarWaterSingleLayerMobileFallback || CVarWaterSingleLayerMobileFallback->GetInt() != 0);
 	}
diff --git a/Engine/Source/Runtime/Renderer/Private/LightMapRendering.cpp b/Engine/Source/Runtime/Renderer/Private/LightMapRendering.cpp
--- a/Engine/Source/Runtime/Renderer/Private/LightMapRendering.cpp
+++ b/Engine/Source/Runtime/Renderer/Private/LightMapRendering.cpp
@@ -263,6 +263,7 @@ FMobileSingleLayerWaterPolicy::FMobileSingleLayerWaterPolicy()
 bool FMobileSingleLayerWaterPolicy::ShouldCompilePermutation(const FMeshMaterialShaderPermutationParameters& Parameters)
 {
 	return IsMobilePlatform(Parameters.Platform) 
+		&& !IsOpenGLPlatform(Parameters.Platform) // not enough samplers
 		&& IsMobileDeferredShadingEnabled(Parameters.Platform)
 		&& !FReadOnlyCVARCache::MobileAllowFramebufferFetch(Parameters.Platform)
 		&& Parameters.MaterialParameters.ShadingModels.HasShadingModel(MSM_SingleLayerWater);
