Add default hair uniform buffer creation.

#Synced-CL 49918411

diff --git a/Engine/Source/Runtime/Renderer/Private/DeferredShadingRenderer.cpp b/Engine/Source/Runtime/Renderer/Private/DeferredShadingRenderer.cpp
--- a/Engine/Source/Runtime/Renderer/Private/DeferredShadingRenderer.cpp
+++ b/Engine/Source/Runtime/Renderer/Private/DeferredShadingRenderer.cpp
@@ -2279,7 +2279,7 @@ void FDeferredShadingSceneRenderer::Render(FRDGBuilder& GraphBuilder, const FSce
 		{
 			for (FViewInfo& View : Views)
 			{
-				View.HairStrandsViewData.UniformBuffer = HairStrands::CreateDefaultHairStrandsViewUniformBuffer(GraphBuilder, View);
+				View.HairStrandsViewData.UniformBuffer = HairStrands::CreateDefaultHairStrandsViewUniformBuffer(GraphBuilder);
 			}
 		}
 	}
diff --git a/Engine/Source/Runtime/Renderer/Private/HairStrands/HairStrandsData.h b/Engine/Source/Runtime/Renderer/Private/HairStrands/HairStrandsData.h
--- a/Engine/Source/Runtime/Renderer/Private/HairStrands/HairStrandsData.h
+++ b/Engine/Source/Runtime/Renderer/Private/HairStrands/HairStrandsData.h
@@ -474,8 +474,10 @@ struct FHairStrandsViewStateData
 
 namespace HairStrands
 {
-	TRDGUniformBufferRef<FHairStrandsViewUniformParameters> CreateDefaultHairStrandsViewUniformBuffer(FRDGBuilder& GraphBuilder, FViewInfo& View);
+	TRDGUniformBufferRef<FHairStrandsViewUniformParameters> CreateDefaultHairStrandsViewUniformBuffer(FRDGBuilder& GraphBuilder);
 	TRDGUniformBufferRef<FHairStrandsViewUniformParameters> BindHairStrandsViewUniformParameters(const FViewInfo& View);
+
+	TRDGUniformBufferRef<FVirtualVoxelParameters> CreateDefaultHairStrandsVoxelUniformBuffer(FRDGBuilder& GraphBuilder);
 	TRDGUniformBufferRef<FVirtualVoxelParameters> BindHairStrandsVoxelUniformParameters(const FViewInfo& View);
 
 	bool HasViewHairStrandsData(const FViewInfo& View);
diff --git a/Engine/Source/Runtime/Renderer/Private/HairStrands/HairStrandsRendering.cpp b/Engine/Source/Runtime/Renderer/Private/HairStrands/HairStrandsRendering.cpp
--- a/Engine/Source/Runtime/Renderer/Private/HairStrands/HairStrandsRendering.cpp
+++ b/Engine/Source/Runtime/Renderer/Private/HairStrands/HairStrandsRendering.cpp
@@ -301,7 +301,7 @@ bool FHairStrandsViewStateData::ReadPositionsChanged()
 namespace HairStrands
 {
 
-TRDGUniformBufferRef<FHairStrandsViewUniformParameters> CreateDefaultHairStrandsViewUniformBuffer(FRDGBuilder& GraphBuilder, FViewInfo& View)
+TRDGUniformBufferRef<FHairStrandsViewUniformParameters> CreateDefaultHairStrandsViewUniformBuffer(FRDGBuilder& GraphBuilder)
 {
 	return InternalCreateHairStrandsViewUniformBuffer(GraphBuilder, nullptr);
 }
@@ -311,6 +311,51 @@ TRDGUniformBufferRef<FHairStrandsViewUniformParameters> BindHairStrandsViewUnifo
 	return View.HairStrandsViewData.UniformBuffer;
 }
 
+TRDGUniformBufferRef<FVirtualVoxelParameters> CreateDefaultHairStrandsVoxelUniformBuffer(FRDGBuilder& GraphBuilder)
+{
+	FRDGBufferSRVRef DummyBuffer = GraphBuilder.CreateSRV(GSystemTextures.GetDefaultBuffer(GraphBuilder, 4u, 0u), PF_R32_UINT);
+	FRDGBufferSRVRef DummyStructuredBuffer = GraphBuilder.CreateSRV(GSystemTextures.GetDefaultStructuredBuffer(GraphBuilder, 4u, 0u), PF_R32_UINT);
+
+	FVirtualVoxelParameters* Out = GraphBuilder.AllocParameters<FVirtualVoxelParameters>();
+	Out->Common.PageCountResolution = FIntVector::ZeroValue;
+	Out->Common.CPUMinVoxelWorldSize = 0.f;
+	Out->Common.PageTextureResolution = FIntVector::ZeroValue;
+	Out->Common.PageCount = 0;
+	Out->Common.PageResolution = 0;
+	Out->Common.PageResolutionLog2 = 0;
+	Out->Common.PageIndexCount = 0;
+	Out->Common.IndirectDispatchGroupSize = 0;
+	Out->Common.NodeDescCount = 0;
+	Out->Common.JitterMode = 0;
+	Out->Common.DensityScale = 0;
+	Out->Common.DensityScale_AO = 0;
+	Out->Common.DensityScale_Shadow = 0;
+	Out->Common.DensityScale_Transmittance = 0;
+	Out->Common.DensityScale_Environment = 0;
+	Out->Common.DensityScale_Raytracing = 0;
+	Out->Common.DepthBiasScale_Shadow = 0;
+	Out->Common.DepthBiasScale_Transmittance = 0;
+	Out->Common.DepthBiasScale_Environment = 0;
+	Out->Common.SteppingScale_Shadow = 0;
+	Out->Common.SteppingScale_Transmittance = 0;
+	Out->Common.SteppingScale_Environment = 0;
+	Out->Common.SteppingScale_Raytracing = 0;
+	Out->Common.HairCoveragePixelRadiusAtDepth1 = 0;
+	Out->Common.Raytracing_ShadowOcclusionThreshold = 0;
+	Out->Common.Raytracing_SkyOcclusionThreshold = 0;
+	Out->Common.TranslatedWorldOffset = FVector3f::ZeroVector;
+	Out->Common.TranslatedWorldOffsetStereoCorrection = FVector3f::ZeroVector;
+	Out->Common.AllocationFeedbackEnable = 0;
+	Out->Common.AllocatedPageCountBuffer = DummyBuffer;
+	Out->Common.PageIndexBuffer = DummyBuffer;
+	Out->Common.PageIndexCoordBuffer = DummyBuffer;
+	Out->Common.NodeDescBuffer = DummyStructuredBuffer;
+	Out->Common.CurrGPUMinVoxelSize = DummyBuffer;
+	Out->Common.NextGPUMinVoxelSize = DummyBuffer;
+	Out->PageTexture = GSystemTextures.GetVolumetricBlackDummy(GraphBuilder);
+	return GraphBuilder.CreateUniformBuffer(Out);
+}
+
 TRDGUniformBufferRef<FVirtualVoxelParameters> BindHairStrandsVoxelUniformParameters(const FViewInfo& View)
 {
 	// Voxel uniform buffer exist only if the view has hair strands data
