Disable scene light scattering for hair as it can create artifact at silouhette.

#Synced-CL 49918411

diff --git a/Engine/Source/Runtime/Renderer/Private/HairStrands/HairStrandsEnvironment.cpp b/Engine/Source/Runtime/Renderer/Private/HairStrands/HairStrandsEnvironment.cpp
--- a/Engine/Source/Runtime/Renderer/Private/HairStrands/HairStrandsEnvironment.cpp
+++ b/Engine/Source/Runtime/Renderer/Private/HairStrands/HairStrandsEnvironment.cpp
@@ -29,8 +29,8 @@
 
 ///////////////////////////////////////////////////////////////////////////////////////////////////
 
-static int32 GHairScatterSceneLighting = 1;
-static FAutoConsoleVariableRef CVarHairScatterSceneLighting(TEXT("r.HairStrands.ScatterSceneLighting"), GHairScatterSceneLighting, TEXT("Enable scene color lighting scattering into hair (valid for short hair only)."));
+static int32 GHairScatterSceneLighting = 0;
+static FAutoConsoleVariableRef CVarHairScatterSceneLighting(TEXT("r.HairStrands.ScatterSceneLighting"), GHairScatterSceneLighting, TEXT("Enable scene color lighting scattering into hair (valid for short hair only), disabled by default."));
 
 static int32 GHairSkylightingEnable = 1;
 static FAutoConsoleVariableRef CVarHairSkylightingEnable(TEXT("r.HairStrands.SkyLighting"), GHairSkylightingEnable, TEXT("Enable sky lighting on hair."));
