Support deferred decals reading the normal buffer on the mobile renderer.

-Use SamplingDeferredNormal at cook time and runtime to determine whether the platform can sample the normal G-Buffer (supported on Mobile Deferred MultiPass).
-Add a new structure (FDecalViewInfo) to allow calling RenderDecals with different contexts (e.g., different FDecalRenderState and EDecalRenderTargetMode).
-Use the decal visibility task to build the relevant deferred decal list. As on the desktop renderer, this allows the decal visibility task to be launched early and the appropriate list to be requested when rendering decals.
-Add a new function (GetRenderTargetMode) used at cook time to correctly build the material render target mode (the render target state differs when normal G-Buffer sampling is allowed).

#tests editor preview modes + QAGame Decal map

#Synced-CL 49917480

diff --git a/Engine/Source/Runtime/Renderer/Private/CompositionLighting/PostProcessDeferredDecals.cpp b/Engine/Source/Runtime/Renderer/Private/CompositionLighting/PostProcessDeferredDecals.cpp
--- a/Engine/Source/Runtime/Renderer/Private/CompositionLighting/PostProcessDeferredDecals.cpp
+++ b/Engine/Source/Runtime/Renderer/Private/CompositionLighting/PostProcessDeferredDecals.cpp
@@ -546,7 +546,7 @@ void CollectDeferredDecalPassPSOInitializers(
 {
 	const EShaderPlatform ShaderPlatform = GetFeatureLevelShaderPlatform(FeatureLevel);
 	const FDecalBlendDesc DecalBlendDesc = DecalRendering::ComputeDecalBlendDesc(ShaderPlatform, Material);
-	EDecalRenderTargetMode DecalRenderTargetMode = DecalRendering::GetRenderTargetMode(DecalBlendDesc, DecalRenderStage);
+	EDecalRenderTargetMode DecalRenderTargetMode = DecalRendering::GetRenderTargetMode(ShaderPlatform, DecalBlendDesc, DecalRenderStage);
 
 
 	TShaderRef<FShader> VertexShader, PixelShader;
@@ -789,11 +789,11 @@ void AddDeferredDecalPass(
 			uint32 NumDecals = 0;
 			uint32 SortedDecalIndex = 1;
 			uint32 LastSortedDecalIndex = 0;
-			EDecalRenderTargetMode LastRenderTargetMode = DecalRendering::GetRenderTargetMode(SortedDecals[0]->BlendDesc, DecalRenderStage);
+			EDecalRenderTargetMode LastRenderTargetMode = DecalRendering::GetRenderTargetMode(ShaderPlatform, SortedDecals[0]->BlendDesc, DecalRenderStage);
 
 			for (; SortedDecalIndex < SortedDecalCount; ++SortedDecalIndex, ++NumDecals)
 			{
-				const EDecalRenderTargetMode RenderTargetMode = DecalRendering::GetRenderTargetMode(SortedDecals[SortedDecalIndex]->BlendDesc, DecalRenderStage);
+				const EDecalRenderTargetMode RenderTargetMode = DecalRendering::GetRenderTargetMode(ShaderPlatform, SortedDecals[SortedDecalIndex]->BlendDesc, DecalRenderStage);
 
 				if (LastRenderTargetMode != RenderTargetMode || NumDecals > MaxNumDecals)
 				{
diff --git a/Engine/Source/Runtime/Renderer/Private/CompositionLighting/PostProcessMeshDecals.cpp b/Engine/Source/Runtime/Renderer/Private/CompositionLighting/PostProcessMeshDecals.cpp
--- a/Engine/Source/Runtime/Renderer/Private/CompositionLighting/PostProcessMeshDecals.cpp
+++ b/Engine/Source/Runtime/Renderer/Private/CompositionLighting/PostProcessMeshDecals.cpp
@@ -262,7 +262,7 @@ bool FMeshDecalMeshProcessor::TryAddMeshBatch(
 
 			const bool bShouldRender =
 				DecalRendering::IsCompatibleWithRenderStage(DecalBlendDesc, PassDecalStage) &&
-				DecalRendering::GetRenderTargetMode(DecalBlendDesc, PassDecalStage) == RenderTargetMode;
+				DecalRendering::GetRenderTargetMode(ShaderPlatform, DecalBlendDesc, PassDecalStage) == RenderTargetMode;
 
 			if (bShouldRender)
 			{
@@ -379,7 +379,8 @@ void FMeshDecalMeshProcessor::CollectDeferredDecalMeshPSOInitializers(
 	EDecalRenderStage DecalRenderStage,
 	TArray<FPSOPrecacheData>& PSOInitializers)
 {
-	EDecalRenderTargetMode LocalRenderTargetMode = DecalRendering::GetRenderTargetMode(DecalBlendDesc, DecalRenderStage);
+	const EShaderPlatform ShaderPlatform = GetFeatureLevelShaderPlatform(FeatureLevel);
+	EDecalRenderTargetMode LocalRenderTargetMode = DecalRendering::GetRenderTargetMode(ShaderPlatform, DecalBlendDesc, DecalRenderStage);
 	PassDrawRenderState.SetBlendState(DecalRendering::GetDecalBlendState(DecalBlendDesc, DecalRenderStage, LocalRenderTargetMode));
 	PassDrawRenderState.SetDepthStencilState(GetMeshDecalDepthStencilState(LocalRenderTargetMode));
 
@@ -470,7 +471,7 @@ void FMeshDecalMeshProcessor::CollectPSOInitializers(
 			continue;
 		}
 
-		EDecalRenderTargetMode LocalRenderTargetMode = DecalRendering::GetRenderTargetMode(DecalBlendDesc, LocalDecalRenderStage);
+		EDecalRenderTargetMode LocalRenderTargetMode = DecalRendering::GetRenderTargetMode(ShaderPlatform, DecalBlendDesc, LocalDecalRenderStage);
 		if (LocalRenderTargetMode == EDecalRenderTargetMode::DBuffer && !bPlatformUsesDBuffers)
 		{
 			continue;
@@ -520,6 +521,11 @@ FMeshPassProcessor* CreateMeshDecalSceneColorAndGBufferMeshProcessor_Mobile(ERHI
 	return new FMeshDecalMeshProcessor(Scene, FeatureLevel, InViewIfDynamicMeshCommand, EDecalRenderTargetMode::SceneColorAndGBuffer, EShadingPath::Mobile, InDrawListContext);
 }
 
+FMeshPassProcessor* CreateMeshDecalSceneColorAndGBufferNoNormalMeshProcessor_Mobile(ERHIFeatureLevel::Type FeatureLevel, const FScene* Scene, const FSceneView* InViewIfDynamicMeshCommand, FMeshPassDrawListContext* InDrawListContext)
+{
+	return new FMeshDecalMeshProcessor(Scene, FeatureLevel, InViewIfDynamicMeshCommand, EDecalRenderTargetMode::SceneColorAndGBufferNoNormal, EShadingPath::Mobile, InDrawListContext);
+}
+
 FMeshPassProcessor* CreateMeshDecalSceneColorMeshProcessor_Mobile(ERHIFeatureLevel::Type FeatureLevel, const FScene* Scene, const FSceneView* InViewIfDynamicMeshCommand, FMeshPassDrawListContext* InDrawListContext)
 {
 	return new FMeshDecalMeshProcessor(Scene, FeatureLevel, InViewIfDynamicMeshCommand, EDecalRenderTargetMode::SceneColor, EShadingPath::Mobile, InDrawListContext);
@@ -531,6 +537,7 @@ FMeshPassProcessor* CreateMeshDecalDBufferMeshProcessor_Mobile(ERHIFeatureLevel:
 }
 
 REGISTER_MESHPASSPROCESSOR_AND_PSOCOLLECTOR(Mobile_MeshDecal_SceneColorAndGBuffer, CreateMeshDecalSceneColorAndGBufferMeshProcessor_Mobile, EShadingPath::Mobile, EMeshPass::MeshDecal_SceneColorAndGBuffer, EMeshPassFlags::CachedMeshCommands | EMeshPassFlags::MainView);
+REGISTER_MESHPASSPROCESSOR_AND_PSOCOLLECTOR(Mobile_MeshDecal_SceneColorAndGBufferNoNormal, CreateMeshDecalSceneColorAndGBufferNoNormalMeshProcessor_Mobile, EShadingPath::Mobile, EMeshPass::MeshDecal_SceneColorAndGBufferNoNormal, EMeshPassFlags::CachedMeshCommands | EMeshPassFlags::MainView);
 REGISTER_MESHPASSPROCESSOR_AND_PSOCOLLECTOR(Mobile_MeshDecal_SceneColor, CreateMeshDecalSceneColorMeshProcessor_Mobile, EShadingPath::Mobile, EMeshPass::MeshDecal_SceneColor, EMeshPassFlags::CachedMeshCommands | EMeshPassFlags::MainView);
 REGISTER_MESHPASSPROCESSOR_AND_PSOCOLLECTOR(Mobile_MeshDecal_DBuffer, CreateMeshDecalDBufferMeshProcessor_Mobile, EShadingPath::Mobile, EMeshPass::MeshDecal_DBuffer, EMeshPassFlags::CachedMeshCommands | EMeshPassFlags::MainView);
 
@@ -599,14 +606,12 @@ bool HasAnyDrawCommandDecalCount(
 		return HasAnyDraw(View.ParallelMeshDrawCommandPasses[EMeshPass::MeshDecal_DBuffer]);
 
 	case EDecalRenderStage::BeforeLighting:
+	case EDecalRenderStage::MobileBeforeLighting:
 		return HasAnyDraw(View.ParallelMeshDrawCommandPasses[EMeshPass::MeshDecal_SceneColorAndGBuffer]) || HasAnyDraw(View.ParallelMeshDrawCommandPasses[EMeshPass::MeshDecal_SceneColorAndGBufferNoNormal]);
 
 	case EDecalRenderStage::Mobile:
 		return HasAnyDraw(View.ParallelMeshDrawCommandPasses[EMeshPass::MeshDecal_SceneColor]);
 
-	case EDecalRenderStage::MobileBeforeLighting:
-		return HasAnyDraw(View.ParallelMeshDrawCommandPasses[EMeshPass::MeshDecal_SceneColorAndGBuffer]);
-
 	case EDecalRenderStage::Emissive:
 		return HasAnyDraw(View.ParallelMeshDrawCommandPasses[EMeshPass::MeshDecal_SceneColor]);
 
@@ -633,6 +638,7 @@ void RenderMeshDecals(
 		break;
 
 	case EDecalRenderStage::BeforeLighting:
+	case EDecalRenderStage::MobileBeforeLighting:
 		DrawDecalMeshCommands(GraphBuilder, Scene, View, DecalPassTextures, InstanceCullingManager, DecalRenderStage, EDecalRenderTargetMode::SceneColorAndGBuffer);
 		DrawDecalMeshCommands(GraphBuilder, Scene, View, DecalPassTextures, InstanceCullingManager, DecalRenderStage, EDecalRenderTargetMode::SceneColorAndGBufferNoNormal);
 		break;
@@ -641,10 +647,6 @@ void RenderMeshDecals(
 		DrawDecalMeshCommands(GraphBuilder, Scene, View, DecalPassTextures, InstanceCullingManager, DecalRenderStage, EDecalRenderTargetMode::SceneColor);
 		break;
 
-	case EDecalRenderStage::MobileBeforeLighting:
-		DrawDecalMeshCommands(GraphBuilder, Scene, View, DecalPassTextures, InstanceCullingManager, DecalRenderStage, EDecalRenderTargetMode::SceneColorAndGBuffer);
-		break;
-
 	case EDecalRenderStage::Emissive:
 		DrawDecalMeshCommands(GraphBuilder, Scene, View, DecalPassTextures, InstanceCullingManager, DecalRenderStage, EDecalRenderTargetMode::SceneColor);
 		break;
diff --git a/Engine/Source/Runtime/Renderer/Private/DecalRenderingCommon.cpp b/Engine/Source/Runtime/Renderer/Private/DecalRenderingCommon.cpp
--- a/Engine/Source/Runtime/Renderer/Private/DecalRenderingCommon.cpp
+++ b/Engine/Source/Runtime/Renderer/Private/DecalRenderingCommon.cpp
@@ -20,6 +20,12 @@ namespace DecalRendering
 	inline bool IsAlphaHoldoutBlendMode(const FDecalBlendDesc& In)		{ return IsAlphaHoldoutBlendMode((EBlendMode)In.BlendMode); }
 	inline bool IsModulateBlendMode(const FDecalBlendDesc& In)			{ return IsModulateBlendMode((EBlendMode)In.BlendMode); }
 	inline bool IsAlphaCompositeBlendMode(const FDecalBlendDesc& In)	{ return IsAlphaCompositeBlendMode((EBlendMode)In.BlendMode); }
+	
+	/** Used at cook and runtime time to determine whether decals support normal G-Buffer sampling (requires Deferred MultiPass mode). */
+	static bool MobileSupportDecalNormalSampling(const FStaticShaderPlatform Platform)
+	{
+		return IsMobileDeferredShadingEnabled(Platform) && !MobileAllowFramebufferFetch(Platform);
+	}
 
 	/** Finalize the initialization of FDecalBlendDesc after BlendMode and bWrite flags have all been set. */
 	void FinalizeBlendDesc(EShaderPlatform Platform, FDecalBlendDesc& Desc)
@@ -201,7 +207,7 @@ namespace DecalRendering
 			EDecalRenderStage DecalRenderStage = (EDecalRenderStage)DecalRenderStageNum;
 			if (IsCompatibleWithRenderStage(DecalBlendDesc, DecalRenderStage))
 			{
-				EDecalRenderTargetMode DecalRenderTargetMode = GetRenderTargetMode(DecalBlendDesc, DecalRenderStage);
+				EDecalRenderTargetMode DecalRenderTargetMode = GetRenderTargetMode(ShaderPlatform, DecalBlendDesc, DecalRenderStage);
 				Mask |= (1 << (uint32)DecalRenderTargetMode);
 			}
 		}
@@ -244,9 +250,8 @@ namespace DecalRendering
 			case EDecalRenderTargetMode::DBuffer:
 				return EDecalRenderStage::BeforeBasePass;
 			case EDecalRenderTargetMode::SceneColorAndGBuffer:
-				return (ShadingPath == EShadingPath::Deferred) ? EDecalRenderStage::BeforeLighting : EDecalRenderStage::MobileBeforeLighting;
 			case EDecalRenderTargetMode::SceneColorAndGBufferNoNormal:
-				return EDecalRenderStage::BeforeLighting;
+				return (ShadingPath == EShadingPath::Deferred) ? EDecalRenderStage::BeforeLighting : EDecalRenderStage::MobileBeforeLighting;
 			case EDecalRenderTargetMode::SceneColor:
 				return (ShadingPath == EShadingPath::Deferred) ? EDecalRenderStage::Emissive : EDecalRenderStage::Mobile;
 			case EDecalRenderTargetMode::AmbientOcclusion:
@@ -256,27 +261,48 @@ namespace DecalRendering
 		return EDecalRenderStage::None;
 	}
 
-	EDecalRenderTargetMode GetRenderTargetMode(FDecalBlendDesc DecalBlendDesc, EDecalRenderStage DecalRenderStage)
+	EDecalRenderTargetMode GetRenderTargetMode(EShaderPlatform Platform, FDecalBlendDesc DecalBlendDesc, EDecalRenderStage DecalRenderStage)
 	{
-		switch(DecalRenderStage)
-		{
-			case EDecalRenderStage::BeforeBasePass:
-				return EDecalRenderTargetMode::DBuffer;
-			case EDecalRenderStage::BeforeLighting:
-				return DecalBlendDesc.bWriteNormal ? EDecalRenderTargetMode::SceneColorAndGBuffer : EDecalRenderTargetMode::SceneColorAndGBufferNoNormal;
-			case EDecalRenderStage::Mobile:
-				return EDecalRenderTargetMode::SceneColor;
-			case EDecalRenderStage::MobileBeforeLighting:
-				return EDecalRenderTargetMode::SceneColorAndGBuffer;
-			case EDecalRenderStage::Emissive:
-				return EDecalRenderTargetMode::SceneColor;
-			case EDecalRenderStage::AmbientOcclusion:
-				return EDecalRenderTargetMode::AmbientOcclusion;
+		switch (DecalRenderStage)
+		{
+		case EDecalRenderStage::BeforeBasePass:
+			return EDecalRenderTargetMode::DBuffer;
+		case EDecalRenderStage::BeforeLighting:
+			return DecalBlendDesc.bWriteNormal ? EDecalRenderTargetMode::SceneColorAndGBuffer : EDecalRenderTargetMode::SceneColorAndGBufferNoNormal;
+		case EDecalRenderStage::Mobile:
+			return EDecalRenderTargetMode::SceneColor;
+		case EDecalRenderStage::MobileBeforeLighting:
+			return (MobileSupportDecalNormalSampling(Platform) && !DecalBlendDesc.bWriteNormal) ? EDecalRenderTargetMode::SceneColorAndGBufferNoNormal : EDecalRenderTargetMode::SceneColorAndGBuffer;
+		case EDecalRenderStage::Emissive:
+			return EDecalRenderTargetMode::SceneColor;
+		case EDecalRenderStage::AmbientOcclusion:
+			return EDecalRenderTargetMode::AmbientOcclusion;
 		}
 
 		return EDecalRenderTargetMode::None;
 	}
 
+	const TCHAR* GetRenderTargetModeName(EDecalRenderTargetMode DecalRenderTargetMode)
+	{
+		switch (DecalRenderTargetMode)
+		{
+		case EDecalRenderTargetMode::DBuffer:
+			return TEXT("DBuffer");
+		case EDecalRenderTargetMode::SceneColorAndGBuffer:
+			return TEXT("SceneColorAndGBuffer");
+		case EDecalRenderTargetMode::SceneColorAndGBufferNoNormal:
+			return TEXT("SceneColorAndGBufferNoNormal");
+		case EDecalRenderTargetMode::SceneColor:
+			return TEXT("SceneColor");
+		case EDecalRenderTargetMode::AmbientOcclusion:
+			return TEXT("AmbientOcclusion");
+		default:
+			break;
+		}
+
+		return TEXT("None");
+	}
+
 	uint32 GetRenderTargetCount(FDecalBlendDesc DecalBlendDesc, EDecalRenderTargetMode RenderTargetMode)
 	{
 		switch (RenderTargetMode)
@@ -962,7 +988,7 @@ namespace DecalRendering
 		DecalRenderStage = DecalRenderStage != EDecalRenderStage::None ? DecalRenderStage : GetBaseRenderStage(DecalBlendDesc);
 		check(DecalRenderStage != EDecalRenderStage::None);
 
-		const EDecalRenderTargetMode RenderTargetMode = GetRenderTargetMode(DecalBlendDesc, DecalRenderStage);
+		const EDecalRenderTargetMode RenderTargetMode = GetRenderTargetMode(Platform, DecalBlendDesc, DecalRenderStage);
 		check(RenderTargetMode != EDecalRenderTargetMode::None);
 
 		const uint32 RenderTargetCount = GetRenderTargetCount(DecalBlendDesc, RenderTargetMode);
diff --git a/Engine/Source/Runtime/Renderer/Private/DecalRenderingCommon.h b/Engine/Source/Runtime/Renderer/Private/DecalRenderingCommon.h
--- a/Engine/Source/Runtime/Renderer/Private/DecalRenderingCommon.h
+++ b/Engine/Source/Runtime/Renderer/Private/DecalRenderingCommon.h
@@ -99,8 +99,12 @@ namespace DecalRendering
 	/** Returns the main render stage for a decal (does not include the emissive, and AO stages). Can return EDecalRenderStage::None if there is no valid main render stage. */
 	EDecalRenderStage GetBaseRenderStage(FDecalBlendDesc DecalBlendDesc);
 
-	/** Get the render target mode that a decal uses for a given stage. Can return EDecalRenderTargetMode::None if there is no valid render target mode. */
-	EDecalRenderTargetMode GetRenderTargetMode(FDecalBlendDesc DecalBlendDesc, EDecalRenderStage DecalRenderStage);
+	/** Get the render target mode that a decal uses for a given stage at runtime and at cook time. Can return EDecalRenderTargetMode::None if there is no valid render target mode. */
+	EDecalRenderTargetMode GetRenderTargetMode(EShaderPlatform Platform, FDecalBlendDesc DecalBlendDesc, EDecalRenderStage DecalRenderStage);
+
+	/** Get the render target mode string */
+	const TCHAR* GetRenderTargetModeName(EDecalRenderTargetMode DecalRenderTargetMode);
+
 
 	/** Setup a mask for each render target mode where the material should render to */
 	uint8 GetDecalRenderTargetModeMask(const FMaterial& Material, ERHIFeatureLevel::Type FeatureLevel);
diff --git a/Engine/Source/Runtime/Renderer/Private/DecalRenderingShared.h b/Engine/Source/Runtime/Renderer/Private/DecalRenderingShared.h
--- a/Engine/Source/Runtime/Renderer/Private/DecalRenderingShared.h
+++ b/Engine/Source/Runtime/Renderer/Private/DecalRenderingShared.h
@@ -141,6 +141,18 @@ private:
 	RDG_FRIEND_ALLOCATOR_FRIEND(FDecalVisibilityTaskData);
 };
 
+struct FDecalViewInfo
+{
+	FDecalViewInfo(int32 InViewIndex, EDecalRenderStage InRenderStage, EDecalRenderTargetMode InRenderTargetMode) :
+		ViewIndex(InViewIndex),
+		RenderStage(InRenderStage),
+		RenderTargetMode(InRenderTargetMode)
+	{
+	}
+	int32 ViewIndex;
+	EDecalRenderStage RenderStage;
+	EDecalRenderTargetMode RenderTargetMode;
+};
 /**
  * Shared deferred decal functionality.
  */
diff --git a/Engine/Source/Runtime/Renderer/Private/MobileDecalRendering.cpp b/Engine/Source/Runtime/Renderer/Private/MobileDecalRendering.cpp
--- a/Engine/Source/Runtime/Renderer/Private/MobileDecalRendering.cpp
+++ b/Engine/Source/Runtime/Renderer/Private/MobileDecalRendering.cpp
@@ -21,7 +21,8 @@
 #include "CompositionLighting/PostProcessDeferredDecals.h"
 #include "DBufferTextures.h"
 
-extern void RenderDeferredDecalsMobile(FRHICommandList& RHICmdList, const FScene& Scene, const FViewInfo& View, EDecalRenderStage DecalRenderStage, EDecalRenderTargetMode RenderTargetMode);
+
+
 
 static bool DoesPlatformSupportDecals(EShaderPlatform ShaderPlatform)
 {
@@ -44,7 +45,7 @@ static bool DoesPlatformSupportDecals(EShaderPlatform ShaderPlatform)
 	return true;
 }
 
-void FMobileSceneRenderer::RenderDecals(FRHICommandList& RHICmdList, FViewInfo& View, const FInstanceCullingDrawParams* InstanceCullingDrawParams)
+void FMobileSceneRenderer::RenderDecals(FRHICommandList& RHICmdList, FViewInfo& View, const FDecalViewInfo& DecalViewInfo, const FInstanceCullingDrawParams* InstanceCullingDrawParams)
 {
 	if (!DoesPlatformSupportDecals(View.GetShaderPlatform()) || !ViewFamily.EngineShowFlags.Decals || View.bIsPlanarReflection)
 	{
@@ -54,12 +55,8 @@ void FMobileSceneRenderer::RenderDecals(FRHICommandList& RHICmdList, FViewInfo&
 	CSV_SCOPED_TIMING_STAT_EXCLUSIVE(RenderDecals);
 	SCOPE_CYCLE_COUNTER(STAT_DecalsDrawTime);
 
-	const bool bIsMobileDeferred = IsMobileDeferredShadingEnabled(View.GetShaderPlatform());
-	const EDecalRenderStage DecalRenderStage = bRequiresDBufferDecals ? EDecalRenderStage::Emissive : bIsMobileDeferred ? EDecalRenderStage::MobileBeforeLighting : EDecalRenderStage::Mobile;
-	const EDecalRenderTargetMode RenderTargetMode = bIsMobileDeferred ? EDecalRenderTargetMode::SceneColorAndGBuffer : EDecalRenderTargetMode::SceneColor;
-
 	// Mesh decals
-	EMeshPass::Type DecalMeshPassType = DecalRendering::GetMeshPassType(RenderTargetMode);
+	EMeshPass::Type DecalMeshPassType = DecalRendering::GetMeshPassType(DecalViewInfo.RenderTargetMode);
 	if (HasAnyDraw(View.ParallelMeshDrawCommandPasses[DecalMeshPassType]))
 	{
 		RHICmdList.SetViewport(View.ViewRect.Min.X, View.ViewRect.Min.Y, 0, View.ViewRect.Max.X, View.ViewRect.Max.Y, 1);
@@ -67,26 +64,27 @@ void FMobileSceneRenderer::RenderDecals(FRHICommandList& RHICmdList, FViewInfo&
 	}
 
 	// Deferred decals
-	if (Scene->Decals.Num() > 0)
+	if (DecalVisibilityTaskData != nullptr)
 	{
 		SCOPED_DRAW_EVENT(RHICmdList, Decals);
-		RenderDeferredDecalsMobile(RHICmdList, *Scene, View, DecalRenderStage, RenderTargetMode);
+		RenderDeferredDecalsMobile(RHICmdList, View, DecalViewInfo);
 	}
 }
 
-void RenderDeferredDecalsMobile(FRHICommandList& RHICmdList, const FScene& Scene, const FViewInfo& View, EDecalRenderStage DecalRenderStage, EDecalRenderTargetMode RenderTargetMode)
+void FMobileSceneRenderer::RenderDeferredDecalsMobile(FRHICommandList& RHICmdList, const FViewInfo& View, const FDecalViewInfo& DecalViewInfo)
 {
-	FVisibleDecalList VisibleDecals;
+	TConstArrayView<const FVisibleDecal*> RelevantDecals = DecalVisibilityTaskData->FinishRelevantDecals(DecalViewInfo.ViewIndex, DecalViewInfo.RenderStage);
 	FRelevantDecalList SortedDecals;
 
-	if (!Scene.Decals.IsEmpty())
+	for (const FVisibleDecal* VisibleDecal : RelevantDecals)
 	{
-		VisibleDecals = DecalRendering::BuildVisibleDecalList(Scene.Decals, View);
-
-		// Build a list of decals that need to be rendered for this view
-		SortedDecals = DecalRendering::BuildRelevantDecalList(VisibleDecals, DecalRenderStage);
-		INC_DWORD_STAT_BY(STAT_Decals, SortedDecals.Num());
+		if (DecalViewInfo.RenderTargetMode == DecalRendering::GetRenderTargetMode(View.GetShaderPlatform(), VisibleDecal->BlendDesc, DecalViewInfo.RenderStage))
+		{
+			SortedDecals.Add(VisibleDecal);
+		}
 	}
+	INC_DWORD_STAT_BY(STAT_Decals, SortedDecals.Num());
+	
 
 	if (SortedDecals.Num() > 0)
 	{
@@ -132,10 +130,10 @@ void RenderDeferredDecalsMobile(FRHICommandList& RHICmdList, const FScene& Scene
 					StencilRef, 0x00>::GetRHI();
 			}
 
-			GraphicsPSOInit.BlendState = DecalRendering::GetDecalBlendState(VisibleDecal.BlendDesc, DecalRenderStage, RenderTargetMode);
+			GraphicsPSOInit.BlendState = DecalRendering::GetDecalBlendState(VisibleDecal.BlendDesc, DecalViewInfo.RenderStage, DecalViewInfo.RenderTargetMode);
 
 			// Set shader params
-			DecalRendering::SetShader(RHICmdList, GraphicsPSOInit, StencilRef, View, VisibleDecal, DecalRenderStage, FrustumComponentToClip, &Scene);
+			DecalRendering::SetShader(RHICmdList, GraphicsPSOInit, StencilRef, View, VisibleDecal, DecalViewInfo.RenderStage, FrustumComponentToClip, Scene);
 
 			RHICmdList.DrawIndexedPrimitive(GetUnitCubeIndexBuffer(), 0, 0, 8, 0, UE_ARRAY_COUNT(GCubeIndices) / 3, View.GetStereoPassInstanceFactor());
 		}
@@ -149,18 +147,20 @@ void FMobileSceneRenderer::RenderDBuffer(FRDGBuilder& GraphBuilder, FSceneTextur
 
 	const EShaderPlatform Platform = GetViewFamilyInfo(Views).GetShaderPlatform();
 
-	for (int32 ViewIndex = 0; ViewIndex < Views.Num(); ++ViewIndex)
-	{
-		FViewInfo& View = Views[ViewIndex];
-
-		if (!View.ShouldRenderView())
+	if (DecalVisibilityTaskData != nullptr)
+	{ 
+		for (int32 ViewIndex = 0; ViewIndex < Views.Num(); ++ViewIndex)
 		{
-			continue;
-		}
+			FViewInfo& View = Views[ViewIndex];
 
-		FVisibleDecalList VisibleDecals = DecalRendering::BuildVisibleDecalList(Scene->Decals, View);
-		FRelevantDecalList SortedDecals = DecalRendering::BuildRelevantDecalList(VisibleDecals, EDecalRenderStage::BeforeBasePass);
-		FDeferredDecalPassTextures DecalPassTextures = GetDeferredDecalPassTextures(GraphBuilder, View, SceneTextures, &DBufferTextures, EDecalRenderStage::BeforeBasePass);
-		AddDeferredDecalPass(GraphBuilder, View, SortedDecals, DecalPassTextures, InstanceCullingManager, EDecalRenderStage::BeforeBasePass);
+			if (!View.ShouldRenderView())
+			{
+				continue;
+			}
+
+			TConstArrayView<const FVisibleDecal*> SortedDecals = DecalVisibilityTaskData->FinishRelevantDecals(ViewIndex, EDecalRenderStage::BeforeBasePass);
+			FDeferredDecalPassTextures DecalPassTextures = GetDeferredDecalPassTextures(GraphBuilder, View, SceneTextures, &DBufferTextures, EDecalRenderStage::BeforeBasePass);
+			AddDeferredDecalPass(GraphBuilder, View, SortedDecals, DecalPassTextures, InstanceCullingManager, EDecalRenderStage::BeforeBasePass);
+		}
 	}
-}
+}
\ No newline at end of file
diff --git a/Engine/Source/Runtime/Renderer/Private/MobileShadingRenderer.cpp b/Engine/Source/Runtime/Renderer/Private/MobileShadingRenderer.cpp
--- a/Engine/Source/Runtime/Renderer/Private/MobileShadingRenderer.cpp
+++ b/Engine/Source/Runtime/Renderer/Private/MobileShadingRenderer.cpp
@@ -80,6 +80,7 @@
 #include "GenerateMips.h"
 #include "MobileSSR.h"
 #include "ViewData.h"
+#include "DecalRenderingShared.h"
 #include "DistanceFieldAmbientOcclusion.h"
 #include "DistanceFieldLightingShared.h"
 #include "DepthCopy.h"
@@ -774,6 +775,15 @@ void FMobileSceneRenderer::InitViews(
 
 		UpdateDirectionalLightUniformBuffers(GraphBuilder, Views[0]);
 	}
+
+	if(!ViewFamily.EngineShowFlags.HitProxies)
+	{ 
+		DecalVisibilityTaskData = FDecalVisibilityTaskData::Launch(GraphBuilder, *Scene, Views);
+	}
+	else
+	{
+		DecalVisibilityTaskData = nullptr;
+	}
 }
 
 static void BeginOcclusionScope(FRDGBuilder& GraphBuilder, TArray<FViewInfo>& Views)
@@ -1007,11 +1017,15 @@ void FMobileSceneRenderer::RenderCustomRenderPassBasePass(FRDGBuilder& GraphBuil
 			}
 		}
 
+		FDecalViewInfo DecalViewInfo = FDecalViewInfo(ViewContext.ViewIndex,
+			bRequiresDBufferDecals ? EDecalRenderStage::Emissive : EDecalRenderStage::Mobile,
+			bDeferredShading ? EDecalRenderTargetMode::SceneColorAndGBuffer : EDecalRenderTargetMode::SceneColor);
+
 		GraphBuilder.AddPass(
 			RDG_EVENT_NAME("RenderMobileBasePass"),
 			PassParameters,
 			ERDGPassFlags::Raster,
-			[this, ParameterCollection, ViewContext, NumAdditionalSubpasses](FRHICommandList& RHICmdList)
+			[this, DecalViewInfo, ParameterCollection, ViewContext, NumAdditionalSubpasses](FRHICommandList& RHICmdList)
 			{
 				FViewInfo& View = *ViewContext.ViewInfo;
 				auto PassParameters = &ParameterCollection->PassParameters;
@@ -1022,7 +1036,7 @@ void FMobileSceneRenderer::RenderCustomRenderPassBasePass(FRDGBuilder& GraphBuil
 					RHICmdList.NextSubpass();
 					if (AddSubpass == 0)
 					{
-						RenderDecals(RHICmdList, View, &ParameterCollection->MeshDecalInstanceCullingDrawParams);
+						RenderDecals(RHICmdList, View, DecalViewInfo, &ParameterCollection->MeshDecalInstanceCullingDrawParams);
 					}
 				}
 			});
@@ -1037,7 +1051,7 @@ void FMobileSceneRenderer::RenderCustomRenderPassBasePass(FRDGBuilder& GraphBuil
 			BasePassRenderTargets.DepthStencil.SetStencilLoadAction(ERenderTargetLoadAction::ELoad);
 			BasePassRenderTargets.DepthStencil.SetDepthStencilAccess(FExclusiveDepthStencil::DepthRead_StencilWrite);
 
-			RenderDeferredMultiPassDecals(GraphBuilder, View, InstanceCullingManager, BasePassRenderTargets);
+			RenderDeferredMultiPassDecals(GraphBuilder, View, ViewContext.ViewIndex, SceneTextures, InstanceCullingManager);
 		}
 
 		if (bIncludeTranslucent)
@@ -1635,6 +1649,11 @@ void FMobileSceneRenderer::Render(FRDGBuilder& GraphBuilder, const FSceneRenderU
 			RenderForward(GraphBuilder, ViewFamilyTexture, SceneTextures, BasePassTextures, InstanceCullingManager);
 		}
 
+		if (DecalVisibilityTaskData)
+		{
+			DecalVisibilityTaskData->Finish();
+		}
+
 		if (DoOcclusionQueries() && !bIsFullDepthPrepassEnabled)
 		{
 			EndOcclusionScope(GraphBuilder, Views);
@@ -2023,12 +2042,16 @@ void FMobileSceneRenderer::RenderForwardSinglePass(FRDGBuilder& GraphBuilder, FM
 	const bool bDoOcclusionInMainPass = bDoOcclusionQueries && !bVulkanAdrenoOcclusionMode;
 	PassParameters->RenderTargets.NumOcclusionQueries = bVulkanAdrenoOcclusionMode ? 0u : NumOcclusionQueries;
 	
+	FDecalViewInfo DecalViewInfo = FDecalViewInfo(ViewContext.ViewIndex,
+		bRequiresDBufferDecals ? EDecalRenderStage::Emissive : EDecalRenderStage::Mobile,
+		EDecalRenderTargetMode::SceneColor);
+
 	GraphBuilder.AddPass(
 		RDG_EVENT_NAME("SceneColorRendering"),
 		PassParameters,
 		// the second view pass should not be merged with the first view pass on mobile since the subpass would not work properly.
 		ERDGPassFlags::Raster | ERDGPassFlags::NeverMerge,
-		[this, PassParameters, ParameterCollection, ViewContext, bDoOcclusionInMainPass, &SceneTextures, ViewQueries = bDoOcclusionInMainPass ? MoveTemp(*ViewOcclusionQueries) : FViewOcclusionQueries()](FRHICommandList& RHICmdList)
+		[this, DecalViewInfo, PassParameters, ParameterCollection, ViewContext, bDoOcclusionInMainPass, &SceneTextures, ViewQueries = bDoOcclusionInMainPass ? MoveTemp(*ViewOcclusionQueries) : FViewOcclusionQueries()](FRHICommandList& RHICmdList)
 	{
 		FViewInfo& View = *ViewContext.ViewInfo;
 			
@@ -2046,7 +2069,7 @@ void FMobileSceneRenderer::RenderForwardSinglePass(FRDGBuilder& GraphBuilder, FM
 		PostRenderBasePass(RHICmdList, View);
 		// scene depth is read only and can be fetched
 		RHICmdList.NextSubpass();
-		RenderDecals(RHICmdList, View, &ParameterCollection->MeshDecalSceneColorInstanceCullingDrawParams);
+		RenderDecals(RHICmdList, View, DecalViewInfo, &ParameterCollection->MeshDecalSceneColorInstanceCullingDrawParams);
 		RenderModulatedShadowProjections(RHICmdList, ViewContext.ViewIndex, View);
 		if (GMaxRHIShaderPlatform != SP_METAL_SIM)
 		{
@@ -2139,6 +2162,10 @@ void FMobileSceneRenderer::RenderForwardMultiPass(FRDGBuilder& GraphBuilder, FMo
 		BuildMeshRenderingCommands(GraphBuilder, EMeshPass::DebugViewMode, View, Scene->GPUScene, InstanceCullingManager, ParameterCollection->DebugViewModeInstanceCullingDrawParams);
 	}
 
+	FDecalViewInfo DecalViewInfo = FDecalViewInfo(ViewContext.ViewIndex, 
+												  bRequiresDBufferDecals ? EDecalRenderStage::Emissive : EDecalRenderStage::Mobile,
+												  EDecalRenderTargetMode::SceneColor);
+
 	GraphBuilder.AddPass(
 		RDG_EVENT_NAME("SceneColorRendering"),
 		PassParameters,
@@ -2214,13 +2241,13 @@ void FMobileSceneRenderer::RenderForwardMultiPass(FRDGBuilder& GraphBuilder, FMo
 		RDG_EVENT_NAME("DecalsAndTranslucency"),
 		SecondPassParameters,
 		ERDGPassFlags::Raster,
-		[this, SecondParameterCollection, ViewContext, bDoOcclusionQueries, &SceneTextures, ViewQueries = (bDoOcclusionQueries) ? MoveTemp(*ViewOcclusionQueries) : FViewOcclusionQueries()](FRHICommandList& RHICmdList)
+		[this, SecondParameterCollection, ViewContext, bDoOcclusionQueries, &SceneTextures, ViewQueries = (bDoOcclusionQueries) ? MoveTemp(*ViewOcclusionQueries) : FViewOcclusionQueries(), DecalViewInfo](FRHICommandList& RHICmdList)
 	{
 		FViewInfo& View = *ViewContext.ViewInfo;
 		auto SecondPassParameters = &SecondParameterCollection->PassParameters;
 			
 		// scene depth is read only and can be fetched
-		RenderDecals(RHICmdList, View, &SecondParameterCollection->MeshDecalSceneColorInstanceCullingDrawParams);
+		RenderDecals(RHICmdList, View, DecalViewInfo, &SecondParameterCollection->MeshDecalSceneColorInstanceCullingDrawParams);
 		RenderModulatedShadowProjections(RHICmdList, ViewContext.ViewIndex, View);
 		RenderFog(RHICmdList, View);
 		// Draw translucency.
@@ -2466,6 +2493,10 @@ void FMobileSceneRenderer::RenderDeferredSinglePass(FRDGBuilder& GraphBuilder, F
 			BuildMeshRenderingCommands(GraphBuilder, StandardTranslucencyMeshPass, View, Scene->GPUScene, InstanceCullingManager, ParameterCollection->TranslucencyInstanceCullingDrawParams);
 		}
 
+		FDecalViewInfo DecalViewInfo = FDecalViewInfo(ViewContext.ViewIndex,
+													  bRequiresDBufferDecals ? EDecalRenderStage::Emissive : EDecalRenderStage::MobileBeforeLighting,
+													  EDecalRenderTargetMode::SceneColorAndGBuffer);
+
 		GraphBuilder.AddPass(
 			RDG_EVENT_NAME("SceneColorRendering"),
 			PassParameters,
@@ -2478,6 +2509,7 @@ void FMobileSceneRenderer::RenderDeferredSinglePass(FRDGBuilder& GraphBuilder, F
 			&SortedLightSet, 
 			bUsingPixelLocalStorage,
 			bDoOcclusionQueries, 
+			DecalViewInfo,
 			ViewQueries = (bDoOcclusionQueries) ? MoveTemp(QueriesPerView[ViewContext.ViewIndex]) : FViewOcclusionQueries()](FRHICommandList& RHICmdList)
 			{
 				FViewInfo& View = *ViewContext.ViewInfo;
@@ -2492,7 +2524,7 @@ void FMobileSceneRenderer::RenderDeferredSinglePass(FRDGBuilder& GraphBuilder, F
 				PostRenderBasePass(RHICmdList, View);
 				// SceneColor + GBuffer write, SceneDepth is read only
 				RHICmdList.NextSubpass();
-				RenderDecals(RHICmdList, View, &ParameterCollection->MeshDecalSceneColorAndGBufferInstanceCullingDrawParams);
+				RenderDecals(RHICmdList, View, DecalViewInfo, &ParameterCollection->MeshDecalSceneColorAndGBufferInstanceCullingDrawParams);
 				// SceneColor write, SceneDepth is read only
 				RHICmdList.NextSubpass();
 				MobileDeferredShadingPass(RHICmdList, ViewContext.ViewIndex, Views.Num(), View, *Scene, SortedLightSet, VisibleLightInfos, MobileSSRQuality);
@@ -2514,10 +2546,10 @@ void FMobileSceneRenderer::RenderDeferredSinglePass(FRDGBuilder& GraphBuilder, F
 	}
 }
 
-void FMobileSceneRenderer::RenderDeferredMultiPassDecals(FRDGBuilder& GraphBuilder, FViewInfo& View, FInstanceCullingManager& InstanceCullingManager, FRenderTargetBindingSlots& BasePassRenderTargets)
+void FMobileSceneRenderer::RenderDeferredMultiPassDecals(FRDGBuilder& GraphBuilder, FViewInfo& View, uint32 ViewIndex, FSceneTextures& SceneTextures, FInstanceCullingManager& InstanceCullingManager)
 {
-	const EMeshPass::Type MeshDecalsPass = EMeshPass::MeshDecal_SceneColorAndGBuffer;
-	if (Scene->Decals.Num() == 0 && !HasAnyDraw(View.ParallelMeshDrawCommandPasses[MeshDecalsPass]))
+	if (Scene->Decals.Num() == 0 && !HasAnyDraw(View.ParallelMeshDrawCommandPasses[EMeshPass::MeshDecal_SceneColorAndGBuffer])
+								 && !HasAnyDraw(View.ParallelMeshDrawCommandPasses[EMeshPass::MeshDecal_SceneColorAndGBufferNoNormal]))
 	{
 		// Skip decal render-pass if there is nothing to draw
 		return;
@@ -2525,24 +2557,65 @@ void FMobileSceneRenderer::RenderDeferredMultiPassDecals(FRDGBuilder& GraphBuild
 
 	View.BeginRenderView();
 
-	auto PassParameters = GraphBuilder.AllocParameters<FMobileRenderPassParameters>();
-	PassParameters->View = View.GetShaderParameters();
-	PassParameters->MobileBasePass = CreateMobileBasePassUniformBuffer(GraphBuilder, View, EMobileBasePass::Opaque, EMobileSceneTextureSetupMode::SceneDepth);
-	PassParameters->RenderTargets = BasePassRenderTargets;
+	const EDecalRenderTargetMode DecalRenderTargetModes[] = { EDecalRenderTargetMode::SceneColorAndGBuffer, EDecalRenderTargetMode::SceneColorAndGBufferNoNormal };
+	const EDecalRenderStage DecalRenderStage = bRequiresDBufferDecals ? EDecalRenderStage::Emissive : EDecalRenderStage::MobileBeforeLighting;
+	constexpr const TCHAR* PassNames[2] = { TEXT("SceneColorAndGBuffer"), TEXT("SceneColorAndGBufferNoNormal")};
+	TConstArrayView<const FVisibleDecal*> RelevantDecals = DecalVisibilityTaskData ? DecalVisibilityTaskData->FinishRelevantDecals(ViewIndex, DecalRenderStage) : TConstArrayView<const FVisibleDecal*>();
 
-	if (Scene->GPUScene.IsEnabled())
+	for (const EDecalRenderTargetMode DecalRenderTargetMode : DecalRenderTargetModes)
 	{
-		BuildMeshRenderingCommands(GraphBuilder, MeshDecalsPass, View, Scene->GPUScene, InstanceCullingManager, PassParameters->InstanceCullingDrawParams);
-	}
+		EMeshPass::Type MeshDecalsPass = DecalRendering::GetMeshPassType(DecalRenderTargetMode);
 
-	GraphBuilder.AddPass(
-		RDG_EVENT_NAME("Decals"),
-		PassParameters,
-		ERDGPassFlags::Raster,
-		[this, PassParameters, &View](FRHICommandList& RHICmdList)
+		const FVisibleDecal* const* FirstDecal = RelevantDecals.FindByPredicate( [&View, &DecalRenderTargetMode, &DecalRenderStage](const FVisibleDecal* Decal) { return DecalRenderTargetMode == DecalRendering::GetRenderTargetMode(View.GetShaderPlatform(), Decal->BlendDesc, DecalRenderStage); });
+		if (FirstDecal == nullptr && !HasAnyDraw(View.ParallelMeshDrawCommandPasses[MeshDecalsPass]))
 		{
-			RenderDecals(RHICmdList, View, &PassParameters->InstanceCullingDrawParams);
-		});
+			continue;
+		}
+
+		auto PassParameters = GraphBuilder.AllocParameters<FMobileRenderPassParameters>();
+		PassParameters->View = View.GetShaderParameters();
+		
+		uint32 ColorTargetIndex = 0;
+		FRenderTargetBindingSlots& RenderTargets = PassParameters->RenderTargets;
+		RenderTargets[ColorTargetIndex++] = FRenderTargetBinding(SceneTextures.Color.Target, ERenderTargetLoadAction::ELoad);
+		if (DecalRenderTargetMode != EDecalRenderTargetMode::SceneColorAndGBufferNoNormal)
+		{
+			RenderTargets[ColorTargetIndex++] = FRenderTargetBinding(SceneTextures.GBufferA, ERenderTargetLoadAction::ELoad);
+		}
+		RenderTargets[ColorTargetIndex++] = FRenderTargetBinding(SceneTextures.GBufferB, ERenderTargetLoadAction::ELoad);
+		RenderTargets[ColorTargetIndex++] = FRenderTargetBinding(SceneTextures.GBufferC, ERenderTargetLoadAction::ELoad);
+		RenderTargets.DepthStencil = FDepthStencilBinding(
+			SceneTextures.Depth.Target,
+			ERenderTargetLoadAction::ELoad,
+			ERenderTargetLoadAction::ELoad,
+			FExclusiveDepthStencil::DepthRead_StencilWrite);
+
+		EMobileSceneTextureSetupMode TextureReadAccess = EMobileSceneTextureSetupMode::SceneDepth;
+		if (DecalRenderTargetMode == EDecalRenderTargetMode::SceneColorAndGBufferNoNormal)
+		{
+			TextureReadAccess |= EMobileSceneTextureSetupMode::GBufferA;
+		}
+		PassParameters->MobileBasePass = CreateMobileBasePassUniformBuffer(GraphBuilder, View, EMobileBasePass::Opaque, TextureReadAccess);
+		PassParameters->RenderTargets = RenderTargets;
+
+		if (Scene->GPUScene.IsEnabled())
+		{
+			BuildMeshRenderingCommands(GraphBuilder, MeshDecalsPass, View, Scene->GPUScene, InstanceCullingManager, PassParameters->InstanceCullingDrawParams);
+		}
+
+		FDecalViewInfo DecalViewInfo = FDecalViewInfo(ViewIndex,
+			DecalRenderStage,
+			DecalRenderTargetMode);
+
+		GraphBuilder.AddPass(
+			RDG_EVENT_NAME("Decals_%s", DecalRendering::GetRenderTargetModeName(DecalRenderTargetMode)),
+			PassParameters,
+			ERDGPassFlags::Raster,
+			[this, DecalViewInfo, PassParameters, &View](FRHICommandList& RHICmdList)
+			{
+				RenderDecals(RHICmdList, View, DecalViewInfo, &PassParameters->InstanceCullingDrawParams);
+			});
+	}
 }
 
 void FMobileSceneRenderer::RenderDeferredMultiPass(FRDGBuilder& GraphBuilder, FSceneTextures& SceneTextures, const FSortedLightSetSceneInfo& SortedLightSet, FMobileBasePassTextures& BasePassTextures, FInstanceCullingManager& InstanceCullingManager)
@@ -2609,10 +2682,7 @@ void FMobileSceneRenderer::RenderDeferredMultiPass(FRDGBuilder& GraphBuilder, FS
 			BuildMeshRenderingCommands(GraphBuilder, EMeshPass::BasePass, View, Scene->GPUScene, InstanceCullingManager, PassParameters->InstanceCullingDrawParams);
 			BuildMeshRenderingCommands(GraphBuilder, EMeshPass::SkyPass, View, Scene->GPUScene, InstanceCullingManager, ParameterCollection->SkyPassInstanceCullingDrawParams);
 			BuildMeshRenderingCommands(GraphBuilder, EMeshPass::DebugViewMode, View, Scene->GPUScene, InstanceCullingManager, ParameterCollection->DebugViewModeInstanceCullingDrawParams);
-			if (bIsFullDepthPrepassEnabled)
-			{
-				BuildMeshRenderingCommands(GraphBuilder, EMeshPass::MeshDecal_SceneColorAndGBuffer, View, Scene->GPUScene, InstanceCullingManager, ParameterCollection->MeshDecalSceneColorAndGBufferInstanceCullingDrawParams);
-			}
+
 		}
 
 		GraphBuilder.AddPass(
@@ -2644,10 +2714,6 @@ void FMobileSceneRenderer::RenderDeferredMultiPass(FRDGBuilder& GraphBuilder, FS
 					RenderOcclusion(RHICmdList, View, ViewQueries);
 				}
 
-				if (bIsFullDepthPrepassEnabled)
-				{
-					RenderDecals(RHICmdList, View, &ParameterCollection->MeshDecalSceneColorAndGBufferInstanceCullingDrawParams);
-				}
 			});
 	}
 
@@ -2667,17 +2733,15 @@ void FMobileSceneRenderer::RenderDeferredMultiPass(FRDGBuilder& GraphBuilder, FS
 	BasePassRenderTargets.DepthStencil.SetDepthStencilAccess(FExclusiveDepthStencil::DepthRead_StencilWrite);
 
 	// Decals
-	if (!bIsFullDepthPrepassEnabled)
+	for (FRenderViewContext& ViewContext : RenderViews)
 	{
-		for (FRenderViewContext& ViewContext : RenderViews)
-		{
-			FViewInfo& View = *ViewContext.ViewInfo;
-
-			SCOPED_GPU_MASK(GraphBuilder.RHICmdList, !View.IsInstancedStereoPass() ? View.GPUMask : (View.GPUMask | View.GetInstancedView()->GPUMask));
-			SCOPED_CONDITIONAL_DRAW_EVENTF(GraphBuilder.RHICmdList, EventView, RenderViews.Num() > 1, TEXT("View%d"), ViewContext.ViewIndex);
+		FViewInfo& View = *ViewContext.ViewInfo;
+	
+		
+		SCOPED_GPU_MASK(GraphBuilder.RHICmdList, !View.IsInstancedStereoPass() ? View.GPUMask : (View.GPUMask | View.GetInstancedView()->GPUMask));
+		SCOPED_CONDITIONAL_DRAW_EVENTF(GraphBuilder.RHICmdList, EventView, RenderViews.Num() > 1, TEXT("View%d"), ViewContext.ViewIndex);
 
-			RenderDeferredMultiPassDecals(GraphBuilder, View, InstanceCullingManager, BasePassRenderTargets);
-		}
+		RenderDeferredMultiPassDecals(GraphBuilder, View, ViewContext.ViewIndex, SceneTextures, InstanceCullingManager);
 	}
 
 	if (!bIsFullDepthPrepassEnabled)
diff --git a/Engine/Source/Runtime/Renderer/Private/SceneRendering.h b/Engine/Source/Runtime/Renderer/Private/SceneRendering.h
--- a/Engine/Source/Runtime/Renderer/Private/SceneRendering.h
+++ b/Engine/Source/Runtime/Renderer/Private/SceneRendering.h
@@ -57,6 +57,8 @@
 #include "SceneExtensions.h"
 #include "PostProcess/LensDistortion.h"
 #include "DepthRendering.h"
+#include "DecalRenderingShared.h"
+#include "DecalRenderingCommon.h"
 
 #if RHI_RAYTRACING
 #include "RayTracing/RayTracingScene.h"
@@ -2917,7 +2919,8 @@ protected:
 
 
 	/** Renders decals. */
-	void RenderDecals(FRHICommandList& RHICmdList, FViewInfo& View, const FInstanceCullingDrawParams* InstanceCullingDrawParams);
+	void RenderDecals(FRHICommandList& RHICmdList, FViewInfo& View, const FDecalViewInfo& DecalViewInfo, const FInstanceCullingDrawParams* InstanceCullingDrawParams);
+	void RenderDeferredDecalsMobile(FRHICommandList& RHICmdList, const FViewInfo& View, const FDecalViewInfo& DecalViewInfo);
 	void RenderDBuffer(FRDGBuilder& GraphBuilder, FSceneTextures& SceneTextures, FDBufferTextures& DBufferTextures, FInstanceCullingManager& InstanceCullingManager);
 
 	/** Renders the atmospheric and height fog */
@@ -2940,7 +2943,7 @@ protected:
 	
 	void RenderDeferredSinglePass(FRDGBuilder& GraphBuilder, FSceneTextures& SceneTextures, const FSortedLightSetSceneInfo& SortedLightSet, FMobileBasePassTextures& BasePassTextures, FInstanceCullingManager& InstanceCullingManager);
 	void RenderDeferredMultiPass(FRDGBuilder& GraphBuilder, FSceneTextures& SceneTextures, const FSortedLightSetSceneInfo& SortedLightSet, FMobileBasePassTextures& BasePassTextures, FInstanceCullingManager& InstanceCullingManager);
-	void RenderDeferredMultiPassDecals(FRDGBuilder& GraphBuilder, FViewInfo& View, FInstanceCullingManager& InstanceCullingManager, FRenderTargetBindingSlots& BasePassRenderTargets);
+	void RenderDeferredMultiPassDecals(FRDGBuilder& GraphBuilder, FViewInfo& View, uint32 ViewIndex, FSceneTextures& SceneTextures, FInstanceCullingManager& InstanceCullingManager);
 
 	void RenderStereoMotionVectors(FRDGBuilder& GraphBuilder, FViewInfo& View, const FSceneTextures& SceneTextures, FInstanceCullingManager& InstanceCullingManager);
 	bool HasStereoMotionVectorRenderTargets(const FSceneTextures& SceneTextures);
@@ -2983,7 +2986,7 @@ private:
 	bool bEnableDistanceFieldAO;
 	ETranslucencyPass::Type StandardTranslucencyPass;
 	EMeshPass::Type StandardTranslucencyMeshPass;
-
+	FDecalVisibilityTaskData* DecalVisibilityTaskData = nullptr;
 	const FViewInfo* CachedView = nullptr;
 };
 
