SceneCapture: Allow main view camera to be used without forcing main view resolution.

* TemporalAA jitter copy is also made dependent on a matching main view resolution.

#Synced-CL 49904425

diff --git a/Engine/Source/Runtime/Engine/Classes/Components/SceneCaptureComponent2D.h b/Engine/Source/Runtime/Engine/Classes/Components/SceneCaptureComponent2D.h
--- a/Engine/Source/Runtime/Engine/Classes/Components/SceneCaptureComponent2D.h
+++ b/Engine/Source/Runtime/Engine/Classes/Components/SceneCaptureComponent2D.h
@@ -179,7 +179,7 @@ public:
 	UPROPERTY(EditAnywhere, BlueprintReadWrite, Category = PassInfo)
 	uint32 bMainViewResolution : 1;
 
-	/** Render with main view camera.  Enables Main View Family and Resolution.  Temporal AA jitter is matched with main view. */
+	/** Render with main view camera.  Enables Main View Family.  Temporal AA jitter is matched with main view if Main View Resolution is enabled and no Resolution Divisor is set. */
 	UPROPERTY(EditAnywhere, BlueprintReadWrite, Category=PassInfo)
 	uint32 bMainViewCamera : 1;
 
@@ -214,12 +214,12 @@ public:
 
 	inline bool ShouldRenderWithMainViewResolution() const
 	{
-		return bMainViewResolution || bMainViewCamera || (ShouldRenderInMainRenderer() && (!UserSceneTextureBaseColor.IsNone() || !UserSceneTextureNormal.IsNone() || !UserSceneTextureSceneColor.IsNone()));
+		return bMainViewResolution || (ShouldRenderInMainRenderer() && (!UserSceneTextureBaseColor.IsNone() || !UserSceneTextureNormal.IsNone() || !UserSceneTextureSceneColor.IsNone()));
 	}
 
 	inline bool ShouldRenderWithMainViewFamily() const
 	{
-		return bMainViewFamily || ShouldRenderWithMainViewResolution();
+		return bMainViewFamily || bMainViewCamera || ShouldRenderWithMainViewResolution();
 	}
 
 	inline bool ShouldRenderWithMainViewCamera() const
diff --git a/Engine/Source/Runtime/Renderer/Private/SceneCaptureRendering.cpp b/Engine/Source/Runtime/Renderer/Private/SceneCaptureRendering.cpp
--- a/Engine/Source/Runtime/Renderer/Private/SceneCaptureRendering.cpp
+++ b/Engine/Source/Runtime/Renderer/Private/SceneCaptureRendering.cpp
@@ -1379,12 +1379,14 @@ void FScene::UpdateSceneCaptureContents(USceneCaptureComponent2D* CaptureCompone
 			SceneRendererCreationFlags |= ESceneRendererCreationFlags::CaptureSceneColor;
 		}
 
-		// Copy temporal AA related settings for main view camera scene capture, to match jitter.  Don't match if the resolution divisor is set,
+		// Copy temporal AA related settings for main view camera scene capture, to match jitter.  Don't match if resolution differs from main view,
 		// if it's set to ignore screen percentage, or if it's final color, which will run its own AA.  For custom render passes (handled above),
 		// computed jitter results are copied from the main view later in FSceneRenderer::PrepareViewStateForVisibility, but this doesn't work
 		// for regular scene captures, because they run in a separate scene renderer before the main view, where the main view's results haven't
 		// been computed yet.
-		if ((CaptureComponent->ShouldRenderWithMainViewCamera() && CaptureComponent->MainViewFamily &&
+		if ((CaptureComponent->MainViewFamily &&
+			CaptureComponent->ShouldRenderWithMainViewCamera() &&
+			CaptureComponent->ShouldRenderWithMainViewResolution() &&
 			CaptureComponent->MainViewResolutionDivisor.X <= 1 && CaptureComponent->MainViewResolutionDivisor.Y <= 1 && !CaptureComponent->ShouldIgnoreScreenPercentage() &&
 			CaptureComponent->CaptureSource != ESceneCaptureSource::SCS_FinalColorLDR &&
 			CaptureComponent->CaptureSource != ESceneCaptureSource::SCS_FinalColorHDR &&
