Add support for rendering distortion before translucency in the Mobile Renderer. This allows proper refraction and distortion effects to be computed using the scene color prior to translucency, matching the behavior of the desktop renderer.

#tests Editor preview modes

#Synced-CL 49917480

diff --git a/Engine/Source/Runtime/Renderer/Private/DistortionRendering.cpp b/Engine/Source/Runtime/Renderer/Private/DistortionRendering.cpp
--- a/Engine/Source/Runtime/Renderer/Private/DistortionRendering.cpp
+++ b/Engine/Source/Runtime/Renderer/Private/DistortionRendering.cpp
@@ -14,6 +14,7 @@
 #include "ScenePrivate.h"
 #include "ScreenPass.h"
 #include "MeshPassProcessor.inl"
+#include "MobileDistortionPass.h"
 #include "Substrate/Substrate.h"
 #include "ScreenRendering.h"
 #include "PostProcess/TemporalAA.h"
@@ -21,7 +22,7 @@
 #include "HeterogeneousVolumes/HeterogeneousVolumes.h"
 #include "BlueNoise.h"
 
-DECLARE_GPU_DRAWCALL_STAT(Distortion);
+DEFINE_GPU_DRAWCALL_STAT(Distortion);
 
 static TAutoConsoleVariable<int32> CVarDisableDistortion(
 	TEXT("r.DisableDistortion"),
@@ -1185,24 +1186,29 @@ FMeshPassProcessor* CreateDistortionPassProcessor(ERHIFeatureLevel::Type Feature
 
 FMeshPassProcessor* CreateMobileDistortionPassProcessor(ERHIFeatureLevel::Type FeatureLevel, const FScene* Scene, const FSceneView* InViewIfDynamicMeshCommand, FMeshPassDrawListContext* InDrawListContext)
 {
-	FMeshPassProcessorRenderState DistortionPassState;
-
-	// We don't have depth, render all pixels, pixel shader will sample SceneDepth from SceneColor.A and discard if occluded
-	DistortionPassState.SetDepthStencilState(TStaticDepthStencilState<false, CF_Always>::GetRHI());
-
-	if (GetSubstrateEnabledUseRoughRefraction())
+	if (IsMobileDistortionBeforeTranslucency())
 	{
-		DistortionPassState.SetBlendState(TStaticBlendState<CW_RGBA, BO_Add, BF_One, BF_One, BO_Add, BF_One, BF_One,
-															CW_RGBA, BO_Add, BF_One, BF_One, BO_Add, BF_One, BF_One,
-															CW_RED,  BO_Min, BF_One, BF_One, BO_Min, BF_One, BF_One>::GetRHI());
+		return CreateDistortionPassProcessor(FeatureLevel, Scene, InViewIfDynamicMeshCommand, InDrawListContext);
 	}
 	else
 	{
-		// additive blending of offsets
-		DistortionPassState.SetBlendState(TStaticBlendState<CW_RGBA, BO_Add, BF_One, BF_One, BO_Add, BF_One, BF_One>::GetRHI());
-	}
+		FMeshPassProcessorRenderState DistortionPassState;
+		// We don't have depth, render all pixels, pixel shader will sample SceneDepth from SceneColor.A and discard if occluded
+		DistortionPassState.SetDepthStencilState(TStaticDepthStencilState<false, CF_Always>::GetRHI());
 
-	return new FDistortionMeshProcessor(Scene, FeatureLevel, InViewIfDynamicMeshCommand, DistortionPassState, DistortionPassState, InDrawListContext);
+		if (GetSubstrateEnabledUseRoughRefraction())
+		{
+			DistortionPassState.SetBlendState(TStaticBlendState<CW_RGBA, BO_Add, BF_One, BF_One, BO_Add, BF_One, BF_One,
+				CW_RGBA, BO_Add, BF_One, BF_One, BO_Add, BF_One, BF_One,
+				CW_RED, BO_Min, BF_One, BF_One, BO_Min, BF_One, BF_One>::GetRHI());
+		}
+		else
+		{
+			// additive blending of offsets
+			DistortionPassState.SetBlendState(TStaticBlendState<CW_RGBA, BO_Add, BF_One, BF_One, BO_Add, BF_One, BF_One>::GetRHI());
+		}
+		return new FDistortionMeshProcessor(Scene, FeatureLevel, InViewIfDynamicMeshCommand, DistortionPassState, DistortionPassState, InDrawListContext);
+	}
 }
 
 REGISTER_MESHPASSPROCESSOR_AND_PSOCOLLECTOR(DistortionPass, CreateDistortionPassProcessor, EShadingPath::Deferred, EMeshPass::Distortion, EMeshPassFlags::MainView);
diff --git a/Engine/Source/Runtime/Renderer/Private/DistortionRendering.h b/Engine/Source/Runtime/Renderer/Private/DistortionRendering.h
--- a/Engine/Source/Runtime/Renderer/Private/DistortionRendering.h
+++ b/Engine/Source/Runtime/Renderer/Private/DistortionRendering.h
@@ -13,6 +13,8 @@
 class FPrimitiveSceneProxy;
 struct FMeshPassProcessorRenderState;
 
+DECLARE_GPU_DRAWCALL_STAT_EXTERN(Distortion);
+
 extern void SetupDistortionParams(FVector4f& DistortionParams, const FViewInfo& View);
 
 class FDistortionMeshProcessor : public FSceneRenderingAllocatorObject<FDistortionMeshProcessor>, public FMeshPassProcessor
diff --git a/Engine/Source/Runtime/Renderer/Private/MobileDistortionPass.cpp b/Engine/Source/Runtime/Renderer/Private/MobileDistortionPass.cpp
--- a/Engine/Source/Runtime/Renderer/Private/MobileDistortionPass.cpp
+++ b/Engine/Source/Runtime/Renderer/Private/MobileDistortionPass.cpp
@@ -21,6 +21,23 @@ END_GLOBAL_SHADER_PARAMETER_STRUCT()
 
 IMPLEMENT_STATIC_UNIFORM_BUFFER_STRUCT(FMobileDistortionPassUniformParameters, "MobileDistortionPass", SceneTextures);
 
+
+
+static TAutoConsoleVariable<int32> CVarMobileAdrenoOcclusionMode(
+	TEXT("r.Mobile.Distortion.BeforeTranslucency"),
+	0,
+	TEXT("0: Render distortion in the PostProcess chain (default).\n")
+	TEXT("1: Render distortion before the translucent pass"),
+	ECVF_RenderThreadSafe);
+
+bool IsMobileDistortionBeforeTranslucency()
+{
+	static IConsoleVariable* CVar = IConsoleManager::Get().FindConsoleVariable(TEXT("r.Mobile.Distortion.BeforeTranslucency"));
+	int32 DistortionBeforeTranslucency = CVar->GetInt();
+
+	return DistortionBeforeTranslucency == 1;
+}
+
 bool IsMobileDistortionActive(const FViewInfo& View)
 {
 	static IConsoleVariable* CVar = IConsoleManager::Get().FindConsoleVariable(TEXT("r.DisableDistortion"));
@@ -159,3 +176,149 @@ FScreenPassTexture AddMobileDistortionMergePass(FRDGBuilder& GraphBuilder, const
 
 	return MoveTemp(DistortionMergeOutput);
 }
+
+void FMobileSceneRenderer::RenderDistortion(
+	FRDGBuilder& GraphBuilder,
+	const FViewInfo& View,
+	FRDGTextureRef SceneColorTexture,
+	FRDGTextureRef SceneDepthTexture)
+{
+	check(SceneDepthTexture);
+	check(SceneColorTexture);
+
+	QUICK_SCOPE_CYCLE_COUNTER(STAT_FSceneRenderer_RenderDistortion);
+	SCOPED_NAMED_EVENT(RenderDistortion, FColor::Emerald);
+	RDG_EVENT_SCOPE_STAT(GraphBuilder, Distortion, "Distortion");
+
+	FRDGTextureRef DistortionTexture = nullptr;
+	FGlobalShaderMap* ShaderMap = GetGlobalShaderMap(FeatureLevel);
+	const FDepthStencilBinding StencilReadBinding(SceneDepthTexture, ERenderTargetLoadAction::ELoad, ERenderTargetLoadAction::ELoad, FExclusiveDepthStencil::DepthRead_StencilRead);
+	FDepthStencilBinding StencilWriteBinding(SceneDepthTexture, ERenderTargetLoadAction::ELoad, ERenderTargetLoadAction::ELoad, FExclusiveDepthStencil::DepthRead_StencilWrite);
+
+	// Use stencil mask to optimize cases with lower screen coverage.
+	// Note: This adds an extra pass which is actually slower as distortion tends towards full-screen.
+	//       It could be worth testing object screen bounds then reverting to a target flip and single pass.
+	{
+		QUICK_SCOPE_CYCLE_COUNTER(STAT_FSceneRenderer_RenderDistortion_Accumulate);
+		RDG_EVENT_SCOPE(GraphBuilder, "Accumulate");
+
+		// Use RGBA8 light target for accumulating distortion offsets.
+		// R = positive X offset
+		// G = positive Y offset
+		// B = negative X offset
+		// A = negative Y offset
+		DistortionTexture = GraphBuilder.CreateTexture(
+			FRDGTextureDesc::Create2D(
+				SceneDepthTexture->Desc.Extent,
+				PF_B8G8R8A8,
+				FClearValueBinding::Transparent,
+				GFastVRamConfig.Distortion | TexCreate_RenderTargetable | TexCreate_ShaderResource,
+				1,
+				SceneDepthTexture->Desc.NumSamples),
+			TEXT("Distortion"));
+
+		ERenderTargetLoadAction LoadAction = ERenderTargetLoadAction::EClear;
+
+		if (auto* Pass = View.ParallelMeshDrawCommandPasses[EMeshPass::Distortion])
+		{
+			FMobileDistortionPassParameters* PassParameters = GraphBuilder.AllocParameters<FMobileDistortionPassParameters>();
+			PassParameters->View = View.GetShaderParameters();
+			PassParameters->Pass = CreateMobileDistortionPassUniformBuffer(GraphBuilder, View);
+			PassParameters->RenderTargets[0] = FRenderTargetBinding(DistortionTexture, LoadAction);
+			PassParameters->RenderTargets.DepthStencil = StencilWriteBinding;
+
+			Pass->BuildRenderingCommands(GraphBuilder, Scene->GPUScene, PassParameters->InstanceCullingDrawParams);
+
+			GraphBuilder.AddPass(
+				{},
+				PassParameters,
+				ERDGPassFlags::Raster,
+				[&View, Pass, PassParameters](FRDGAsyncTask, FRHICommandList& RHICmdList)
+				{
+					QUICK_SCOPE_CYCLE_COUNTER(STAT_FSceneRender_RenderDistortion_Accumulate_Meshes);
+					SetStereoViewport(RHICmdList, View);
+					Pass->Draw(RHICmdList, &PassParameters->InstanceCullingDrawParams);
+				});
+
+			LoadAction = ERenderTargetLoadAction::ELoad;
+		}
+
+		if (LoadAction == ERenderTargetLoadAction::EClear)
+		{
+			AddClearRenderTargetPass(GraphBuilder, DistortionTexture);
+		}
+	}
+
+	TShaderMapRef<FScreenPassVS> VertexShader(ShaderMap);
+	FScreenPassPipelineState PipelineState(VertexShader, {});
+	FScreenPassTextureViewport Viewport(SceneColorTexture);
+	TShaderMapRef<FMobileDistortionMergePS> MergePixelShader(ShaderMap);
+
+	FRDGTextureDesc DistortedSceneColorDesc = SceneColorTexture->Desc;
+	//Remove fast clear flag on the DistoredSceneColor which is used in the Apply and Merge passes. 
+	// This can save the Fast clear eliminate in the Merge pass when the RTV is transient allocated.
+	EnumAddFlags(DistortedSceneColorDesc.Flags, TexCreate_NoFastClear);
+	EnumRemoveFlags(DistortedSceneColorDesc.Flags, TexCreate_FastVRAM);
+	FRDGTextureRef DistortionSceneColorTexture = GraphBuilder.CreateTexture(DistortedSceneColorDesc, TEXT("DistortedSceneColor"));
+
+	// Apply distortion and store off-screen.
+	{
+		QUICK_SCOPE_CYCLE_COUNTER(STAT_FSceneRenderer_RenderDistortion_Apply);
+		RDG_EVENT_SCOPE(GraphBuilder, "Apply");
+
+		// Test against stencil mask but don't clear.
+		PipelineState.DepthStencilState = TStaticDepthStencilState<
+			false, CF_Always,
+			true, CF_Equal, SO_Keep, SO_Keep, SO_Keep,
+			false, CF_Always, SO_Keep, SO_Keep, SO_Keep,
+			DISTORTION_STENCIL_MASK_BIT, DISTORTION_STENCIL_MASK_BIT>::GetRHI();
+		PipelineState.StencilRef = DISTORTION_STENCIL_MASK_BIT;
+		PipelineState.PixelShader = MergePixelShader;
+		ERenderTargetLoadAction LoadAction = ERenderTargetLoadAction::ENoAction;
+
+		FMobileDistortionMergePS::FParameters* PassParameters = GraphBuilder.AllocParameters<FMobileDistortionMergePS::FParameters>();
+		PassParameters->RenderTargets[0] = FRenderTargetBinding(DistortionSceneColorTexture, LoadAction);
+		PassParameters->RenderTargets.DepthStencil = StencilReadBinding;
+		PassParameters->View = View.ViewUniformBuffer;
+		PassParameters->SceneColorTexture = SceneColorTexture;
+		PassParameters->SceneColorTextureSampler = TStaticSamplerState<SF_Bilinear, AM_Clamp, AM_Clamp, AM_Clamp>::GetRHI();
+		PassParameters->DistortionAccumulateTexture = DistortionTexture;
+		PassParameters->DistortionAccumulateSampler = TStaticSamplerState<SF_Bilinear, AM_Clamp, AM_Clamp, AM_Clamp>::GetRHI();
+
+		AddDrawScreenPass(GraphBuilder, {}, FScreenPassViewInfo(View, 1), Viewport, Viewport, PipelineState, PassParameters,
+			[MergePixelShader, PassParameters](FRHICommandList& RHICmdList)
+			{
+				SetShaderParameters(RHICmdList, MergePixelShader, MergePixelShader.GetPixelShader(), *PassParameters);
+			});
+
+		LoadAction = ERenderTargetLoadAction::ELoad;
+	}
+
+	// Merge distortion back to scene color.
+	{
+		QUICK_SCOPE_CYCLE_COUNTER(STAT_FSceneRenderer_RenderDistortion_Merge);
+		RDG_EVENT_SCOPE(GraphBuilder, "Merge");
+			
+		TShaderMapRef<FCopyRectPS> CopyPixelShader(ShaderMap);
+
+		// Test against stencil mask and clear it.
+		PipelineState.DepthStencilState = TStaticDepthStencilState<
+			false, CF_Always,
+			true, CF_Equal, SO_Keep, SO_Keep, SO_Zero,
+			false, CF_Always, SO_Keep, SO_Keep, SO_Keep,
+			DISTORTION_STENCIL_MASK_BIT, DISTORTION_STENCIL_MASK_BIT>::GetRHI();
+		PipelineState.StencilRef = DISTORTION_STENCIL_MASK_BIT;
+		PipelineState.PixelShader = CopyPixelShader;
+
+		FCopyRectPS::FParameters* Parameters = GraphBuilder.AllocParameters<FCopyRectPS::FParameters>();
+		Parameters->InputTexture = DistortionSceneColorTexture;
+		Parameters->RenderTargets[0] = FRenderTargetBinding(SceneColorTexture, ERenderTargetLoadAction::ENoAction);
+		Parameters->RenderTargets.DepthStencil = StencilWriteBinding;
+		Parameters->InputSampler = TStaticSamplerState<SF_Bilinear>::GetRHI();
+		AddDrawScreenPass(GraphBuilder, {}, FScreenPassViewInfo(View, 1), Viewport, Viewport, PipelineState, Parameters,
+			[CopyPixelShader, Parameters](FRHICommandList& RHICmdList)
+			{
+				SetShaderParameters(RHICmdList, CopyPixelShader, CopyPixelShader.GetPixelShader(), *Parameters);
+			});
+	}
+}
\ No newline at end of file
diff --git a/Engine/Source/Runtime/Renderer/Private/MobileShadingRenderer.cpp b/Engine/Source/Runtime/Renderer/Private/MobileShadingRenderer.cpp
--- a/Engine/Source/Runtime/Renderer/Private/MobileShadingRenderer.cpp
+++ b/Engine/Source/Runtime/Renderer/Private/MobileShadingRenderer.cpp
@@ -84,6 +84,7 @@
 #include "DistanceFieldAmbientOcclusion.h"
 #include "DistanceFieldLightingShared.h"
 #include "DepthCopy.h"
+#include "DistortionRendering.h"
 #include "SingleLayerWaterRendering.h"
 #include "PostProcess/PostProcessCompositeDebugPrimitives.h"
 #include "HighResScreenshot.h"
@@ -2884,7 +2885,11 @@ void FMobileSceneRenderer::RenderDeferredMultiPass(FRDGBuilder& GraphBuilder, FS
 				});
 		}
 
-
+		if (IsMobileDistortionBeforeTranslucency())
+		{
+			RenderDistortion(GraphBuilder, View, SceneTextures.Color.Target, SceneTextures.Depth.Target);
+		}
+			
 		if (bShouldRenderSingleLayerWater 
 			&& View.bHasSingleLayerWaterMaterial
 			&& HasAnyDraw(View.ParallelMeshDrawCommandPasses[EMeshPass::SingleLayerWaterPass]))
diff --git a/Engine/Source/Runtime/Renderer/Private/PostProcess/PostProcessing.cpp b/Engine/Source/Runtime/Renderer/Private/PostProcess/PostProcessing.cpp
--- a/Engine/Source/Runtime/Renderer/Private/PostProcess/PostProcessing.cpp
+++ b/Engine/Source/Runtime/Renderer/Private/PostProcess/PostProcessing.cpp
@@ -2824,7 +2824,7 @@ void AddMobilePostProcessingPasses(FRDGBuilder& GraphBuilder, FScene* Scene, con
 		bool bUseTAA = View.AntiAliasingMethod == AAM_TemporalAA;
 		ensure(View.AntiAliasingMethod != AAM_TSR);
 
-		bool bUseDistortion = IsMobileDistortionActive(View);
+		bool bUseDistortion = IsMobileDistortionActive(View) && !IsMobileDistortionBeforeTranslucency();
 
 		bool bUseSeparateTranslucency = IsMobileSeparateTranslucencyActive(View);
 
diff --git a/Engine/Source/Runtime/Renderer/Private/SceneRendering.h b/Engine/Source/Runtime/Renderer/Private/SceneRendering.h
--- a/Engine/Source/Runtime/Renderer/Private/SceneRendering.h
+++ b/Engine/Source/Runtime/Renderer/Private/SceneRendering.h
@@ -2945,6 +2945,11 @@ protected:
 	void RenderDeferredMultiPass(FRDGBuilder& GraphBuilder, FSceneTextures& SceneTextures, const FSortedLightSetSceneInfo& SortedLightSet, FMobileBasePassTextures& BasePassTextures, FInstanceCullingManager& InstanceCullingManager);
 	void RenderDeferredMultiPassDecals(FRDGBuilder& GraphBuilder, FViewInfo& View, uint32 ViewIndex, FSceneTextures& SceneTextures, FInstanceCullingManager& InstanceCullingManager);
 
+	void RenderDistortion(FRDGBuilder& GraphBuilder,
+		const FViewInfo& View,
+		FRDGTextureRef SceneColorTexture,
+		FRDGTextureRef SceneDepthTexture);
+
 	void RenderStereoMotionVectors(FRDGBuilder& GraphBuilder, FViewInfo& View, const FSceneTextures& SceneTextures, FInstanceCullingManager& InstanceCullingManager);
 	bool HasStereoMotionVectorRenderTargets(const FSceneTextures& SceneTextures);
 
