Add missing declaration. Compile fix.

#Synced-CL 49917480

diff --git a/Engine/Source/Runtime/Renderer/Private/MobileDistortionPass.h b/Engine/Source/Runtime/Renderer/Private/MobileDistortionPass.h
--- a/Engine/Source/Runtime/Renderer/Private/MobileDistortionPass.h
+++ b/Engine/Source/Runtime/Renderer/Private/MobileDistortionPass.h
@@ -28,3 +28,6 @@ FScreenPassTexture AddMobileDistortionMergePass(FRDGBuilder& GraphBuilder, const
 
 // Returns whether distortion is enabled and there primitives to draw
 bool IsMobileDistortionActive(const FViewInfo& View);
+
+// Returns if the distortion is render before the translucent pass
+bool IsMobileDistortionBeforeTranslucency();
