Fix Cvar declaration.

#Synced-CL 49947857

diff --git a/Engine/Source/Runtime/Renderer/Private/MobileDistortionPass.cpp b/Engine/Source/Runtime/Renderer/Private/MobileDistortionPass.cpp
--- a/Engine/Source/Runtime/Renderer/Private/MobileDistortionPass.cpp
+++ b/Engine/Source/Runtime/Renderer/Private/MobileDistortionPass.cpp
@@ -23,7 +23,7 @@ IMPLEMENT_STATIC_UNIFORM_BUFFER_STRUCT(FMobileDistortionPassUniformParameters, "
 
 
 
-static TAutoConsoleVariable<int32> CVarMobileAdrenoOcclusionMode(
+static TAutoConsoleVariable<int32> CVarMobileDistortionBeforeTranslucency(
 	TEXT("r.Mobile.Distortion.BeforeTranslucency"),
 	0,
 	TEXT("0: Render distortion in the PostProcess chain (default).\n")
