Fix missing PSO precaching opt-out for light rendering.

This was dropped by mistaked when cleaning the original CL.

#Synced-CL 49992015

#Synced-CL 49992015

diff --git a/Engine/Source/Runtime/Renderer/Private/LightRendering.cpp b/Engine/Source/Runtime/Renderer/Private/LightRendering.cpp
--- a/Engine/Source/Runtime/Renderer/Private/LightRendering.cpp
+++ b/Engine/Source/Runtime/Renderer/Private/LightRendering.cpp
@@ -1178,6 +1178,10 @@ class FDeferredLightPS : public FGlobalShader
 			return EShaderPermutationPrecacheRequest::NotPrecached;
 		}
 
+		if (!PermutationVector.Get<FUbershader>())
+		{
+			return EShaderPermutationPrecacheRequest::NotPrecached;
+		}
 		return FGlobalShader::ShouldPrecachePermutation(Parameters);
 	}
 
