Implemented AnimSequenceTransformProvider. This is workflow a replacement for AnimBanks. The UAnimSequenceTransformProviderData contains a flat list of sequence references. No more indirection into a UAnimBank. The new implementation supports GPU retargeting and simple linear blending between sequences. Sequences are initially uploaded at quarter rate and the refined over time to reduce hitching. UAnimSequenceTransformProviderDataInstance contains an API for managing tracks.

#Synced-CL 49985745

diff --git a/Engine/Shaders/Private/Skinning/AnimSequenceTransformProvider.usf b/Engine/Shaders/Private/Skinning/AnimSequenceTransformProvider.usf
new file mode 100644
--- /dev/null
+++ b/Engine/Shaders/Private/Skinning/AnimSequenceTransformProvider.usf
@@ -0,0 +1,298 @@
+// Copyright Epic Games, Inc. All Rights Reserved.
+
+#include "/Engine/Public/Platform.ush"
+#include "/Engine/Private/Common.ush"
+#include "/Engine/Private/SceneData.ush"
+#include "/Engine/Private/BoneTransform.ush"
+#include "/Engine/Private/Visualization.ush"
+#include "/Engine/Private/Skinning/SkinningDataDecode.ush"
+#include "/Engine/Shared/SkinningDefinitions.h"
+
+ByteAddressBuffer RetargetingDataBuffer;
+ByteAddressBuffer SequenceTransformBuffer;
+ByteAddressBuffer EvaluateBlockHeaderBuffer;
+RWByteAddressBuffer LocalSpaceTransformBufferUAV;
+
+#define RETARGET_MODE_ANIMATION 0
+#define RETARGET_MODE_SKELETON 1
+#define RETARGET_MODE_ANIMATION_SCALED 2
+#define RETARGET_MODE_ANIMATION_RELATIVE 3
+#define RETARGET_MODE_ORIENT_AND_SCALE 4
+
+FBoneRetargetingData LoadBoneRetargetingData(ByteAddressBuffer Buffer, uint BaseOffset, uint Index)
+{
+	uint Offset = BaseOffset + Index * (uint)sizeof(FBoneRetargetingData);
+
+	return Buffer.Load<FBoneRetargetingData>(Offset);
+}
+
+FBoneTransformWithScale LoadBoneTransformWithScale(ByteAddressBuffer Buffer, uint BaseOffset, uint Index)
+{
+	uint Offset = BaseOffset + Index * (uint)sizeof(FBoneTransformWithScale);
+
+	FBoneTransformWithScale Output;
+	Output.Rotation       = Buffer.Load<float4>(Offset);
+	float4 LocationScaleX = Buffer.Load<float4>(Offset + 16);
+	float2 ScaleYZ        = Buffer.Load<float2>(Offset + 32);
+
+	Output.Location = LocationScaleX.xyz;
+	Output.Scale = float3(LocationScaleX.w, ScaleYZ);
+	return Output;
+}
+
+void StoreBoneTransformWithScale(RWByteAddressBuffer Buffer, uint BaseOffset, uint Index, FBoneTransformWithScale Transform)
+{
+	uint Offset = BaseOffset + Index * (uint)sizeof(FBoneTransformWithScale);
+
+	Buffer.Store<float4>(Offset, Transform.Rotation);
+	Buffer.Store<float4>(Offset + 16, float4(Transform.Location, Transform.Scale.x));
+	Buffer.Store<float2>(Offset + 32, Transform.Scale.yz);
+}
+
+FBoneTransformWithScale EvaluateAndRetargetBone(
+	uint MeshBoneIndex,
+	uint SequenceTransformBufferOffset,
+	uint RetargetingTransformBufferOffset,
+	uint RetargetingDataBufferOffset,
+	uint KeyIndex0,
+	uint KeyIndex1,
+	float Alpha,
+	uint NumAnimBones)
+{
+	FBoneRetargetingData RetargetingData = LoadBoneRetargetingData(RetargetingDataBuffer, RetargetingDataBufferOffset, MeshBoneIndex);
+
+	const uint RetargetingTransformIndex = MeshBoneIndex * (RetargetingData.bRequiresRefPoseRetarget ? 2 : 1);
+	const bool bIsRootBone               = MeshBoneIndex == 0;
+
+	FBoneTransformWithScale RetargetingTransform = LoadBoneTransformWithScale(SequenceTransformBuffer, RetargetingTransformBufferOffset, RetargetingTransformIndex);
+
+	if (RetargetingData.AnimBoneIndex == BONE_RETARGETING_DATA_INVALID_BONE_INDEX)
+	{
+		// When the anim bone doesn't exist the ref pose is stored in the retargeting transform.
+		return RetargetingTransform;
+	}
+
+	const uint KeyIndexA = RetargetingData.AnimBoneIndex + KeyIndex0 * NumAnimBones;
+	const uint KeyIndexB = RetargetingData.AnimBoneIndex + KeyIndex1 * NumAnimBones;
+
+	FBoneTransformWithScale MeshLocalTransformA    = LoadBoneTransformWithScale(SequenceTransformBuffer, SequenceTransformBufferOffset, KeyIndexA);
+	FBoneTransformWithScale MeshLocalTransformB    = LoadBoneTransformWithScale(SequenceTransformBuffer, SequenceTransformBufferOffset, KeyIndexB);
+	FBoneTransformWithScale MeshLocalAnimTransform = BlendBoneTransformsWithScale(MeshLocalTransformA, MeshLocalTransformB, Alpha);
+
+	if (RetargetingData.bRequiresRefPoseRetarget && !bIsRootBone)
+	{
+		FBoneTransformWithScale RefPoseData = LoadBoneTransformWithScale(SequenceTransformBuffer, RetargetingTransformBufferOffset, RetargetingTransformIndex + 1);
+
+		// Unpack Q0 from Rotation
+		FQuat Q0 = FQuat(RefPoseData.Rotation[0], RefPoseData.Rotation[1], RefPoseData.Rotation[2], RefPoseData.Rotation[3]);
+
+		// Unpack Q1 from Location.xyz + Scale.x
+		FQuat Q1 = FQuat(RefPoseData.Location.x, RefPoseData.Location.y, RefPoseData.Location.z, RefPoseData.Scale.x);
+
+		MeshLocalAnimTransform.Rotation = QuatNormalize(QuatMultiply(QuatMultiply(Q0, MeshLocalAnimTransform.Rotation), Q1));
+
+		if (RetargetingData.RetargetMode != RETARGET_MODE_SKELETON)
+		{
+			MeshLocalAnimTransform.Location = QuatRotateVector(Q0, MeshLocalAnimTransform.Location);
+		}
+	}
+
+	switch (RetargetingData.RetargetMode)
+	{
+		case RETARGET_MODE_ANIMATION:
+			break;
+
+		case RETARGET_MODE_SKELETON:
+			MeshLocalAnimTransform.Location = RetargetingTransform.Location;
+			if (bIsRootBone)
+			{
+				MeshLocalAnimTransform.Rotation = RetargetingTransform.Rotation;
+			}
+			break;
+
+		case RETARGET_MODE_ANIMATION_RELATIVE:
+			MeshLocalAnimTransform.Rotation = QuatNormalize(QuatMultiply(MeshLocalAnimTransform.Rotation, RetargetingTransform.Rotation));
+			MeshLocalAnimTransform.Location += RetargetingTransform.Location;
+			MeshLocalAnimTransform.Scale *= RetargetingTransform.Scale;
+			break;
+
+		case RETARGET_MODE_ANIMATION_SCALED:
+			MeshLocalAnimTransform.Location *= RetargetingTransform.Scale;
+			break;
+
+		case RETARGET_MODE_ORIENT_AND_SCALE:
+			MeshLocalAnimTransform.Location = QuatRotateVector(RetargetingTransform.Rotation, MeshLocalAnimTransform.Location) * RetargetingTransform.Scale.x;
+			break;
+	}
+
+	return MeshLocalAnimTransform;
+}
+
+[numthreads(BONES_PER_GROUP, 1, 1)]
+void EvaluateAnimSequenceCS(uint3 GroupId : SV_GroupID, uint GroupIndex : SV_GroupIndex)
+{
+	const uint BoneBlock          = GroupId.x;
+	const uint LocalMeshBoneIndex = GroupIndex;
+
+	const FEvaluateAnimSequenceBlockHeader BlockHeader = EvaluateBlockHeaderBuffer.Load<FEvaluateAnimSequenceBlockHeader>(BoneBlock * uint(sizeof(FEvaluateAnimSequenceBlockHeader)));
+
+	if (LocalMeshBoneIndex >= BlockHeader.BlockTransformCount)
+	{
+		return;
+	}
+
+	const uint MeshBoneIndex = LocalMeshBoneIndex + (BlockHeader.BlockLocalIndex * BONES_PER_GROUP);
+
+	FBoneTransformWithScale MeshLocalAnimTransform = EvaluateAndRetargetBone(
+		MeshBoneIndex,
+		BlockHeader.SequenceTransformBufferOffset,
+		BlockHeader.RetargetingTransformBufferOffset,
+		BlockHeader.RetargetingDataBufferOffset,
+		BlockHeader.KeyIndex0,
+		BlockHeader.KeyIndex1,
+		BlockHeader.Alpha,
+		BlockHeader.NumAnimBones);
+
+	StoreBoneTransformWithScale(LocalSpaceTransformBufferUAV, BlockHeader.DstTransformBufferOffset, MeshBoneIndex, MeshLocalAnimTransform);
+}
+
+[numthreads(BONES_PER_GROUP, 1, 1)]
+void EvaluateAnimSequenceBlendCS(uint3 GroupId : SV_GroupID, uint GroupIndex : SV_GroupIndex)
+{
+	const uint BoneBlock          = GroupId.x;
+	const uint LocalMeshBoneIndex = GroupIndex;
+
+	const FEvaluateAnimSequenceBlendBlockHeader BlockHeader = EvaluateBlockHeaderBuffer.Load<FEvaluateAnimSequenceBlendBlockHeader>(BoneBlock * uint(sizeof(FEvaluateAnimSequenceBlendBlockHeader)));
+
+	if (LocalMeshBoneIndex >= BlockHeader.BlockTransformCount)
+	{
+		return;
+	}
+
+	const uint MeshBoneIndex = LocalMeshBoneIndex + (BlockHeader.BlockLocalIndex * BONES_PER_GROUP);
+
+	FBoneTransformWithScale SourceMeshLocalAnimTransform = EvaluateAndRetargetBone(
+		MeshBoneIndex,
+		BlockHeader.SourceSequenceTransformBufferOffset,
+		BlockHeader.SourceRetargetingTransformBufferOffset,
+		BlockHeader.SourceRetargetingDataBufferOffset,
+		BlockHeader.SourceKeyIndex0,
+		BlockHeader.SourceKeyIndex1,
+		BlockHeader.SourceAlpha,
+		BlockHeader.SourceNumAnimBones);
+
+	FBoneTransformWithScale TargetMeshLocalAnimTransform = EvaluateAndRetargetBone(
+		MeshBoneIndex,
+		BlockHeader.TargetSequenceTransformBufferOffset,
+		BlockHeader.TargetRetargetingTransformBufferOffset,
+		BlockHeader.TargetRetargetingDataBufferOffset,
+		BlockHeader.TargetKeyIndex0,
+		BlockHeader.TargetKeyIndex1,
+		BlockHeader.TargetAlpha,
+		BlockHeader.TargetNumAnimBones);
+
+	FBoneTransformWithScale MeshLocalAnimTransform = BlendBoneTransformsWithScale(SourceMeshLocalAnimTransform, TargetMeshLocalAnimTransform, BlockHeader.BlendWeight);
+
+	StoreBoneTransformWithScale(LocalSpaceTransformBufferUAV, BlockHeader.DstTransformBufferOffset, MeshBoneIndex, MeshLocalAnimTransform);
+}
+
+ByteAddressBuffer LocalSpaceTransformBuffer;
+ByteAddressBuffer ConcatenateBlockHeaderBuffer;
+ByteAddressBuffer LeafBonePathBuffer;
+RWByteAddressBuffer GlobalSpaceTransformBufferUAV;
+
+groupshared FConcatenateBlockHeader SharedConcatenateBlockHeader;
+
+[numthreads(BONES_PER_GROUP, 1, 1)]
+void ConcatenateLocalToGlobalCS(uint3 GroupId : SV_GroupID, uint GroupIndex : SV_GroupIndex)
+{
+	const uint BlockIndex         = GroupId.x;
+	const uint LocalLeafBoneIndex = GroupIndex;
+	const uint PathElementSize    = 4;
+
+	FConcatenateBlockHeader BlockHeader = ConcatenateBlockHeaderBuffer.Load<FConcatenateBlockHeader>(GroupId.x * uint(sizeof(FConcatenateBlockHeader)));
+
+	if (LocalLeafBoneIndex >= BlockHeader.BlockLeafBoneCount)
+	{
+		return;
+	}
+
+	const uint LeafBoneIndex = LocalLeafBoneIndex + (BlockHeader.BlockLocalIndex * BONES_PER_GROUP);
+
+	// Load relative offset from offset table, then add base to get absolute offset
+	const uint RelativePathHeaderOffset = LeafBonePathBuffer.Load<uint>(BlockHeader.LeafBonePathBufferOffset + LeafBoneIndex * PathElementSize);
+	const uint PathHeaderOffset = BlockHeader.LeafBonePathBufferOffset + RelativePathHeaderOffset;
+
+	FLeafBonePathHeader PathHeader = LeafBonePathBuffer.Load<FLeafBonePathHeader>(PathHeaderOffset);
+	FBoneTransformWithScale ParentTransform = LoadBoneTransformWithScale(LocalSpaceTransformBuffer, BlockHeader.TransformBufferOffset, 0);
+
+	if (PathHeader.bIsRootOwned)
+	{
+		StoreBoneTransformWithScale(GlobalSpaceTransformBufferUAV, BlockHeader.TransformBufferOffset, 0, ParentTransform);
+	}
+
+	if (PathHeader.PathDepth > 0)
+	{
+		uint PathNodeOffset = PathHeaderOffset + PathElementSize;
+		FLeafBonePathNode PathNode;
+
+		for (uint PathIndex = 0; PathIndex < PathHeader.PathDepth; ++PathIndex)
+		{
+			PathNode = LeafBonePathBuffer.Load<FLeafBonePathNode>(PathNodeOffset);
+			PathNodeOffset += PathElementSize;
+
+			FBoneTransformWithScale ChildTransform = LoadBoneTransformWithScale(LocalSpaceTransformBuffer, BlockHeader.TransformBufferOffset, PathNode.BoneIndex);
+			ParentTransform = MultiplyBoneTransformWithScale(ChildTransform, ParentTransform);
+
+			if (PathNode.bIsOwned)
+			{
+				StoreBoneTransformWithScale(GlobalSpaceTransformBufferUAV, BlockHeader.TransformBufferOffset, PathNode.BoneIndex, ParentTransform);
+			}
+
+			if (PathNode.bIsLeaf)
+			{
+				break;
+			}
+		}
+	}
+}
+
+ByteAddressBuffer ScatterBlockHeaderBuffer;
+ByteAddressBuffer GlobalSpaceTransformBuffer;
+ByteAddressBuffer BoneMapBuffer;
+FCompressedBoneTransformRWBuffer SkinningTransformBufferUAV;
+
+[numthreads(BONES_PER_GROUP, 1, 1)]
+void ScatterTransformsCS(uint3 GroupId : SV_GroupID, uint GroupIndex : SV_GroupIndex)
+{
+	const uint BlockHeaderIndex  = GroupId.x;
+	const uint LocalTransformIndex = GroupIndex;
+
+	const FScatterTransformsBlockHeader BlockHeader = ScatterBlockHeaderBuffer.Load<FScatterTransformsBlockHeader>(BlockHeaderIndex * uint(sizeof(FScatterTransformsBlockHeader)));
+
+	if (LocalTransformIndex >= BlockHeader.BlockTransformCount)
+	{
+		return;
+	}
+
+	const uint TransformIndex = LocalTransformIndex + (BlockHeader.BlockLocalIndex * BONES_PER_GROUP);
+
+	uint MeshBoneIndex = TransformIndex;
+
+	if (BlockHeader.BoneMapOffset != ~0u)
+	{
+		MeshBoneIndex = BoneMapBuffer.Load<uint>(BlockHeader.BoneMapOffset + TransformIndex * uint(sizeof(uint)));
+	}
+
+	FBoneTransformWithScale InvGlobalRefPoseTransform = LoadBoneTransformWithScale(SequenceTransformBuffer, BlockHeader.InvGlobalRefPoseTransformBufferOffset, MeshBoneIndex);
+	FBoneTransformWithScale GlobalSpaceTransform = LoadBoneTransformWithScale(GlobalSpaceTransformBuffer, BlockHeader.SrcTransformBufferOffset, MeshBoneIndex);
+	GlobalSpaceTransform = MultiplyBoneTransformWithScale(InvGlobalRefPoseTransform, GlobalSpaceTransform);
+	float4x3 GlobalSpaceMatrix = (float4x3)MatrixFromBoneTransformWithScale(GlobalSpaceTransform);
+
+	StoreCompressedBoneTransform(SkinningTransformBufferUAV, BlockHeader.DstTransformBufferOffset1, TransformIndex, GlobalSpaceMatrix);
+
+	if (BlockHeader.DstTransformBufferOffset1 != BlockHeader.DstTransformBufferOffset2)
+	{
+		StoreCompressedBoneTransform(SkinningTransformBufferUAV, BlockHeader.DstTransformBufferOffset2, TransformIndex, GlobalSpaceMatrix);
+	}
+}
\ No newline at end of file
diff --git a/Engine/Shaders/Shared/SkinningDefinitions.h b/Engine/Shaders/Shared/SkinningDefinitions.h
--- a/Engine/Shaders/Shared/SkinningDefinitions.h
+++ b/Engine/Shaders/Shared/SkinningDefinitions.h
@@ -10,6 +10,12 @@
 // This path is experimental and currently adds cost to the shader.
 #define GPUSKIN_TRANSFORM_STORAGE_DIRECT_ENABLED 0
 
+#ifdef __cplusplus
+#define SKINNING_STATIC_ASSERT(bCondition, Message) static_assert(bCondition, Message)
+#else
+#define SKINNING_STATIC_ASSERT(bCondition, Message)
+#endif
+
 #if defined(__cplusplus) || COMPILER_SUPPORTS_HLSL2021
 
 #ifdef __cplusplus
@@ -34,6 +40,89 @@ struct FAnimBankBlockHeader
 	uint BlockTransformCount            : 16;
 };
 
+#define BONE_RETARGETING_DATA_INVALID_BONE_INDEX 0xFFFF
+
+struct FBoneRetargetingData
+{
+	uint AnimBoneIndex : 16;
+	uint RetargetMode  : 15;
+	uint bRequiresRefPoseRetarget : 1;
+};
+SKINNING_STATIC_ASSERT(sizeof(FBoneRetargetingData) == 4, "FBoneRetargetingData must be 32 bytes");
+
+struct FEvaluateAnimSequenceBlockHeader
+{
+	float  Alpha;
+	uint KeyIndex0;
+	uint KeyIndex1;
+	uint SequenceTransformBufferOffset;
+	uint RetargetingTransformBufferOffset;
+	uint RetargetingDataBufferOffset;
+	uint DstTransformBufferOffset;
+	uint NumAnimBones;
+	uint BlockLocalIndex      : 16;
+	uint BlockTransformCount  : 16;
+};
+SKINNING_STATIC_ASSERT(sizeof(FEvaluateAnimSequenceBlockHeader) == 36, "FEvaluateAnimSequenceBlockHeader must be 36 bytes");
+
+struct FEvaluateAnimSequenceBlendBlockHeader
+ {
+	float SourceAlpha;
+	float TargetAlpha;
+	float BlendWeight;
+	uint SourceKeyIndex0;
+	uint SourceKeyIndex1;
+	uint TargetKeyIndex0;
+	uint TargetKeyIndex1;
+	uint SourceSequenceTransformBufferOffset;
+	uint TargetSequenceTransformBufferOffset;
+	uint SourceRetargetingTransformBufferOffset;
+	uint TargetRetargetingTransformBufferOffset;
+	uint SourceRetargetingDataBufferOffset;
+	uint TargetRetargetingDataBufferOffset;
+	uint DstTransformBufferOffset;
+	uint SourceNumAnimBones   : 16;
+	uint TargetNumAnimBones   : 16;
+	uint BlockLocalIndex      : 16;
+	uint BlockTransformCount  : 16;
+ };
+SKINNING_STATIC_ASSERT(sizeof(FEvaluateAnimSequenceBlendBlockHeader) == 64, "FEvaluateAnimSequenceBlendBlockHeader must be 64 bytes");
+
+struct FLeafBonePathHeader
+{
+	uint PathDepth      : 31;
+	uint bIsRootOwned   : 1;
+};
+SKINNING_STATIC_ASSERT(sizeof(FLeafBonePathHeader) == 4, "FLeafBonePathHeader must be 4 bytes");
+
+struct FLeafBonePathNode
+{
+	uint BoneIndex : 30;
+	uint bIsOwned  : 1;
+	uint bIsLeaf   : 1;
+};
+SKINNING_STATIC_ASSERT(sizeof(FLeafBonePathNode) == 4, "FLeafBonePathNode must be 4 bytes");
+struct FConcatenateBlockHeader
+{
+	uint LeafBonePathBufferOffset;
+	uint TransformBufferOffset;
+	uint BlockLocalIndex    : 16;
+	uint BlockLeafBoneCount : 16;
+};
+SKINNING_STATIC_ASSERT(sizeof(FConcatenateBlockHeader) == 12, "FConcatenateBlockHeader must be 12 bytes");
+
+struct FScatterTransformsBlockHeader
+{
+	uint InvGlobalRefPoseTransformBufferOffset;
+	uint SrcTransformBufferOffset;
+	uint DstTransformBufferOffset1;
+	uint DstTransformBufferOffset2;
+	uint BoneMapOffset;
+	uint BlockLocalIndex     : 16;
+	uint BlockTransformCount : 16;
+};
+SKINNING_STATIC_ASSERT(sizeof(FScatterTransformsBlockHeader) == 24, "FScatterTransformsBlockHeader must be 24 bytes");
+
 #define ANIM_BANK_FLAG_NONE			0x0
 #define ANIM_BANK_FLAG_LOOPING		0x1
 #define ANIM_BANK_FLAG_AUTOSTART	0x2
@@ -124,29 +213,36 @@ struct FCompressedBoneTransform
 #define REF_POSE_TRANSFORM_PROVIDER_GUID 0x665207E7, 0x449A4FB1, 0xA298F7AD, 0x8F989B11
 #define ANIM_BANK_GPU_TRANSFORM_PROVIDER_GUID 0xA5C0027A, 0x8F884C7C, 0x9312F138, 0x71A9300F
 #define ANIM_BANK_CPU_TRANSFORM_PROVIDER_GUID 0xE7D6173D, 0x246F431A, 0x912D384E, 0x156C0D2C
+#define ANIM_SEQUENCE_GPU_TRANSFORM_PROVIDER_GUID 0x5CFA81C1, 0xF3D8486B, 0x81FF23D3, 0x72EB523F
+#define ANIM_SEQUENCE_CPU_TRANSFORM_PROVIDER_GUID 0xBA5349A0, 0xCDFC4F76, 0x987168CE, 0xC7AA0E7B
 #define ANIM_RUNTIME_TRANSFORM_PROVIDER_GUID 0xF1508490, 0xFCC24BB9, 0xA9F277B3, 0x1AF766F0
 #define ANIM_MESH_OBJECT_TRANSFORM_PROVIDER_GUID 0x1073C0C5, 0xB981F6FD, 0xBF41B1C2, 0xC4BA0C50
 
-inline void StoreCompressedBoneTransform(FCompressedBoneTransform* CompressedTransform, const FMatrix44f& Transform)
+inline void StoreCompressedBoneTransform(FCompressedBoneTransform& CompressedTransform, const FMatrix44f& Transform)
 {
 #if USE_COMPRESSED_BONE_TRANSFORM
 	//TODO: Optimize
-	*(FVector3f*)&CompressedTransform->Data[0] = Transform.GetOrigin();
+	*(FVector3f*)&CompressedTransform.Data[0] = Transform.GetOrigin();
 
 	float Tmp[8] = { Transform.M[0][0], Transform.M[0][1], Transform.M[0][2], Transform.M[1][0],
 					 Transform.M[1][1], Transform.M[1][2], Transform.M[2][0], Transform.M[2][1] };
 
-	uint16* Ptr = (uint16*)&CompressedTransform->Data[3];
+	uint16* Ptr = (uint16*)&CompressedTransform.Data[3];
 
 	FPlatformMath::VectorStoreHalf(&Ptr[0], &Tmp[0]);
 	FPlatformMath::VectorStoreHalf(&Ptr[4], &Tmp[4]);
 	FPlatformMath::StoreHalf(&Ptr[8], Transform.M[2][2]);
 	Ptr[9] = 0;
 #else
-	Transform.To3x4MatrixTranspose((float*)CompressedTransform);
+	Transform.To3x4MatrixTranspose((float*)&CompressedTransform);
 #endif
 }
 
+inline void StoreCompressedBoneTransform(FCompressedBoneTransform* CompressedTransform, const FMatrix44f& Transform)
+{
+	StoreCompressedBoneTransform(*CompressedTransform, Transform);
+}
+
 inline void StoreCompressedBoneTransform(FCompressedBoneTransform& CompressedTransform, const FMatrix3x4& Transform)
 {
 #if USE_COMPRESSED_BONE_TRANSFORM
diff --git a/Engine/Source/Editor/UnrealEd/Private/Factories/TransformProviderFactory.cpp b/Engine/Source/Editor/UnrealEd/Private/Factories/TransformProviderFactory.cpp
--- a/Engine/Source/Editor/UnrealEd/Private/Factories/TransformProviderFactory.cpp
+++ b/Engine/Source/Editor/UnrealEd/Private/Factories/TransformProviderFactory.cpp
@@ -8,6 +8,8 @@
 #include "Styling/AppStyle.h"
 #include "Animation/Skeleton.h"
 #include "Animation/TransformProviderData.h"
+#include "Animation/AnimSequenceTransformProviderData.h"
+#include "Animation/AnimRuntimeTransformProviderData.h"
 #include "Editor.h"
 
 #include "ClassViewerModule.h"
@@ -47,6 +49,9 @@ public:
 	/** All children of these classes will be included unless filtered out by another setting. */
 	TSet< const UClass* > AllowedChildrenOfClasses;
 
+	/** Classes that are explicitly disallowed from being selected. */
+	TSet<const UClass*> DisallowedClasses;
+
 	/** Disallowed class flags. */
 	EClassFlags DisallowedClassFlags;
 
@@ -56,6 +61,7 @@ public:
 	virtual bool IsClassAllowed(const FClassViewerInitializationOptions& InInitOptions, const UClass* InClass, TSharedRef<FClassViewerFilterFuncs> InFilterFuncs) override
 	{
 		const bool bAllowed = !InClass->HasAnyClassFlags(DisallowedClassFlags)
+			&& !DisallowedClasses.Contains(InClass)
 			&& InClass->CanCreateAssetOfClass()
 			&& InFilterFuncs->IfInChildOfClassesSet(AllowedChildrenOfClasses, InClass) != EFilterReturn::Failed;
 
@@ -99,6 +105,8 @@ bool UTransformProviderDataFactory::ConfigureProperties()
 
 	Filter->DisallowedClassFlags = CLASS_Abstract | CLASS_Deprecated | CLASS_NewerVersionExists;
 	Filter->AllowedChildrenOfClasses.Add(UTransformProviderData::StaticClass());
+	Filter->DisallowedClasses.Add(UAnimSequenceTransformProviderDataInstance::StaticClass());
+	Filter->DisallowedClasses.Add(UAnimRuntimeTransformProviderData::StaticClass());
 
 	const FText TitleText = LOCTEXT("CreateTransformProviderDataOptions", "Pick Transform Provider Data Class");
 	UClass* ChosenClass = nullptr;
diff --git a/Engine/Source/Runtime/Engine/Classes/Animation/AnimBank.h b/Engine/Source/Runtime/Engine/Classes/Animation/AnimBank.h
--- a/Engine/Source/Runtime/Engine/Classes/Animation/AnimBank.h
+++ b/Engine/Source/Runtime/Engine/Classes/Animation/AnimBank.h
@@ -14,7 +14,7 @@
 #include "UObject/ObjectPtr.h"
 #include "UObject/WeakObjectPtr.h"
 #include "AnimTrackPool.h"
-
+#include "AnimSequenceTransformProviderData.h"
 #include "AnimBank.generated.h"
 
 ENGINE_API DECLARE_LOG_CATEGORY_EXTERN(LogAnimBank, Log, All);
@@ -550,63 +550,6 @@ public:
 	TSoftObjectPtr<class UTransformProviderData> TransformProvider;
 };
 
-UENUM(BlueprintType)
-enum class EAnimBankTrackLoopMode : uint8
-{
-	// The behavior is unchanged from the sequence default.
-	Auto,
-
-	// Looping is forced for the current track.
-	Loop,
-
-	// Clamp is forced for the current track.
-	Clamp
-};
-
-USTRUCT(BlueprintType)
-struct FAnimBankTrackAutoPlayData
-{
-	GENERATED_USTRUCT_BODY()
-
-	// The anim bank item index to play.
-	UPROPERTY(EditAnywhere, BlueprintReadWrite, Category = Animation)
-	int32 ItemIndex = 0;
-
-	// The starting time position to begin play.
-	UPROPERTY(EditAnywhere, BlueprintReadWrite, Category = Animation)
-	float Position = 0.0f;
-
-	// The rate multiplier to play the animation.
-	UPROPERTY(EditAnywhere, BlueprintReadWrite, Category = Animation)
-	float PlayRate = 1.0f;
-
-	// The mode to control looping behavior.
-	UPROPERTY(EditAnywhere, BlueprintReadWrite, Category = Animation)
-	EAnimBankTrackLoopMode LoopMode = EAnimBankTrackLoopMode::Auto;
-};
-
-USTRUCT(BlueprintType)
-struct FAnimBankTrackManualData
-{
-	GENERATED_USTRUCT_BODY()
-
-	// The anim bank item index to play.
-	UPROPERTY(EditAnywhere, BlueprintReadWrite, Category = Animation)
-	int32 ItemIndex = 0;
-
-	// The time position so sample in the sequence.
-	UPROPERTY(EditAnywhere, BlueprintReadWrite, Category = Animation)
-	float Position = 0.0f;
-
-	// Whether to apply the previous position when updating the animation. When false, the previous position is derived from the last frame.
-	UPROPERTY(EditAnywhere, BlueprintReadWrite, Category = Animation)
-	uint32 bHasPreviousPosition : 1 = 0;
-
-	// A time position to update previous frame transforms.
-	UPROPERTY(EditAnywhere, BlueprintReadWrite, Category = Animation, meta = (EditCondition = "bHasPreviousPosition", EditConditionHides))
-	float PreviousPosition = 0.0f;
-};
-
 class FAnimBankTransformProviderProxy;
 class UAnimBankDataInstance;
 
@@ -616,10 +559,10 @@ public:
 	/////////////////////////////////////////////////////////////////////////////////////////////////
 	// AutoPlay
 
-	static FAnimBankTrackPackedData Pack(FAnimBankTrackAutoPlayData AutoPlayData)
+	static FAnimBankTrackPackedData Pack(FAnimSequenceTrackAutoPlayData AutoPlayData)
 	{
 		FAnimBankTrackPackedData Data;
-		Data.ItemIndex            = AutoPlayData.ItemIndex;
+		Data.ItemIndex            = AutoPlayData.SequenceIndex;
 		Data.bAutoPlay            = 1;
 		Data.LoopMode             = static_cast<uint32>(AutoPlayData.LoopMode);
 		Data.bHasPreviousPosition = true;
@@ -634,25 +577,25 @@ public:
 		return bAutoPlay;
 	}
 
-	FAnimBankTrackAutoPlayData UnpackAutoPlay() const
+	FAnimSequenceTrackAutoPlayData UnpackAutoPlay() const
 	{
 		check(IsAutoPlay());
 
-		FAnimBankTrackAutoPlayData Data;
-		Data.ItemIndex = ItemIndex;
+		FAnimSequenceTrackAutoPlayData Data;
+		Data.SequenceIndex = ItemIndex;
 		Data.Position  = Position;
 		Data.PlayRate  = PlayRate;
-		Data.LoopMode  = static_cast<EAnimBankTrackLoopMode>(LoopMode);
+		Data.LoopMode  = static_cast<EAnimSequenceTrackLoopMode>(LoopMode);
 		return Data;
 	}
 
 	/////////////////////////////////////////////////////////////////////////////////////////////////
 	// Manual
 
-	static FAnimBankTrackPackedData Pack(FAnimBankTrackManualData ManualData)
+	static FAnimBankTrackPackedData Pack(FAnimSequenceTrackManualData ManualData)
 	{
 		FAnimBankTrackPackedData Data;
-		Data.ItemIndex            = ManualData.ItemIndex;
+		Data.ItemIndex            = ManualData.SequenceIndex;
 		Data.bAutoPlay            = 0;
 		Data.LoopMode             = 0;
 		Data.bHasPreviousPosition = ManualData.bHasPreviousPosition;
@@ -667,12 +610,12 @@ public:
 		return !bAutoPlay;
 	}
 
-	FAnimBankTrackManualData UnpackManual() const
+	FAnimSequenceTrackManualData UnpackManual() const
 	{
 		check(IsManual());
 
-		FAnimBankTrackManualData Data;
-		Data.ItemIndex            = ItemIndex;
+		FAnimSequenceTrackManualData Data;
+		Data.SequenceIndex        = ItemIndex;
 		Data.Position             = Position;
 		Data.PreviousPosition     = PreviousPosition;
 		Data.bHasPreviousPosition = bHasPreviousPosition;
@@ -701,9 +644,9 @@ public:
 		return bHasPreviousPosition;
 	}
 
-	EAnimBankTrackLoopMode GetLoopMode() const
+	EAnimSequenceTrackLoopMode GetLoopMode() const
 	{
-		return static_cast<EAnimBankTrackLoopMode>(LoopMode);
+		return static_cast<EAnimSequenceTrackLoopMode>(LoopMode);
 	}
 
 	ENGINE_API void Update(float SequenceLength, bool bLooping, float DeltaTime);
@@ -829,72 +772,6 @@ protected:
 	bool IsAnimBankValid(FInstancedSkinningSceneExtensionProxy* ExtensionProxy) const;
 };
 
-UCLASS(config=Engine, hidecategories=Object, MinimalAPI, BlueprintType, Within=InstancedSkinnedMeshComponent)
-class UAnimBankDataInstance : public UAnimBankData
-{
-	GENERATED_BODY()
-
-public:
-	virtual bool IsEnabled() const override;
-	virtual uint32 GetUniqueAnimationCount() const override;
-	virtual bool GetAnimationBounds(uint32 AnimationIndex, FRenderBounds& OutBounds) const override;
-	virtual uint32 GetSkinningDataOffset(int32 InstanceIndex, const FTransform& ComponentTransform, const FSkinnedMeshInstanceData& InstanceData) const override;
-	virtual FTransformProviderRenderProxy* CreateRenderProxy(FInstancedSkinningSceneExtensionProxy* ExtensionProxy) const override;
-	virtual void BeginDestroy() override;
-	virtual void Serialize(FArchive& Ar) override;
-
-public:
-	// Copies items from this anim bank to create a new anim bank instance.
-	UFUNCTION(BlueprintCallable, Category = "TransformProviders|AnimBank")
-	static ENGINE_API UAnimBankDataInstance* CreateAnimBankDataInstance(UAnimBankData* AnimBank, UInstancedSkinnedMeshComponent* Owner);
-
-	// Allocates a new track with auto play data, or -1 if the track allocation failed.
-	UFUNCTION(BlueprintCallable, Category = "TransformProviders|AnimBank")
-	ENGINE_API int32 AllocateAutoPlayTrack(const FAnimBankTrackAutoPlayData& Data);
-
-	// Allocates a new track with manual play data, or -1 if the track allocation failed.
-	UFUNCTION(BlueprintCallable, Category = "TransformProviders|AnimBank")
-	ENGINE_API int32 AllocateManualTrack(const FAnimBankTrackManualData& Data);
-
-	// Updates a track with auto play data. Returns whether the update succeeded.
-	UFUNCTION(BlueprintCallable, Category = "TransformProviders|AnimBank")
-	ENGINE_API bool UpdateAutoPlayTrack(int32 TrackIndex, const FAnimBankTrackAutoPlayData& Data);
-
-	// Updates a track with manual data. Returns whether the update succeeded.
-	UFUNCTION(BlueprintCallable, Category = "TransformProviders|AnimBank")
-	ENGINE_API bool UpdateManualTrack(int32 TrackIndex, const FAnimBankTrackManualData& Data);
-
-	// Deallocates a track. Returns whether the deallocation succeeded.
-	UFUNCTION(BlueprintCallable, Category = "TransformProviders|AnimBank")
-	ENGINE_API bool DeallocateTrack(int32 TrackIndex);
-
-	// Submits all pending changes to rendering.
-	UFUNCTION(BlueprintCallable, Category = "TransformProviders|AnimBank")
-	ENGINE_API virtual void SubmitChanges() override;
-
-private:
-	void InitRenderData();
-	void InitSequences();
-
-	struct FSequence
-	{
-		FSequence(float InPlayLength, bool bInLooping)
-			: PlayLength(InPlayLength)
-			, bLooping(bInLooping)
-		{}
-
-		float PlayLength;
-		bool bLooping;
-	};
-
-	FAnimBankTrackPool Tracks;
-	TArray<FSequence> Sequences;
-	TUniquePtr<FAnimBankRenderData> RenderData{ new FAnimBankRenderData };
-	UE::FMutex SubmitChangesMutex;
-
-	friend class UAnimBankData;
-};
-
 namespace UE::AnimBank
 {
 
diff --git a/Engine/Source/Runtime/Engine/Classes/Animation/AnimRuntimeTransformProviderData.h b/Engine/Source/Runtime/Engine/Classes/Animation/AnimRuntimeTransformProviderData.h
--- a/Engine/Source/Runtime/Engine/Classes/Animation/AnimRuntimeTransformProviderData.h
+++ b/Engine/Source/Runtime/Engine/Classes/Animation/AnimRuntimeTransformProviderData.h
@@ -83,10 +83,10 @@ inline FArchive& operator <<(FArchive& Ar, FAnimRuntimeTrackData& Data)
 
 using FAnimRuntimeTrackPool = TAnimTrackPool<FAnimRuntimeTrackData>;
 
-class FAnimRuntimeRenderData
+class FAnimRuntimeTransformProviderRenderData
 {
 public:
-	ENGINE_API ~FAnimRuntimeRenderData();
+	ENGINE_API ~FAnimRuntimeTransformProviderRenderData();
 
 private:
 	void Patch(FAnimRuntimeTrackPool::FPatch&& Patch);
@@ -117,7 +117,7 @@ struct FAnimRuntimeTrackAllocation
 
 class FAnimRuntimeTransformProviderProxy : public FTransformProviderRenderProxy
 {
-	FAnimRuntimeTransformProviderProxy(FInstancedSkinningSceneExtensionProxy* InSceneExtensionProxy, FAnimRuntimeRenderData& InRenderData)
+	FAnimRuntimeTransformProviderProxy(FInstancedSkinningSceneExtensionProxy* InSceneExtensionProxy, FAnimRuntimeTransformProviderRenderData& InRenderData)
 		: SceneExtensionProxy(InSceneExtensionProxy)
 		, RenderData(InRenderData)
 		, Tracks(InRenderData.Tracks)
@@ -139,7 +139,7 @@ public:
 
 private:
 	FInstancedSkinningSceneExtensionProxy* SceneExtensionProxy = nullptr;
-	FAnimRuntimeRenderData& RenderData;
+	FAnimRuntimeTransformProviderRenderData& RenderData;
 	FAnimRuntimeTrackPool& Tracks;
 	IAnimRuntimeTransformProvider* AnimRuntimeTransformProvider = nullptr;
 
@@ -207,7 +207,7 @@ private:
 	int32 NumTransforms = 0;
 	int32 NumPoolInstances = 0;
 	FAnimRuntimeTrackPool Tracks;
-	TUniquePtr<FAnimRuntimeRenderData> RenderData{ new FAnimRuntimeRenderData };
+	TUniquePtr<FAnimRuntimeTransformProviderRenderData> RenderData{ new FAnimRuntimeTransformProviderRenderData };
 	TIntrusiveDoubleLinkedList<FAnimRuntimeTrackTransformData> FreeList;
 	UE::FMutex SubmitChangesMutex;
 };
\ No newline at end of file
diff --git a/Engine/Source/Runtime/Engine/Classes/Animation/AnimSequenceTransformProviderData.h b/Engine/Source/Runtime/Engine/Classes/Animation/AnimSequenceTransformProviderData.h
new file mode 100644
--- /dev/null
+++ b/Engine/Source/Runtime/Engine/Classes/Animation/AnimSequenceTransformProviderData.h
@@ -0,0 +1,510 @@
+// Copyright Epic Games, Inc. All Rights Reserved.
+
+#pragma once
+
+#include "Animation/TransformProviderData.h"
+#include "Experimental/Containers/RobinHoodHashTable.h"
+#include "RenderCommandFence.h"
+#include "SkinningDefinitions.h"
+#include "UObject/ObjectPtr.h"
+#include "UObject/WeakObjectPtr.h"
+#include "AnimTrackPool.h"
+
+#include "AnimSequenceTransformProviderData.generated.h"
+
+class USkinnedAsset;
+class UAnimBankData;
+class UAnimSequence;
+class UAnimSequenceTransformProviderData;
+class FAnimSequenceTransformProviderProxy;
+class IAnimSequenceTransformProvider;
+
+UENUM(BlueprintType)
+enum class EAnimSequenceTrackLoopMode : uint8
+{
+	// Looping is forced for the current track.
+	Loop,
+
+	// Clamp is forced for the current track.
+	Clamp
+};
+
+USTRUCT(BlueprintType)
+struct FAnimSequenceTrackAutoPlayData
+{
+	GENERATED_USTRUCT_BODY()
+
+	// The sequence index to play.
+	UPROPERTY(EditAnywhere, BlueprintReadWrite, Category = Animation)
+	int32 SequenceIndex = 0;
+
+	// The starting time position to begin play.
+	UPROPERTY(EditAnywhere, BlueprintReadWrite, Category = Animation)
+	float Position = 0.0f;
+
+	// The rate multiplier to play the animation.
+	UPROPERTY(EditAnywhere, BlueprintReadWrite, Category = Animation)
+	float PlayRate = 1.0f;
+
+	// The mode to control looping behavior.
+	UPROPERTY(EditAnywhere, BlueprintReadWrite, Category = Animation)
+	EAnimSequenceTrackLoopMode LoopMode = EAnimSequenceTrackLoopMode::Loop;
+
+	// The duration of the blend transition in seconds. Use 0 for an instance switch.
+	UPROPERTY(EditAnywhere, BlueprintReadWrite, Category = Animation, meta = (ClampMin = "0.0"))
+	float BlendTime = 0.0f;
+};
+
+USTRUCT(BlueprintType)
+struct FAnimSequenceTrackManualData
+{
+	GENERATED_USTRUCT_BODY()
+
+	// The sequence index to play.
+	UPROPERTY(EditAnywhere, BlueprintReadWrite, Category = Animation)
+	int32 SequenceIndex = 0;
+
+	// The time position so sample in the sequence.
+	UPROPERTY(EditAnywhere, BlueprintReadWrite, Category = Animation)
+	float Position = 0.0f;
+
+	// Whether to apply the previous position when updating the animation. When false, the previous position is derived from the last frame.
+	UPROPERTY(EditAnywhere, BlueprintReadWrite, Category = Animation)
+	uint32 bHasPreviousPosition : 1 = 0;
+
+	// A time position to update previous frame transforms.
+	UPROPERTY(EditAnywhere, BlueprintReadWrite, Category = Animation, meta = (EditCondition = "bHasPreviousPosition", EditConditionHides))
+	float PreviousPosition = 0.0f;
+};
+
+USTRUCT(BlueprintType)
+struct FAnimSequenceTransformProviderSequence
+{
+	GENERATED_USTRUCT_BODY()
+
+	UPROPERTY(EditAnywhere, BlueprintReadWrite, Category = Animation)
+	TObjectPtr<class UAnimSequence> Sequence;
+
+	// The rate multiplier to play the animation.
+	UPROPERTY(EditAnywhere, BlueprintReadWrite, Category = Animation)
+	float PlayRate = 1.0f;
+
+	// The mode to control looping behavior.
+	UPROPERTY(EditAnywhere, BlueprintReadWrite, Category = Animation)
+	EAnimSequenceTrackLoopMode LoopMode = EAnimSequenceTrackLoopMode::Loop;
+};
+
+class FAnimSequenceTransformProviderProxy;
+class UAnimSequenceTransformProviderDataInstance;
+
+enum class EAnimSequenceTrackPatchFlags : uint8
+{
+	None         = 0,
+	PlayRate     = 1 << 0,
+	LoopMode     = 1 << 1
+};
+ENUM_CLASS_FLAGS(EAnimSequenceTrackPatchFlags);
+
+class FAnimSequenceTrackPackedData
+{
+public:
+	/////////////////////////////////////////////////////////////////////////////////////////////////
+	// AutoPlay
+
+	static FAnimSequenceTrackPackedData Pack(FAnimSequenceTrackAutoPlayData AutoPlayData)
+	{
+		FAnimSequenceTrackPackedData Data;
+		Data.SequenceIndex = AutoPlayData.SequenceIndex;
+		Data.bAutoPlay = 1;
+		Data.LoopMode = static_cast<uint32>(AutoPlayData.LoopMode);
+		Data.PatchFlags = 0;
+		Data.bHasPreviousPosition = true;
+		Data.Position = AutoPlayData.Position;
+		Data.PreviousPosition = AutoPlayData.Position;
+		Data.PlayRate = AutoPlayData.PlayRate;
+		Data.BlendTime = AutoPlayData.BlendTime;
+		return Data;
+	}
+
+	bool IsAutoPlay() const
+	{
+		return bAutoPlay;
+	}
+
+	FAnimSequenceTrackAutoPlayData UnpackAutoPlay() const
+	{
+		check(IsAutoPlay());
+
+		FAnimSequenceTrackAutoPlayData Data;
+		Data.SequenceIndex = SequenceIndex;
+		Data.Position = Position;
+		Data.PlayRate = PlayRate;
+		Data.BlendTime = BlendTime;
+		Data.LoopMode = static_cast<EAnimSequenceTrackLoopMode>(LoopMode);
+		return Data;
+	}
+
+	/////////////////////////////////////////////////////////////////////////////////////////////////
+	// Manual
+
+	static FAnimSequenceTrackPackedData Pack(FAnimSequenceTrackManualData ManualData)
+	{
+		FAnimSequenceTrackPackedData Data;
+		Data.SequenceIndex = ManualData.SequenceIndex;
+		Data.bAutoPlay = 0;
+		Data.LoopMode = 0;
+		Data.PatchFlags = 0;
+		Data.bHasPreviousPosition = ManualData.bHasPreviousPosition;
+		Data.Position = ManualData.Position;
+		Data.PreviousPosition = ManualData.bHasPreviousPosition ? ManualData.PreviousPosition : ManualData.Position;
+		Data.PlayRate = 0.0f;
+		Data.BlendTime = 0.0f;
+		return Data;
+	}
+
+	bool IsManual() const
+	{
+		return !bAutoPlay;
+	}
+
+	FAnimSequenceTrackManualData UnpackManual() const
+	{
+		check(IsManual());
+
+		FAnimSequenceTrackManualData Data;
+		Data.SequenceIndex = SequenceIndex;
+		Data.Position = Position;
+		Data.PreviousPosition = PreviousPosition;
+		Data.bHasPreviousPosition = bHasPreviousPosition;
+		return Data;
+	}
+
+	/////////////////////////////////////////////////////////////////////////////////////////////////
+
+	int32 GetSequenceIndex() const
+	{
+		return SequenceIndex;
+	}
+
+	float GetPosition() const
+	{
+		return Position;
+	}
+
+	float GetPreviousPosition() const
+	{
+		return PreviousPosition;
+	}
+
+	bool HasPreviousPosition() const
+	{
+		return bHasPreviousPosition;
+	}
+
+	void SetLoopMode(EAnimSequenceTrackLoopMode NewLoopMode)
+	{
+		LoopMode = static_cast<uint32>(NewLoopMode);
+		PatchFlags |= (uint32)EAnimSequenceTrackPatchFlags::LoopMode;
+	}
+
+	EAnimSequenceTrackLoopMode GetLoopMode() const
+	{
+		return static_cast<EAnimSequenceTrackLoopMode>(LoopMode);
+	}
+
+	float GetBlendTime() const
+	{
+		return BlendTime;
+	}
+
+	void ResetPatchFlags()
+	{
+		PatchFlags = 0;
+	}
+
+	EAnimSequenceTrackPatchFlags GetPatchFlags() const
+	{
+		return static_cast<EAnimSequenceTrackPatchFlags>(PatchFlags);
+	}
+
+	void SetPlayRate(float NewPlayRate)
+	{
+		PlayRate = NewPlayRate;
+		PatchFlags |= (uint32)EAnimSequenceTrackPatchFlags::PlayRate;
+	}
+
+	float GetPlayRate() const
+	{
+		return PlayRate;
+	}
+
+	ENGINE_API void Tick(float SequenceLength, float DeltaTime);
+
+	ENGINE_API void Wrap(float SequenceLength);
+
+	ENGINE_API void Serialize(FArchive& Ar);
+
+private:
+	uint32 SequenceIndex : 20;
+	uint32 LoopMode : 1;
+	uint32 bAutoPlay : 1;
+	uint32 bHasPreviousPosition : 1;
+	uint32 PatchFlags : 8;
+	uint32 Unused : 1 = 0;
+	float Position;
+	float PreviousPosition;
+	float PlayRate;
+	float BlendTime;
+
+	friend class FAnimSequenceTransformProvider;
+};
+
+inline FArchive& operator <<(FArchive& Ar, FAnimSequenceTrackPackedData& Data)
+{
+	Data.Serialize(Ar);
+	return Ar;
+}
+
+using FAnimSequenceTrackPool = TAnimTrackPool<FAnimSequenceTrackPackedData>;
+
+class FAnimSequenceTrackRenderData
+{
+public:
+	FAnimSequenceTrackRenderData() = default;
+	FAnimSequenceTrackRenderData(const FAnimSequenceTrackRenderData&) = default;
+	FAnimSequenceTrackRenderData& operator=(const FAnimSequenceTrackRenderData&) = default;
+
+	FAnimSequenceTrackRenderData(const FAnimSequenceTrackPackedData& Data)
+	{
+		*this = Data;
+	}
+
+	FAnimSequenceTrackRenderData& operator=(const FAnimSequenceTrackPackedData& PackedData)
+	{
+		Source = PackedData;
+		Target = PackedData;
+		BlendWeight = 0.0f;
+		return *this;
+	}
+
+	bool IsBlending() const
+	{
+		// The check for BlendWeight > 0 is for blending two positions of the same sequence.
+		return Source.GetSequenceIndex() != Target.GetSequenceIndex() || BlendWeight > 0.0f;
+	}
+
+	float GetBlendWeight() const { return BlendWeight; }
+
+	const FAnimSequenceTrackPackedData& GetSource() const { return Source; }
+	const FAnimSequenceTrackPackedData& GetTarget() const { return Target; }
+	
+	ENGINE_API void Patch(const FAnimSequenceTrackPackedData& Request, bool bWasActive);
+
+	ENGINE_API void Tick(float SourceSequenceLength, float TargetSequenceLength, float DeltaTime);
+
+private:
+	FAnimSequenceTrackPackedData Source;
+	FAnimSequenceTrackPackedData Target;
+	float BlendWeight = 0.0f;
+};
+
+using FAnimSequenceTrackRenderPool = TAnimTrackPool<FAnimSequenceTrackRenderData>;
+
+class FAnimSequenceTransformProviderRenderData
+{
+public:
+	ENGINE_API ~FAnimSequenceTransformProviderRenderData();
+
+private:
+	void Patch(FAnimSequenceTrackPool::FPatch&& Patch);
+	void InsertProxy(FInstancedSkinningSceneExtensionProxy* Proxy);
+	void RemoveProxy(FInstancedSkinningSceneExtensionProxy* Proxy);
+
+	FAnimSequenceTrackRenderPool Tracks;
+	Experimental::TRobinHoodHashSet<FInstancedSkinningSceneExtensionProxy*> Proxies;
+
+	friend FAnimSequenceTransformProviderProxy;
+	friend UAnimSequenceTransformProviderDataInstance;
+};
+
+class FAnimSequenceTransformProviderProxy : public FTransformProviderRenderProxy
+{
+	friend class UAnimSequenceTransformProviderDataInstance;
+	friend class UAnimSequenceTransformProviderData;
+
+	ENGINE_API FAnimSequenceTransformProviderProxy(TConstArrayView<FAnimSequenceTransformProviderSequence> Sequences, FInstancedSkinningSceneExtensionProxy* SceneExtensionProxy);
+	ENGINE_API FAnimSequenceTransformProviderProxy(TConstArrayView<FAnimSequenceTransformProviderSequence> Sequences, FInstancedSkinningSceneExtensionProxy* SceneExtensionProxy, FAnimSequenceTransformProviderRenderData& RenderData);
+
+public:
+	virtual void CreateRenderThreadResources(FSceneInterface& Scene, FRHICommandListBase& RHICmdList) override;
+	virtual void DestroyRenderThreadResources() override;
+
+	FAnimSequenceTrackRenderPool& GetTracks()
+	{
+		return *Tracks;
+	}
+
+	TConstArrayView<FAnimSequenceTransformProviderSequence> GetSequences() const
+	{
+		return Sequences;
+	}
+
+	FInstancedSkinningSceneExtensionProxy* GetSceneExtensionProxy() const
+	{
+		return SceneExtensionProxy;
+	}
+
+private:
+	TConstArrayView<FAnimSequenceTransformProviderSequence> Sequences;
+	TUniquePtr<FAnimSequenceTrackRenderPool> TrackAllocator;
+	FAnimSequenceTransformProviderRenderData* RenderData = nullptr;
+	FAnimSequenceTrackRenderPool* Tracks;
+	FInstancedSkinningSceneExtensionProxy* SceneExtensionProxy = nullptr;
+	const USkinnedAsset* SkinnedAsset = nullptr;
+	IAnimSequenceTransformProvider* AnimSequenceTransformProvider = nullptr;
+};
+
+class IAnimSequenceTransformProvider
+{
+public:
+	virtual ~IAnimSequenceTransformProvider() = default;
+
+	virtual void RegisterProxy(FAnimSequenceTransformProviderProxy* Proxy) {}
+	virtual void UnregisterProxy(FAnimSequenceTransformProviderProxy* Proxy) {}
+};
+
+UCLASS(config = Engine, hidecategories = Object, MinimalAPI, BlueprintType)
+class UAnimSequenceTransformProviderData : public UTransformProviderData
+{
+	GENERATED_BODY()
+
+public:
+	virtual bool IsEnabled() const override;
+	virtual const FGuid& GetTransformProviderID() const override;
+	virtual uint32 GetUniqueAnimationCount() const override;
+	virtual bool UsesSkeletonBatching() const override;
+	virtual bool HasAnimationBounds() const override;
+	virtual uint32 GetSkinningDataOffset(int32 InstanceIndex, const FTransform& ComponentTransform, const FSkinnedMeshInstanceData& InstanceData) const override;
+	virtual FTransformProviderRenderProxy* CreateRenderProxy(FInstancedSkinningSceneExtensionProxy* ExtensionProxy) const override;
+	virtual void BeginDestroy() override;
+	virtual bool IsReadyForFinishDestroy() override;
+
+#if WITH_EDITOR
+	virtual void PostEditChangeProperty(FPropertyChangedEvent& PropertyChangedEvent) override;
+#endif
+
+public:
+	// Returns the play length for a sequence.
+	UFUNCTION(BlueprintCallable, Category = "TransformProviders|AnimSequence")
+	ENGINE_API const TArray<struct FAnimSequenceTransformProviderSequence>& GetSequences() const;
+
+	// Returns the play length for a sequence.
+	UFUNCTION(BlueprintCallable, Category = "TransformProviders|AnimSequence")
+	ENGINE_API float GetSequencePlayLength(int32 SequenceIndex) const;
+
+#if WITH_EDITOR
+	// Creates a transform provider data asset from an AnimBankData asset.
+	UFUNCTION(BlueprintCallable, Category = "TransformProviders|AnimSequence", meta = (DevelopmentOnly))
+	static ENGINE_API UAnimSequenceTransformProviderData* CreateFromAnimBankData(UAnimBankData* InAnimBankData);
+
+	// Creates a transform provider data asset from an AnimBank asset.
+	UFUNCTION(BlueprintCallable, Category = "TransformProviders|AnimSequence", meta = (DevelopmentOnly))
+	static ENGINE_API UAnimSequenceTransformProviderData* CreateFromAnimBank(UAnimBank* InAnimBank);
+#endif
+
+protected:
+	bool IsValidFor(FInstancedSkinningSceneExtensionProxy* ExtensionProxy) const;
+
+private:
+	UPROPERTY(EditAnywhere, BlueprintReadOnly, Category = TransformProvider, meta = (AllowPrivateAccess = "true", ShowOnlyInnerProperties))
+	TArray<struct FAnimSequenceTransformProviderSequence> Sequences;
+
+	FRenderCommandFence DestroyFence;
+
+	friend class UAnimSequenceTransformProviderDataInstance;
+};
+
+UCLASS(config = Engine, hidecategories = Object, MinimalAPI, BlueprintType, Within = InstancedSkinnedMeshComponent, NotBlueprintable, NotPlaceable)
+class UAnimSequenceTransformProviderDataInstance : public UAnimSequenceTransformProviderData
+{
+	GENERATED_BODY()
+
+public:
+	virtual bool IsEnabled() const override;
+	virtual uint32 GetUniqueAnimationCount() const override;
+	virtual uint32 GetSkinningDataOffset(int32 InstanceIndex, const FTransform& ComponentTransform, const FSkinnedMeshInstanceData& InstanceData) const override;
+	virtual FTransformProviderRenderProxy* CreateRenderProxy(FInstancedSkinningSceneExtensionProxy* ExtensionProxy) const override;
+	virtual void BeginDestroy() override;
+	virtual void Serialize(FArchive& Ar) override;
+	virtual UTransformProviderData* GetRootTransformProvider() override;
+
+public:
+	// Copies items from this anim bank to create a new anim bank instance.
+	UFUNCTION(BlueprintCallable, Category = "TransformProviders|AnimSequence")
+	static ENGINE_API UAnimSequenceTransformProviderDataInstance* CreateAnimSequenceTransformProviderDataInstance(UAnimSequenceTransformProviderData* ProviderData, UInstancedSkinnedMeshComponent* Owner);
+
+	// Allocates a new track with auto play data, or -1 if the track allocation failed.
+	UFUNCTION(BlueprintCallable, Category = "TransformProviders|AnimSequence")
+	ENGINE_API int32 AllocateAutoPlayTrack(const FAnimSequenceTrackAutoPlayData& Data);
+
+	// Allocates a new track with manual play data, or -1 if the track allocation failed.
+	UFUNCTION(BlueprintCallable, Category = "TransformProviders|AnimSequence")
+	ENGINE_API int32 AllocateManualTrack(const FAnimSequenceTrackManualData& Data);
+
+	// Updates a track with auto play data. Returns whether the update succeeded.
+	UFUNCTION(BlueprintCallable, Category = "TransformProviders|AnimSequence")
+	ENGINE_API bool UpdateAutoPlayTrack(int32 TrackIndex, const FAnimSequenceTrackAutoPlayData& Data);
+
+	// Updates a track with manual data. Returns whether the update succeeded.
+	UFUNCTION(BlueprintCallable, Category = "TransformProviders|AnimSequence")
+	ENGINE_API bool UpdateManualTrack(int32 TrackIndex, const FAnimSequenceTrackManualData& Data);
+
+	// Deallocates a track. Returns whether the deallocation succeeded.
+	UFUNCTION(BlueprintCallable, Category = "TransformProviders|AnimSequence")
+	ENGINE_API bool DeallocateTrack(int32 TrackIndex);
+
+	// Submits all pending changes to rendering.
+	UFUNCTION(BlueprintCallable, Category = "TransformProviders|AnimSequence")
+	ENGINE_API virtual void SubmitChanges() override;
+
+	// Get the sequence index currently playing on a track
+	UFUNCTION(BlueprintCallable, Category = "TransformProviders|AnimSequence")
+	ENGINE_API int32 GetTrackSequenceIndex(int32 TrackIndex) const;
+
+	// Assigns a new play rate to an existing track.
+	UFUNCTION(BlueprintCallable, Category = "TransformProviders|AnimSequence")
+	ENGINE_API bool SetPlayRate(int32 TrackIndex, float PlayRate);
+
+	// Get current play rate of a track
+	UFUNCTION(BlueprintCallable, Category = "TransformProviders|AnimSequence")
+	ENGINE_API float GetPlayRate(int32 TrackIndex) const;
+
+	// Assigns a new loop mode to an existing track.
+	UFUNCTION(BlueprintCallable, Category = "TransformProviders|AnimSequence")
+	ENGINE_API bool SetLoopMode(int32 TrackIndex, EAnimSequenceTrackLoopMode LoopMode);
+
+	// Get current loop mode of a track
+	UFUNCTION(BlueprintCallable, Category = "TransformProviders|AnimSequence")
+	ENGINE_API EAnimSequenceTrackLoopMode GetLoopMode(int32 TrackIndex) const;
+
+private:
+	void Initialize(UAnimSequenceTransformProviderData* ProviderData);
+	void MarkOwnerRenderStateDirty();
+
+	// The parent provider data this instance was created from.
+	UPROPERTY()
+	TObjectPtr<UAnimSequenceTransformProviderData> Parent;
+
+	// Cached owner component for efficient dirty marking.
+	UInstancedSkinnedMeshComponent* CachedOwnerComponent = nullptr;
+
+	// Prevents redundant MarkRenderDynamicDataDirty calls within a frame.
+	bool bPendingRenderStateDirty = false;
+
+	FAnimSequenceTrackPool Tracks;
+	TUniquePtr<FAnimSequenceTransformProviderRenderData> RenderData{ new FAnimSequenceTransformProviderRenderData };
+	FRenderCommandFence InstanceDestroyFence;
+	UE::FMutex SubmitChangesMutex;
+
+	friend class UAnimSequenceTransformProviderData;
+};
\ No newline at end of file
diff --git a/Engine/Source/Runtime/Engine/Classes/Animation/AnimTrackPool.h b/Engine/Source/Runtime/Engine/Classes/Animation/AnimTrackPool.h
--- a/Engine/Source/Runtime/Engine/Classes/Animation/AnimTrackPool.h
+++ b/Engine/Source/Runtime/Engine/Classes/Animation/AnimTrackPool.h
@@ -14,6 +14,34 @@ template <typename DataType>
 class TAnimTrackPool
 {
 public:
+	TAnimTrackPool() = default;
+	TAnimTrackPool(const TAnimTrackPool&) = default;
+	TAnimTrackPool(TAnimTrackPool&&) = default;
+	TAnimTrackPool& operator=(const TAnimTrackPool&) = default;
+	TAnimTrackPool& operator=(TAnimTrackPool&&) = default;
+
+	template<typename SourceTrackDataType>
+	TAnimTrackPool& operator=(const TAnimTrackPool<SourceTrackDataType>& Other)
+	{
+		const int32 NumTracks = Other.GetNumTracks();
+
+		Datas.SetNum(NumTracks);
+		ActiveBits.Init(false, NumTracks);
+		DirtyBits.Init(false, NumTracks);
+		NumActiveTracks = Other.GetNumActiveTracks();
+
+		for (int32 Index = 0; Index < NumTracks; ++Index)
+		{
+			if (Other.IsActiveIndex(Index))
+			{
+				Datas[Index] = Other.GetData(Index);
+				ActiveBits[Index] = true;
+			}
+		}
+
+		return *this;
+	}
+
 	class FPatch
 	{
 	public:
@@ -22,7 +50,6 @@ public:
 			return NumTracks;
 		}
 
-	private:
 		struct FActivateOp
 		{
 			FActivateOp(int32 InIndex, const DataType& InData)
@@ -47,8 +74,6 @@ public:
 		int32 NumActiveTracks;
 		TArray<FActivateOp, FConcurrentLinearArrayAllocator> ActivateOps;
 		TArray<FDeactivateOp, FConcurrentLinearArrayAllocator> DeactivateOps;
-
-		friend class TAnimTrackPool;
 	};
 
 	void Reserve(int32 Num)
@@ -101,8 +126,31 @@ public:
 		}
 		return false;
 	}
+	
+	struct FDefaultDeactivate
+	{
+		void operator()(DataType&) const
+		{
+		}
+	};
+	
+	struct FDefaultAssign
+	{
+		void operator()(DataType& Out, const DataType& In) const
+		{
+			Out = In;
+		}
+	};
+
+	struct FDefaultReset
+	{
+		void operator()(DataType& InOut) const
+		{
+		}
+	};
 
-	FPatch Finalize()
+	template <typename ResetFunctionType = FDefaultReset>
+	FPatch Finalize(ResetFunctionType ResetFunction = FDefaultReset())
 	{
 		FPatch Patch;
 		Patch.NumTracks = Datas.Num();
@@ -116,6 +164,7 @@ public:
 			if (ActiveBits[Index])
 			{
 				Patch.ActivateOps.Emplace(Index, Datas[Index]);
+				ResetFunction(Datas[Index]);
 			}
 			else
 			{
@@ -127,44 +176,32 @@ public:
 		return Patch;
 	}
 
-	struct FDefaultAssign
-	{
-		void operator()(DataType& Out, const DataType& In) const
-		{
-			Out = In;
-		}
-	};
-
-	struct FDefaultDeactivate
-	{
-		void operator()(DataType&) const
-		{
-		}
-	};
-
-	template <typename AssignFunctionType = FDefaultAssign, typename DeactivateFunctionType = FDefaultDeactivate>
+	template <typename SourceTrackData, typename AssignFunctionType = FDefaultAssign, typename DeactivateFunctionType = FDefaultDeactivate>
 	void Patch(
-		FPatch&& Patch,
+		typename TAnimTrackPool<SourceTrackData>::FPatch&& Patch,
 		AssignFunctionType AssignFunction = FDefaultAssign{},
 		DeactivateFunctionType DeactivateFunction = FDefaultDeactivate{})
 	{
+		using FSourcePatch = typename TAnimTrackPool<SourceTrackData>::FPatch;
+
 		Datas.SetNum(Patch.NumTracks);
 		ActiveBits.SetNum(Patch.NumTracks, false);
 		DirtyBits.SetNum(Patch.NumTracks, false);
 		NumActiveTracks = Patch.NumActiveTracks;
 
-		for (typename FPatch::FDeactivateOp Op : Patch.DeactivateOps)
+		for (typename FSourcePatch::FDeactivateOp Op : Patch.DeactivateOps)
 		{
 			DeactivateFunction(Datas[Op.Index]);
 			ActiveBits[Op.Index] = false;
 			SetDirty(Op.Index);
 		}
 
-		for (const typename FPatch::FActivateOp& Op : Patch.ActivateOps)
+		for (const typename FSourcePatch::FActivateOp& Op : Patch.ActivateOps)
 		{
-			ActiveBits[Op.Index] = true;
-			AssignFunction(Datas[Op.Index], Op.Data);
+			auto Bit = ActiveBits[Op.Index];
+			AssignFunction(Datas[Op.Index], Op.Data, Bit);
 			SetDirty(Op.Index);
+			Bit = true;
 		}
 	}
 
diff --git a/Engine/Source/Runtime/Engine/Private/Animation/AnimBank.cpp b/Engine/Source/Runtime/Engine/Private/Animation/AnimBank.cpp
--- a/Engine/Source/Runtime/Engine/Private/Animation/AnimBank.cpp
+++ b/Engine/Source/Runtime/Engine/Private/Animation/AnimBank.cpp
@@ -1229,9 +1229,9 @@ void FAnimBankTrackPackedData::Update(float SequenceLength, bool bLooping, float
 
 void FAnimBankTrackPackedData::Wrap(float SequenceLength, bool bLooping)
 {
-	const EAnimBankTrackLoopMode LoopModeEnum = GetLoopMode();
+	const EAnimSequenceTrackLoopMode LoopModeEnum = GetLoopMode();
 
-	if (LoopModeEnum == EAnimBankTrackLoopMode::Loop || (LoopModeEnum == EAnimBankTrackLoopMode::Auto && bLooping))
+	if (LoopModeEnum == EAnimSequenceTrackLoopMode::Loop)
 	{
 		Position         = GetLoopedPosition(Position, SequenceLength);
 		PreviousPosition = GetLoopedPosition(PreviousPosition, SequenceLength);
@@ -1328,7 +1328,11 @@ void FAnimBankRenderData::Patch(FAnimBankTrackPool::FPatch&& Patch)
 			Proxy->SetUniqueAnimationCount(Patch.GetNumTracks());
 		}
 	}
-	Tracks.Patch(MoveTemp(Patch));
+	Tracks.Patch<FAnimBankTrackPackedData>(MoveTemp(Patch),
+		[&](FAnimBankTrackPackedData& Dst, const FAnimBankTrackPackedData& Src, bool bWasActive)
+	{
+		Dst = Src;
+	});
 }
 
 ///////////////////////////////////////////////////////////////////////////////////////////////////
@@ -1347,10 +1351,11 @@ FAnimBankTransformProviderProxy::FAnimBankTransformProviderProxy(TConstArrayView
 		const FAnimBankData& BankData = Item.BankAsset->GetData();
 		const FAnimBankEntry& BankEntry = BankData.Entries[Item.SequenceIndex];
 
-		FAnimBankTrackAutoPlayData Data;
-		Data.ItemIndex            = ItemIndex;
+		FAnimSequenceTrackAutoPlayData Data;
+		Data.SequenceIndex        = ItemIndex;
 		Data.Position             = BankEntry.Position;
 		Data.PlayRate             = BankEntry.PlayRate;
+		Data.LoopMode             = BankEntry.IsLooping() ? EAnimSequenceTrackLoopMode::Loop : EAnimSequenceTrackLoopMode::Clamp;
 
 		TrackAllocator->AllocateTrack(FAnimBankTrackPackedData::Pack(Data));
 	}
@@ -1558,183 +1563,6 @@ float UAnimBankData::GetSequencePlayLength(int32 ItemIndex) const
 
 ///////////////////////////////////////////////////////////////////////////////////////////////////
 
-UAnimBankDataInstance* UAnimBankDataInstance::CreateAnimBankDataInstance(UAnimBankData* AnimBank, UInstancedSkinnedMeshComponent* Owner)
-{
-	if (IsValid(AnimBank) && IsValid(Owner))
-	{
-		UAnimBankDataInstance* AnimBankInstance = NewObject<UAnimBankDataInstance>(Owner);
-		AnimBankInstance->AnimBankItems = AnimBank->AnimBankItems;
-		AnimBankInstance->bEnabled = true;
-		AnimBankInstance->InitSequences();
-		return AnimBankInstance;
-	}
-
-	return nullptr;
-}
-
-bool UAnimBankDataInstance::IsEnabled() const
-{
-	return Tracks.GetNumTracks() != 0 && Super::IsEnabled();
-}
-
-uint32 UAnimBankDataInstance::GetUniqueAnimationCount() const
-{
-	return Tracks.GetNumTracks();
-}
-
-bool UAnimBankDataInstance::GetAnimationBounds(uint32 AnimationIndex, FRenderBounds& OutBounds) const
-{
-	if (Tracks.IsActiveIndex(AnimationIndex))
-	{
-		return Super::GetAnimationBounds(Tracks.GetData(AnimationIndex).GetItemIndex(), OutBounds);
-	}
-
-	return false;
-}
-
-uint32 UAnimBankDataInstance::GetSkinningDataOffset(int32 InstanceIndex, const FTransform& ComponentTransform, const FSkinnedMeshInstanceData& InstanceData) const
-{
-	const uint32 AnimationIndex = InstanceData.AnimationIndex;
-	if (AnimationIndex < uint32(Tracks.GetNumTracks()))
-	{
-		return AnimationIndex * 2u;
-	}
-
-	return 0u;
-}
-
-FTransformProviderRenderProxy* UAnimBankDataInstance::CreateRenderProxy(FInstancedSkinningSceneExtensionProxy* SceneExtensionProxy) const
-{
-	if (IsAnimBankValid(SceneExtensionProxy))
-	{
-		return new FAnimBankTransformProviderProxy(AnimBankItems, SceneExtensionProxy, *RenderData);
-	}
-	return nullptr;
-}
-
-void UAnimBankDataInstance::BeginDestroy()
-{
-	ENQUEUE_RENDER_COMMAND(AnimBankDataInstanceRelease)([RenderData = MoveTemp(RenderData)](FRHICommandList&) mutable
-	{
-		// Move destruction onto the render thread.
-	});
-
-	Super::BeginDestroy();
-}
-
-void UAnimBankDataInstance::SubmitChanges()
-{
-	check(RenderData);
-
-	// SubmitChanges is called during the ISKM instance data update. Guard against multiple concurrent updates.
-	UE::TScopeLock Lock(SubmitChangesMutex);
-
-	if (Tracks.GetNumDirtyTracks() > 0)
-	{
-		ENQUEUE_RENDER_COMMAND(AnimBankDataInstanceSubmit)([&RenderData = *RenderData, Patch = Tracks.Finalize()] (FRHICommandList&) mutable
-		{
-			RenderData.Patch(MoveTemp(Patch));
-		});
-	}
-}
-
-void UAnimBankDataInstance::InitRenderData()
-{
-	check(RenderData);
-	RenderData->Tracks = Tracks;
-}
-
-void UAnimBankDataInstance::InitSequences()
-{
-	Sequences.Reserve(AnimBankItems.Num());
-
-	for (const FAnimBankItem& Item : AnimBankItems)
-	{
-		if (Item.BankAsset)
-		{
-			const bool bLooping = Item.BankAsset->IsSequenceLooping(Item.SequenceIndex);
-			const float PlayLength = Item.BankAsset->GetSequencePlayLength(Item.SequenceIndex);
-			Sequences.Emplace(PlayLength, bLooping);
-		}
-	}
-}
-
-void UAnimBankDataInstance::Serialize(FArchive& Ar)
-{
-	Ar.UsingCustomVersion(FAnimObjectVersion::GUID);
-	Super::Serialize(Ar);
-	Tracks.Serialize(Ar);
-
-	if (Ar.IsLoading())
-	{
-		InitRenderData();
-		InitSequences();
-	}
-}
-
-int32 UAnimBankDataInstance::AllocateAutoPlayTrack(const FAnimBankTrackAutoPlayData& Data)
-{
-	FAnimBankTrackPackedData PackedData = FAnimBankTrackPackedData::Pack(Data);
-
-	if (Sequences.IsValidIndex(Data.ItemIndex))
-	{
-		const FSequence& Sequence = Sequences[Data.ItemIndex];
-		PackedData.Wrap(Sequence.PlayLength, Sequence.bLooping);
-		return Tracks.AllocateTrack(PackedData);
-	}
-
-	return INDEX_NONE;
-}
-
-int32 UAnimBankDataInstance::AllocateManualTrack(const FAnimBankTrackManualData& Data)
-{
-	FAnimBankTrackPackedData PackedData = FAnimBankTrackPackedData::Pack(Data);
-
-	if (Sequences.IsValidIndex(Data.ItemIndex))
-	{
-		const FSequence& Sequence = Sequences[Data.ItemIndex];
-		PackedData.Wrap(Sequence.PlayLength, Sequence.bLooping);
-		return Tracks.AllocateTrack(PackedData);
-	}
-
-	return INDEX_NONE;
-}
-
-bool UAnimBankDataInstance::UpdateAutoPlayTrack(int32 TrackIndex, const FAnimBankTrackAutoPlayData& Data)
-{
-	FAnimBankTrackPackedData PackedData = FAnimBankTrackPackedData::Pack(Data);
-
-	if (Sequences.IsValidIndex(Data.ItemIndex))
-	{
-		const FSequence& Sequence = Sequences[Data.ItemIndex];
-		PackedData.Wrap(Sequence.PlayLength, Sequence.bLooping);
-		return Tracks.UpdateTrack(TrackIndex, PackedData);
-	}
-
-	return false;
-}
-
-bool UAnimBankDataInstance::UpdateManualTrack(int32 TrackIndex, const FAnimBankTrackManualData& Data)
-{
-	FAnimBankTrackPackedData PackedData = FAnimBankTrackPackedData::Pack(Data);
-
-	if (Sequences.IsValidIndex(Data.ItemIndex))
-	{
-		const FSequence& Sequence = Sequences[Data.ItemIndex];
-		PackedData.Wrap(Sequence.PlayLength, Sequence.bLooping);
-		return Tracks.UpdateTrack(TrackIndex, PackedData);
-	}
-
-	return false;
-}
-
-bool UAnimBankDataInstance::DeallocateTrack(int32 TrackIndex)
-{
-	return Tracks.DeallocateTrack(TrackIndex);
-}
-
-///////////////////////////////////////////////////////////////////////////////////////////////////
-
 void FAnimBankSequence::ValidatePosition()
 {
 	float Min = 0.0f, Max = 0.0f;
diff --git a/Engine/Source/Runtime/Engine/Private/Animation/AnimRuntimeTransformProviderData.cpp b/Engine/Source/Runtime/Engine/Private/Animation/AnimRuntimeTransformProviderData.cpp
--- a/Engine/Source/Runtime/Engine/Private/Animation/AnimRuntimeTransformProviderData.cpp
+++ b/Engine/Source/Runtime/Engine/Private/Animation/AnimRuntimeTransformProviderData.cpp
@@ -16,23 +16,23 @@ inline void TryAddTail(TIntrusiveDoubleLinkedList<FAnimRuntimeTrackTransformData
 	}
 }
 
-FAnimRuntimeRenderData::~FAnimRuntimeRenderData()
+FAnimRuntimeTransformProviderRenderData::~FAnimRuntimeTransformProviderRenderData()
 {
 	check(Proxies.Num() == 0);
 	check(FreeList.IsEmpty());
 }
 
-void FAnimRuntimeRenderData::InsertProxy(FInstancedSkinningSceneExtensionProxy* Proxy)
+void FAnimRuntimeTransformProviderRenderData::InsertProxy(FInstancedSkinningSceneExtensionProxy* Proxy)
 {
 	Proxies.FindOrAdd(Proxy);
 }
 
-void FAnimRuntimeRenderData::RemoveProxy(FInstancedSkinningSceneExtensionProxy* Proxy)
+void FAnimRuntimeTransformProviderRenderData::RemoveProxy(FInstancedSkinningSceneExtensionProxy* Proxy)
 {
 	Proxies.Remove(Proxy);
 }
 
-void FAnimRuntimeRenderData::Patch(FAnimRuntimeTrackPool::FPatch&& Patch)
+void FAnimRuntimeTransformProviderRenderData::Patch(FAnimRuntimeTrackPool::FPatch&& Patch)
 {
 	if (Patch.GetNumTracks() != Tracks.GetNumTracks())
 	{
@@ -44,10 +44,10 @@ void FAnimRuntimeRenderData::Patch(FAnimRuntimeTrackPool::FPatch&& Patch)
 
 	UE::TScopeLock Lock(FreeListMutex);
 
-	Tracks.Patch(
+	Tracks.Patch<FAnimRuntimeTrackData>(
 		MoveTemp(Patch),
 		// Assignment
-		[&] (FAnimRuntimeTrackData& Dst, const FAnimRuntimeTrackData& Src)
+		[&] (FAnimRuntimeTrackData& Dst, const FAnimRuntimeTrackData& Src, bool bWasActive)
 		{
 			// Retire the previous transform data if it's not a duplicate of the current.
 			if (Dst.Current != Dst.Previous)
diff --git a/Engine/Source/Runtime/Engine/Private/Animation/AnimSequenceTransformProviderData.cpp b/Engine/Source/Runtime/Engine/Private/Animation/AnimSequenceTransformProviderData.cpp
new file mode 100644
--- /dev/null
+++ b/Engine/Source/Runtime/Engine/Private/Animation/AnimSequenceTransformProviderData.cpp
@@ -0,0 +1,835 @@
+// Copyright Epic Games, Inc. All Rights Reserved.
+
+#include "Animation/AnimSequenceTransformProviderData.h"
+#include "SkinningSceneExtensionProxy.h"
+#include "Components/InstancedSkinnedMeshComponent.h"
+#include "Animation/SkeletonRemapping.h"
+#include "Animation/SkeletonRemappingRegistry.h"
+#include "Animation/AnimSequence.h"
+#include "Animation/Skeleton.h"
+#include "SceneInterface.h"
+#include "UObject/AnimObjectVersion.h"
+#include "UObject/UObjectThreadContext.h"
+#include "ObjectCacheContext.h"
+
+#if WITH_EDITOR
+#include "Animation/AnimBank.h"
+#include "AssetRegistry/AssetRegistryModule.h"
+#include "Misc/PackageName.h"
+#endif
+
+///////////////////////////////////////////////////////////////////////////////////////////////////
+
+static TAutoConsoleVariable<bool> CVarAnimSequenceTransformProviderBlending(
+	TEXT("r.AnimSequenceTransformProvider.Blending"),
+	true,
+	TEXT(""),
+	ECVF_Scalability | ECVF_RenderThreadSafe
+);
+
+///////////////////////////////////////////////////////////////////////////////////////////////////
+
+inline float GetLoopedPosition(float Position, float TrackLength)
+{
+	if (TrackLength >= 0.0f)
+	{
+		Position = TrackLength == 0.0f ? 0.0f : FMath::Fmod(Position, TrackLength);
+
+		// Returns a negative range if input is negative.
+		if (Position < 0.0f)
+		{
+			Position += TrackLength;
+		}
+	}
+	return Position;
+}
+
+void FAnimSequenceTrackPackedData::Tick(float SequenceLength, float DeltaTime)
+{
+	if (bAutoPlay)
+	{
+		PreviousPosition = Position;
+		Position += PlayRate * DeltaTime;
+		bHasPreviousPosition = false;
+		Wrap(SequenceLength);
+	}
+}
+
+void FAnimSequenceTrackPackedData::Wrap(float SequenceLength)
+{
+	const EAnimSequenceTrackLoopMode LoopModeEnum = GetLoopMode();
+
+	if (LoopModeEnum == EAnimSequenceTrackLoopMode::Loop)
+	{
+		Position = GetLoopedPosition(Position, SequenceLength);
+		PreviousPosition = GetLoopedPosition(PreviousPosition, SequenceLength);
+	}
+	else
+	{
+		Position = FMath::Clamp(Position, 0.0f, SequenceLength);
+		PreviousPosition = FMath::Clamp(PreviousPosition, 0.0f, SequenceLength);
+	}
+}
+
+void FAnimSequenceTrackPackedData::Serialize(FArchive& Ar)
+{
+	constexpr auto MakeBitMask = [](uint32 Bits) -> uint32 { return (Bits >= 32) ? 0xffffffffu : ((1u << Bits) - 1u); };
+
+	constexpr uint32 SequenceIndexBits = 20;
+	constexpr uint32 SequenceIndexMask = MakeBitMask(SequenceIndexBits);
+
+	constexpr uint32 LoopModeShift = SequenceIndexBits;
+	constexpr uint32 LoopModeBits = 1;
+	constexpr uint32 LoopModeMask = MakeBitMask(LoopModeBits);
+
+	constexpr uint32 AutoPlayShift = LoopModeShift + LoopModeBits;
+	constexpr uint32 AutoPlayBits = 1;
+	constexpr uint32 AutoPlayMask = MakeBitMask(AutoPlayBits);
+
+	constexpr uint32 HasPreviousPositionShift = AutoPlayShift + AutoPlayBits;
+	constexpr uint32 HasPreviousPositionBits = 1;
+	constexpr uint32 HasPreviousPositionMask = MakeBitMask(HasPreviousPositionBits);
+
+	uint32 PackedBits;
+	if (Ar.IsSaving())
+	{
+		PackedBits =
+			(SequenceIndex & SequenceIndexMask)
+			| ((LoopMode & LoopModeMask) << LoopModeShift)
+			| ((bAutoPlay & AutoPlayMask) << AutoPlayShift)
+			| ((bHasPreviousPosition & HasPreviousPositionMask) << HasPreviousPositionShift);
+		Ar << PackedBits;
+		Ar << Position;
+
+		if (bAutoPlay)
+		{
+			Ar << PlayRate;
+			Ar << BlendTime;
+		}
+		else
+		{
+			Ar << PreviousPosition;
+		}
+	}
+	else
+	{
+		Ar << PackedBits;
+		SequenceIndex = PackedBits & SequenceIndexMask;
+		LoopMode = (PackedBits >> LoopModeShift) & LoopModeMask;
+		bAutoPlay = (PackedBits >> AutoPlayShift) & AutoPlayMask;
+		bHasPreviousPosition = (PackedBits >> HasPreviousPositionShift) & HasPreviousPositionMask;
+		PatchFlags = 0;
+		Ar << Position;
+
+		if (bAutoPlay)
+		{
+			Ar << PlayRate;
+			Ar << BlendTime;
+			PreviousPosition = 0.0f;
+		}
+		else
+		{
+			Ar << PreviousPosition;
+		}
+	}
+}
+
+///////////////////////////////////////////////////////////////////////////////////////////////////
+
+void FAnimSequenceTrackRenderData::Patch(const FAnimSequenceTrackPackedData& Data, bool bWasActive)
+{
+	// Only patching specific flags.
+	if (Data.GetPatchFlags() != EAnimSequenceTrackPatchFlags::None)
+	{
+		if (EnumHasAnyFlags(Data.GetPatchFlags(), EAnimSequenceTrackPatchFlags::PlayRate))
+		{
+			Source.SetPlayRate(Data.GetPlayRate());
+			Target.SetPlayRate(Data.GetPlayRate());
+		}
+
+		if (EnumHasAnyFlags(Data.GetPatchFlags(), EAnimSequenceTrackPatchFlags::LoopMode))
+		{
+			Source.SetLoopMode(Data.GetLoopMode());
+			Target.SetLoopMode(Data.GetLoopMode());
+		}
+
+		Source.ResetPatchFlags();
+		Target.ResetPatchFlags();
+	}
+	// Patching the entire struct. Starts a new animation.
+	else
+	{
+		if (Data.GetBlendTime() <= 0.0f || !bWasActive)
+		{
+			Source = Data;
+		}
+		else if (BlendWeight > 0.0f)
+		{
+			Source = Target;
+		}
+
+		BlendWeight = 0.0f;
+		Target = Data;
+	}
+}
+
+void FAnimSequenceTrackRenderData::Tick(float SourceSequenceLength, float TargetSequenceLength, float DeltaTime)
+{
+	Source.Tick(SourceSequenceLength, DeltaTime);
+	Source.Wrap(SourceSequenceLength);
+
+	Target.Tick(TargetSequenceLength, DeltaTime);
+	Target.Wrap(TargetSequenceLength);
+
+	const float BlendTime = Target.GetBlendTime();
+
+	if (BlendTime > 0.0f && BlendWeight < 1.0f)
+	{
+		BlendWeight += DeltaTime / BlendTime;
+
+		if (BlendWeight >= 1.0f)
+		{
+			BlendWeight = 0.0f;
+			Source = Target;
+		}
+	}
+}
+
+///////////////////////////////////////////////////////////////////////////////////////////////////
+
+FAnimSequenceTransformProviderRenderData::~FAnimSequenceTransformProviderRenderData()
+{
+	check(Proxies.Num() == 0);
+}
+
+void FAnimSequenceTransformProviderRenderData::InsertProxy(FInstancedSkinningSceneExtensionProxy* Proxy)
+{
+	Proxies.FindOrAdd(Proxy);
+}
+
+void FAnimSequenceTransformProviderRenderData::RemoveProxy(FInstancedSkinningSceneExtensionProxy* Proxy)
+{
+	Proxies.Remove(Proxy);
+}
+
+void FAnimSequenceTransformProviderRenderData::Patch(FAnimSequenceTrackPool::FPatch&& Patch)
+{
+	if (Patch.GetNumTracks() != Tracks.GetNumTracks())
+	{
+		for (FInstancedSkinningSceneExtensionProxy* Proxy : Proxies)
+		{
+			Proxy->SetUniqueAnimationCount(Patch.GetNumTracks());
+		}
+	}
+
+	Tracks.Patch<FAnimSequenceTrackPackedData>(MoveTemp(Patch), [](FAnimSequenceTrackRenderData& Dst, const FAnimSequenceTrackPackedData& Src, bool bWasActive)
+	{
+		Dst.Patch(Src, bWasActive);
+	});
+}
+
+///////////////////////////////////////////////////////////////////////////////////////////////////
+
+FAnimSequenceTransformProviderProxy::FAnimSequenceTransformProviderProxy(TConstArrayView<FAnimSequenceTransformProviderSequence> InSequences, FInstancedSkinningSceneExtensionProxy* InSceneExtensionProxy)
+	: Sequences(InSequences)
+	, TrackAllocator(new FAnimSequenceTrackRenderPool)
+	, Tracks(TrackAllocator.Get())
+	, SceneExtensionProxy(InSceneExtensionProxy)
+{
+	TrackAllocator->Reserve(Sequences.Num());
+
+	for (int32 SequenceIndex = 0; SequenceIndex < Sequences.Num(); ++SequenceIndex)
+	{
+		const FAnimSequenceTransformProviderSequence& Sequence = Sequences[SequenceIndex];
+
+		FAnimSequenceTrackAutoPlayData Data;
+		Data.SequenceIndex = SequenceIndex;
+		Data.LoopMode = Sequence.LoopMode;
+		Data.PlayRate = Sequence.PlayRate;
+
+		TrackAllocator->AllocateTrack(FAnimSequenceTrackPackedData::Pack(Data));
+	}
+}
+
+FAnimSequenceTransformProviderProxy::FAnimSequenceTransformProviderProxy(TConstArrayView<FAnimSequenceTransformProviderSequence> InSequences, FInstancedSkinningSceneExtensionProxy* InSceneExtensionProxy, FAnimSequenceTransformProviderRenderData& InRenderData)
+	: Sequences(InSequences)
+	, RenderData(&InRenderData)
+	, Tracks(&InRenderData.Tracks)
+	, SceneExtensionProxy(InSceneExtensionProxy)
+{
+}
+
+void FAnimSequenceTransformProviderProxy::CreateRenderThreadResources(FSceneInterface& Scene, FRHICommandListBase& RHICmdList)
+{
+	if (RenderData)
+	{
+		RenderData->InsertProxy(SceneExtensionProxy);
+	}
+
+	AnimSequenceTransformProvider = Scene.GetAnimSequenceTransformProvider();
+
+	if (AnimSequenceTransformProvider)
+	{
+		AnimSequenceTransformProvider->RegisterProxy(this);
+	}
+}
+
+void FAnimSequenceTransformProviderProxy::DestroyRenderThreadResources()
+{
+	if (RenderData)
+	{
+		RenderData->RemoveProxy(SceneExtensionProxy);
+	}
+
+	if (AnimSequenceTransformProvider)
+	{
+		AnimSequenceTransformProvider->UnregisterProxy(this);
+		AnimSequenceTransformProvider = nullptr;
+	}
+}
+
+///////////////////////////////////////////////////////////////////////////////////////////////////
+
+bool UAnimSequenceTransformProviderData::IsEnabled() const
+{
+	if (Sequences.IsEmpty())
+	{
+		return false;
+	}
+
+	for (const FAnimSequenceTransformProviderSequence& Sequence : Sequences)
+	{
+		if (Sequence.Sequence == nullptr)
+		{
+			return false;
+		}
+	}
+
+	return bEnabled;
+}
+
+const FGuid& UAnimSequenceTransformProviderData::GetTransformProviderID() const
+{
+	static FGuid AnimSequenceGPUProviderId(ANIM_SEQUENCE_GPU_TRANSFORM_PROVIDER_GUID);
+	static FGuid AnimSequenceCPUProviderId(ANIM_SEQUENCE_CPU_TRANSFORM_PROVIDER_GUID);
+
+	static const auto AnimSequenceGPUVar = IConsoleManager::Get().FindTConsoleVariableDataBool(TEXT("r.AnimSequenceTransformProvider.GPU"));
+	if (AnimSequenceGPUVar && AnimSequenceGPUVar->GetValueOnAnyThread() == 1)
+	{
+		return AnimSequenceGPUProviderId;
+	}
+	else
+	{
+		return AnimSequenceCPUProviderId;
+	}
+}
+
+uint32 UAnimSequenceTransformProviderData::GetUniqueAnimationCount() const
+{
+	return Sequences.Num();
+}
+
+bool UAnimSequenceTransformProviderData::UsesSkeletonBatching() const
+{
+	return false;
+}
+
+bool UAnimSequenceTransformProviderData::HasAnimationBounds() const
+{
+	return false;
+}
+
+uint32 UAnimSequenceTransformProviderData::GetSkinningDataOffset(int32 InstanceIndex, const FTransform& ComponentTransform, const FSkinnedMeshInstanceData& InstanceData) const
+{
+	const uint32 AnimationIndex = InstanceData.AnimationIndex;
+	if (AnimationIndex < uint32(Sequences.Num()))
+	{
+		return AnimationIndex * 2u;
+	}
+
+	return 0u;
+}
+
+void UAnimSequenceTransformProviderData::BeginDestroy()
+{
+	DestroyFence.BeginFence();
+	Super::BeginDestroy();
+}
+
+bool UAnimSequenceTransformProviderData::IsReadyForFinishDestroy()
+{
+	if (!Super::IsReadyForFinishDestroy())
+	{
+		return false;
+	}
+
+	return DestroyFence.IsFenceComplete();
+}
+
+bool UAnimSequenceTransformProviderData::IsValidFor(FInstancedSkinningSceneExtensionProxy* ExtensionProxy) const
+{
+	const USkinnedAsset* SkinnedAsset = ExtensionProxy->GetSkinnedAsset();
+
+	for (const FAnimSequenceTransformProviderSequence& Sequence : Sequences)
+	{
+		if (Sequence.Sequence == nullptr)
+		{
+			UE_LOG(LogAnimation, Warning, TEXT("AnimSequenceTransformProvider: Invalid null sequence in provider '%s' for skinned asset '%s'"),
+				*GetPathName(),
+				SkinnedAsset ? *SkinnedAsset->GetPathName() : TEXT("None"));
+			return false;
+		}
+
+		const USkeleton* SourceSkeleton = Sequence.Sequence->GetSkeleton();
+		const USkeleton* TargetSkeleton = SkinnedAsset ? SkinnedAsset->GetSkeleton() : nullptr;
+
+		if (SourceSkeleton == nullptr)
+		{
+			UE_LOG(LogAnimation, Warning, TEXT("AnimSequenceTransformProvider: Sequence '%s' has no skeleton, cannot use with skinned asset '%s'"),
+				*Sequence.Sequence->GetPathName(),
+				SkinnedAsset ? *SkinnedAsset->GetPathName() : TEXT("None"));
+			return false;
+		}
+
+		if (TargetSkeleton == nullptr)
+		{
+			UE_LOG(LogAnimation, Warning, TEXT("AnimSequenceTransformProvider: Skinned asset '%s' has no skeleton, cannot use with sequence '%s'"),
+				SkinnedAsset ? *SkinnedAsset->GetPathName() : TEXT("None"),
+				*Sequence.Sequence->GetPathName());
+			return false;
+		}
+
+		if (SourceSkeleton != TargetSkeleton)
+		{
+			const FSkeletonRemapping& SkeletonRemapping = UE::Anim::FSkeletonRemappingRegistry::Get().GetRemapping(SourceSkeleton, TargetSkeleton);
+
+			if (!SkeletonRemapping.IsValid())
+			{
+				UE_LOG(LogAnimation, Warning, TEXT("AnimSequenceTransformProvider: Incompatible skeletons: sequence '%s' uses skeleton '%s', but skinned asset '%s' uses skeleton '%s' with no valid remapping"),
+					*Sequence.Sequence->GetPathName(),
+					*SourceSkeleton->GetPathName(),
+					*SkinnedAsset->GetPathName(),
+					*TargetSkeleton->GetPathName());
+				return false;
+			}
+		}
+	}
+
+	return true;
+}
+
+FTransformProviderRenderProxy* UAnimSequenceTransformProviderData::CreateRenderProxy(FInstancedSkinningSceneExtensionProxy* SceneExtensionProxy) const
+{
+	if (IsValidFor(SceneExtensionProxy))
+	{
+		return new FAnimSequenceTransformProviderProxy(Sequences, SceneExtensionProxy);
+	}
+	return nullptr;
+}
+
+float UAnimSequenceTransformProviderData::GetSequencePlayLength(int32 SequenceIndex) const
+{
+	if (Sequences.IsValidIndex(SequenceIndex) && Sequences[SequenceIndex].Sequence)
+	{
+		return Sequences[SequenceIndex].Sequence->GetPlayLength();
+	}
+	return 0.0f;
+}
+
+const TArray<FAnimSequenceTransformProviderSequence>& UAnimSequenceTransformProviderData::GetSequences() const
+{
+	return Sequences;
+}
+
+#if WITH_EDITOR
+
+void UAnimSequenceTransformProviderData::PostEditChangeProperty(FPropertyChangedEvent& PropertyChangedEvent)
+{
+	Super::PostEditChangeProperty(PropertyChangedEvent);
+
+	if (IsTemplate() || FUObjectThreadContext::Get().IsRoutingPostLoad)
+	{
+		return;
+	}
+
+	// Propagate Sequences changes to all instances
+	FObjectCacheContextScope ObjectCacheScope;
+	for (UInstancedSkinnedMeshComponent* Component : ObjectCacheScope.GetContext().GetInstancedSkinnedMeshComponents(GetRootTransformProvider()))
+	{
+		if (Component && IsValid(Component))
+		{
+			// Sync instance data if it's a child of this provider
+			if (UAnimSequenceTransformProviderDataInstance* Instance = Cast<UAnimSequenceTransformProviderDataInstance>(Component->GetTransformProvider()))
+			{
+				if (Instance->Parent == this)
+				{
+					Instance->Initialize(this);
+				}
+			}
+
+			Component->MarkRenderStateDirty();
+		}
+	}
+}
+
+UAnimSequenceTransformProviderData* UAnimSequenceTransformProviderData::CreateFromAnimBankData(UAnimBankData* InAnimBankData)
+{
+	if (!InAnimBankData)
+	{
+		return nullptr;
+	}
+
+	UPackage* Package = InAnimBankData->GetOutermost();
+	const FString NewAssetName = InAnimBankData->GetName() + TEXT("_ASTP");
+	const FString NewAssetPath = FPackageName::GetLongPackagePath(Package->GetName()) / NewAssetName;
+
+	if (StaticFindObject(UAnimSequenceTransformProviderData::StaticClass(), nullptr, *NewAssetPath))
+	{
+		UE_LOG(LogAnimation, Warning, TEXT("CreateFromAnimBankData: Asset already exists at %s"), *NewAssetPath);
+		return nullptr;
+	}
+
+	UPackage* NewPackage = CreatePackage(*NewAssetPath);
+	if (!NewPackage)
+	{
+		return nullptr;
+	}
+
+	UAnimSequenceTransformProviderData* NewProviderData = NewObject<UAnimSequenceTransformProviderData>(
+		NewPackage,
+		*NewAssetName,
+		RF_Public | RF_Standalone);
+
+	if (!NewProviderData)
+	{
+		return nullptr;
+	}
+
+	for (const FAnimBankItem& BankItem : InAnimBankData->AnimBankItems)
+	{
+		FAnimSequenceTransformProviderSequence Sequence;
+		if (BankItem.BankAsset)
+		{
+			const TArray<FAnimBankSequence>& BankSequences = BankItem.BankAsset->Sequences;
+			if (BankSequences.IsValidIndex(BankItem.SequenceIndex))
+			{
+				const FAnimBankSequence& BankSequence = BankSequences[BankItem.SequenceIndex];
+				Sequence.Sequence = BankSequence.Sequence;
+				Sequence.LoopMode = BankSequence.bLooping ? EAnimSequenceTrackLoopMode::Loop : EAnimSequenceTrackLoopMode::Clamp;
+				Sequence.PlayRate = BankSequence.PlayRate;
+			}
+		}
+		NewProviderData->Sequences.Add(Sequence);
+	}
+
+	NewProviderData->MarkPackageDirty();
+	FAssetRegistryModule::AssetCreated(NewProviderData);
+	return NewProviderData;
+}
+
+UAnimSequenceTransformProviderData* UAnimSequenceTransformProviderData::CreateFromAnimBank(UAnimBank* InAnimBank)
+{
+	if (!InAnimBank)
+	{
+		return nullptr;
+	}
+
+	UPackage* Package = InAnimBank->GetOutermost();
+	const FString NewAssetName = InAnimBank->GetName() + TEXT("_ASTP");
+	const FString NewAssetPath = FPackageName::GetLongPackagePath(Package->GetName()) / NewAssetName;
+
+	if (StaticFindObject(UAnimSequenceTransformProviderData::StaticClass(), nullptr, *NewAssetPath))
+	{
+		UE_LOG(LogAnimation, Warning, TEXT("CreateFromAnimBankData: Asset already exists at %s"), *NewAssetPath);
+		return nullptr;
+	}
+
+	UPackage* NewPackage = CreatePackage(*NewAssetPath);
+	if (!NewPackage)
+	{
+		return nullptr;
+	}
+
+	UAnimSequenceTransformProviderData* NewProviderData = NewObject<UAnimSequenceTransformProviderData>(
+		NewPackage,
+		*NewAssetName,
+		RF_Public | RF_Standalone);
+
+	if (!NewProviderData)
+	{
+		return nullptr;
+	}
+
+	for (const FAnimBankSequence& BankSequence : InAnimBank->Sequences)
+	{
+		FAnimSequenceTransformProviderSequence Sequence;
+		Sequence.Sequence = BankSequence.Sequence;
+		Sequence.LoopMode = BankSequence.bLooping ? EAnimSequenceTrackLoopMode::Loop : EAnimSequenceTrackLoopMode::Clamp;
+		Sequence.PlayRate = BankSequence.PlayRate;
+		NewProviderData->Sequences.Add(Sequence);
+	}
+
+	NewProviderData->MarkPackageDirty();
+	FAssetRegistryModule::AssetCreated(NewProviderData);
+	return NewProviderData;
+}
+
+#endif // WITH_EDITOR
+
+///////////////////////////////////////////////////////////////////////////////////////////////////
+
+UAnimSequenceTransformProviderDataInstance* UAnimSequenceTransformProviderDataInstance::CreateAnimSequenceTransformProviderDataInstance(UAnimSequenceTransformProviderData* ProviderData, UInstancedSkinnedMeshComponent* Owner)
+{
+	if (IsValid(ProviderData) && IsValid(Owner))
+	{
+		UAnimSequenceTransformProviderDataInstance* AnimSequenceInstance = NewObject<UAnimSequenceTransformProviderDataInstance>(Owner);
+		AnimSequenceInstance->Initialize(ProviderData);
+		return AnimSequenceInstance;
+	}
+
+	return nullptr;
+}
+
+bool UAnimSequenceTransformProviderDataInstance::IsEnabled() const
+{
+	return Tracks.GetNumTracks() != 0 && Super::IsEnabled();
+}
+
+uint32 UAnimSequenceTransformProviderDataInstance::GetUniqueAnimationCount() const
+{
+	return Tracks.GetNumTracks();
+}
+
+uint32 UAnimSequenceTransformProviderDataInstance::GetSkinningDataOffset(int32 InstanceIndex, const FTransform& ComponentTransform, const FSkinnedMeshInstanceData& InstanceData) const
+{
+	const uint32 AnimationIndex = InstanceData.AnimationIndex;
+	if (AnimationIndex < uint32(Tracks.GetNumTracks()))
+	{
+		return AnimationIndex * 2u;
+	}
+
+	return 0u;
+}
+
+UTransformProviderData* UAnimSequenceTransformProviderDataInstance::GetRootTransformProvider()
+{
+	return Parent ? Parent->GetRootTransformProvider() : this;
+}
+
+FTransformProviderRenderProxy* UAnimSequenceTransformProviderDataInstance::CreateRenderProxy(FInstancedSkinningSceneExtensionProxy* SceneExtensionProxy) const
+{
+	if (IsValidFor(SceneExtensionProxy))
+	{
+		return new FAnimSequenceTransformProviderProxy(Sequences, SceneExtensionProxy, *RenderData);
+	}
+	return nullptr;
+}
+
+void UAnimSequenceTransformProviderDataInstance::BeginDestroy()
+{
+	ENQUEUE_RENDER_COMMAND(AnimSequenceDataInstanceRelease)([RenderData = MoveTemp(RenderData)](FRHICommandList&) mutable
+	{
+		// Move destruction onto the render thread.
+	});
+
+	Super::BeginDestroy();
+}
+
+void UAnimSequenceTransformProviderDataInstance::MarkOwnerRenderStateDirty()
+{
+	if (bPendingRenderStateDirty)
+	{
+		return;
+	}
+
+	if (!CachedOwnerComponent)
+	{
+		CachedOwnerComponent = GetTypedOuter<UInstancedSkinnedMeshComponent>();
+	}
+
+	if (CachedOwnerComponent)
+	{
+		CachedOwnerComponent->MarkRenderInstancesDirty();
+		bPendingRenderStateDirty = true;
+	}
+}
+
+void UAnimSequenceTransformProviderDataInstance::SubmitChanges()
+{
+	check(RenderData);
+
+	// SubmitChanges is called during the ISKM instance data update. Guard against multiple concurrent updates.
+	UE::TScopeLock Lock(SubmitChangesMutex);
+
+	bPendingRenderStateDirty = false;
+
+	if (Tracks.GetNumDirtyTracks() > 0)
+	{
+		FAnimSequenceTrackPool::FPatch Patch = Tracks.Finalize([] (FAnimSequenceTrackPackedData& Data)
+		{
+			Data.ResetPatchFlags();
+		});
+
+		ENQUEUE_RENDER_COMMAND(AnimSequenceDataInstanceSubmit)([&RenderData = *RenderData, Patch = MoveTemp(Patch)](FRHICommandList&) mutable
+		{
+			RenderData.Patch(MoveTemp(Patch));
+		});
+	}
+}
+
+void UAnimSequenceTransformProviderDataInstance::Initialize(UAnimSequenceTransformProviderData* ProviderData)
+{
+	Parent    = ProviderData;
+	Sequences = ProviderData->Sequences;
+	bEnabled  = ProviderData->bEnabled;
+}
+
+void UAnimSequenceTransformProviderDataInstance::Serialize(FArchive& Ar)
+{
+	Ar.UsingCustomVersion(FAnimObjectVersion::GUID);
+	Super::Serialize(Ar);
+	Tracks.Serialize(Ar);
+
+	if (Ar.IsLoading())
+	{
+		check(RenderData);
+		RenderData->Tracks = Tracks;
+		if (Parent)
+		{
+			Initialize(Parent);
+		}
+	}
+}
+
+int32 UAnimSequenceTransformProviderDataInstance::AllocateAutoPlayTrack(const FAnimSequenceTrackAutoPlayData& Data)
+{
+	FAnimSequenceTrackPackedData PackedData = FAnimSequenceTrackPackedData::Pack(Data);
+
+	if (Sequences.IsValidIndex(Data.SequenceIndex))
+	{
+		PackedData.Wrap(GetSequencePlayLength(Data.SequenceIndex));
+		int32 TrackIndex = Tracks.AllocateTrack(PackedData);
+		if (TrackIndex != INDEX_NONE)
+		{
+			MarkOwnerRenderStateDirty();
+		}
+		return TrackIndex;
+	}
+
+	return INDEX_NONE;
+}
+
+int32 UAnimSequenceTransformProviderDataInstance::AllocateManualTrack(const FAnimSequenceTrackManualData& Data)
+{
+	FAnimSequenceTrackPackedData PackedData = FAnimSequenceTrackPackedData::Pack(Data);
+
+	if (Sequences.IsValidIndex(Data.SequenceIndex))
+	{
+		PackedData.Wrap(GetSequencePlayLength(Data.SequenceIndex));
+		int32 TrackIndex = Tracks.AllocateTrack(PackedData);
+		if (TrackIndex != INDEX_NONE)
+		{
+			MarkOwnerRenderStateDirty();
+		}
+		return TrackIndex;
+	}
+
+	return INDEX_NONE;
+}
+
+bool UAnimSequenceTransformProviderDataInstance::UpdateAutoPlayTrack(int32 TrackIndex, const FAnimSequenceTrackAutoPlayData& Data)
+{
+	FAnimSequenceTrackPackedData PackedData = FAnimSequenceTrackPackedData::Pack(Data);
+
+	if (Sequences.IsValidIndex(Data.SequenceIndex))
+	{
+		PackedData.Wrap(GetSequencePlayLength(Data.SequenceIndex));
+		if (Tracks.UpdateTrack(TrackIndex, PackedData))
+		{
+			MarkOwnerRenderStateDirty();
+			return true;
+		}
+	}
+
+	return false;
+}
+
+bool UAnimSequenceTransformProviderDataInstance::UpdateManualTrack(int32 TrackIndex, const FAnimSequenceTrackManualData& Data)
+{
+	FAnimSequenceTrackPackedData PackedData = FAnimSequenceTrackPackedData::Pack(Data);
+
+	if (Sequences.IsValidIndex(Data.SequenceIndex))
+	{
+		PackedData.Wrap(GetSequencePlayLength(Data.SequenceIndex));
+		if (Tracks.UpdateTrack(TrackIndex, PackedData))
+		{
+			MarkOwnerRenderStateDirty();
+			return true;
+		}
+	}
+
+	return false;
+}
+
+bool UAnimSequenceTransformProviderDataInstance::DeallocateTrack(int32 TrackIndex)
+{
+	if (Tracks.DeallocateTrack(TrackIndex))
+	{
+		MarkOwnerRenderStateDirty();
+		return true;
+	}
+	return false;
+}
+
+bool UAnimSequenceTransformProviderDataInstance::SetLoopMode(int32 TrackIndex, EAnimSequenceTrackLoopMode LoopMode)
+{
+	if (Tracks.IsActiveIndex(TrackIndex))
+	{
+		Tracks.GetData(TrackIndex).SetLoopMode(LoopMode);
+		Tracks.SetDirty(TrackIndex);
+		MarkOwnerRenderStateDirty();
+		return true;
+	}
+
+	return false;
+}
+
+bool UAnimSequenceTransformProviderDataInstance::SetPlayRate(int32 TrackIndex, float PlayRate)
+{
+	if (Tracks.IsActiveIndex(TrackIndex))
+	{
+		Tracks.GetData(TrackIndex).SetPlayRate(PlayRate);
+		Tracks.SetDirty(TrackIndex);
+		MarkOwnerRenderStateDirty();
+		return true;
+	}
+
+	return false;
+}
+
+float UAnimSequenceTransformProviderDataInstance::GetPlayRate(int32 TrackIndex) const
+{
+	if (Tracks.IsActiveIndex(TrackIndex))
+	{
+		return Tracks.GetData(TrackIndex).GetPlayRate();
+	}
+	return 1.0f;
+}
+
+EAnimSequenceTrackLoopMode UAnimSequenceTransformProviderDataInstance::GetLoopMode(int32 TrackIndex) const
+{
+	if (Tracks.IsActiveIndex(TrackIndex))
+	{
+		return Tracks.GetData(TrackIndex).GetLoopMode();
+	}
+	return EAnimSequenceTrackLoopMode::Loop;
+}
+
+int32 UAnimSequenceTransformProviderDataInstance::GetTrackSequenceIndex(int32 TrackIndex) const
+{
+	if (Tracks.IsActiveIndex(TrackIndex))
+	{
+		return Tracks.GetData(TrackIndex).GetSequenceIndex();
+	}
+	return INDEX_NONE;
+}
\ No newline at end of file
diff --git a/Engine/Source/Runtime/Engine/Private/Components/InstancedSkinnedMeshComponent.cpp b/Engine/Source/Runtime/Engine/Private/Components/InstancedSkinnedMeshComponent.cpp
--- a/Engine/Source/Runtime/Engine/Private/Components/InstancedSkinnedMeshComponent.cpp
+++ b/Engine/Source/Runtime/Engine/Private/Components/InstancedSkinnedMeshComponent.cpp
@@ -386,15 +386,15 @@ void UInstancedSkinnedMeshComponent::SendRenderInstanceData_Concurrent()
 {
 	Super::SendRenderInstanceData_Concurrent();
 
-	// If instance data is entirely GPU driven, don't upload from CPU.
-	if (bIsInstanceDataGPUOnly)
+	if (TransformProvider)
 	{
-		return;
+		TransformProvider->SubmitChanges();
 	}
 
-	if (TransformProvider)
+	// If instance data is entirely GPU driven, don't upload from CPU.
+	if (bIsInstanceDataGPUOnly)
 	{
-		TransformProvider->SubmitChanges();
+		return;
 	}
 
 	// If the primitive isn't hidden update its instances.
@@ -505,6 +505,8 @@ FPrimitiveSceneProxy* UInstancedSkinnedMeshComponent::CreateSceneProxy()
 		{
 			return nullptr;
 		}
+
+		TransformProvider->SubmitChanges();
 	}
 
 	if (CheckPSOPrecachingAndBoostPriority() && GetPSOPrecacheProxyCreationStrategy() == EPSOPrecacheProxyCreationStrategy::DelayUntilPSOPrecached)
diff --git a/Engine/Source/Runtime/Engine/Private/InstancedSkeletalMeshObject.cpp b/Engine/Source/Runtime/Engine/Private/InstancedSkeletalMeshObject.cpp
--- a/Engine/Source/Runtime/Engine/Private/InstancedSkeletalMeshObject.cpp
+++ b/Engine/Source/Runtime/Engine/Private/InstancedSkeletalMeshObject.cpp
@@ -26,7 +26,7 @@ FInstancedSkeletalMeshObject::FInstancedSkeletalMeshObject(
 		LODs.Emplace(InRenderData, LODIndex, InFeatureLevel);
 	}
 
-	bUseGpuLodSelection = NumLODs > 1 && InMeshDesc.bUseGpuLodSelection && GPUSKIN_TRANSFORM_STORAGE_DIRECT_ENABLED;
+	bUseGpuLodSelection = NumLODs > 1 && InMeshDesc.bUseGpuLodSelection && (GPUSKIN_TRANSFORM_STORAGE_DIRECT_ENABLED || InRenderData->bSharedLODBoneMap);
 	BoneTransformStorageMode = Type == EType::GPUSkin && !bUseGpuLodSelection ? ::GetBoneTransformStorageMode() : EBoneTransformStorageMode::Direct;
 
 	InitResources(InMeshDesc);
@@ -319,4 +319,4 @@ public:
 	}
 };
 
-REGISTER_SKELETAL_MESH_UPDATE_BACKEND(FInstancedSkeletalMeshUpdatePacket);
\ No newline at end of file
+REGISTER_SKELETAL_MESH_UPDATE_BACKEND(FInstancedSkeletalMeshUpdatePacket);
diff --git a/Engine/Source/Runtime/Engine/Private/InstancedSkinnedMeshSceneProxy.cpp b/Engine/Source/Runtime/Engine/Private/InstancedSkinnedMeshSceneProxy.cpp
--- a/Engine/Source/Runtime/Engine/Private/InstancedSkinnedMeshSceneProxy.cpp
+++ b/Engine/Source/Runtime/Engine/Private/InstancedSkinnedMeshSceneProxy.cpp
@@ -73,8 +73,7 @@ FInstancedSkinnedMeshSceneProxy::FInstancedSkinnedMeshSceneProxy(const FInstance
 
 	bAlwaysHasVelocity = true;
 	bInstancedSkinnedMesh = true;
-	bDoesMeshBatchesUseSceneInstanceCount = true;
-	bUseGpuLodSelection = NumLODs > 1 && InMeshDesc.bUseGpuLodSelection && GPUSKIN_TRANSFORM_STORAGE_DIRECT_ENABLED;
+	bUseGpuLodSelection = NumLODs > 1 && InMeshDesc.bUseGpuLodSelection && (GPUSKIN_TRANSFORM_STORAGE_DIRECT_ENABLED && InRenderData->bSharedLODBoneMap);
 
 	SetupInstanceSceneDataBuffers(Data.InstanceDataSceneProxy->GeInstanceSceneDataBuffers());
 }
\ No newline at end of file
diff --git a/Engine/Source/Runtime/Engine/Private/ObjectCacheContext.cpp b/Engine/Source/Runtime/Engine/Private/ObjectCacheContext.cpp
--- a/Engine/Source/Runtime/Engine/Private/ObjectCacheContext.cpp
+++ b/Engine/Source/Runtime/Engine/Private/ObjectCacheContext.cpp
@@ -895,7 +895,10 @@ TObjectCacheIterator<UInstancedSkinnedMeshComponent> FObjectCacheContext::GetIns
 {
 	return InstancedSkinnedMeshCache.Get<TObjectIterator<UInstancedSkinnedMeshComponent>>(InTransformProvider,
 		[](UInstancedSkinnedMeshComponent* Component) { return 1; },
-		[](UInstancedSkinnedMeshComponent* Component, int32) { return Component->GetTransformProvider(); });
+		[](UInstancedSkinnedMeshComponent* Component, int32) -> UTransformProviderData* {
+			UTransformProviderData* Provider = Component->GetTransformProvider();
+			return Provider ? Provider->GetRootTransformProvider() : nullptr;
+		});
 }
 
 TObjectCacheIterator<UMaterialInterface> FObjectCacheContext::GetUsedMaterials(IPrimitiveComponent* Component)
diff --git a/Engine/Source/Runtime/Engine/Private/SkeletalMeshRenderData.cpp b/Engine/Source/Runtime/Engine/Private/SkeletalMeshRenderData.cpp
--- a/Engine/Source/Runtime/Engine/Private/SkeletalMeshRenderData.cpp
+++ b/Engine/Source/Runtime/Engine/Private/SkeletalMeshRenderData.cpp
@@ -468,6 +468,7 @@ FSkeletalMeshRenderData::FSkeletalMeshRenderData()
 	, PendingFirstLODIdx(0)
 	, LODBiasModifier(0)
 	, bSupportRayTracing(true)
+	, bSharedLODBoneMap(false)
 	, bInitialized(false)
 {
 	ClearNaniteResources(NaniteResourcesPtr);
@@ -626,6 +627,38 @@ PRAGMA_DISABLE_DEPRECATION_WARNINGS
 PRAGMA_ENABLE_DEPRECATION_WARNINGS
 }
 
+static bool ComputeLODsShareBoneMaps(const TIndirectArray<FSkeletalMeshLODRenderData>& LODRenderData)
+{
+	const int32 NumLODs = LODRenderData.Num();
+
+	if (NumLODs <= 1)
+	{
+		return true;
+	}
+
+	const FSkeletalMeshLODRenderData& LOD0 = LODRenderData[0];
+
+	for (int32 LODIndex = 1; LODIndex < NumLODs; ++LODIndex)
+	{
+		const FSkeletalMeshLODRenderData& LODData = LODRenderData[LODIndex];
+
+		if (LODData.RenderSections.Num() != LOD0.RenderSections.Num())
+		{
+			return false;
+		}
+
+		for (int32 SectionIndex = 0; SectionIndex < LOD0.RenderSections.Num(); ++SectionIndex)
+		{
+			if (LOD0.RenderSections[SectionIndex].BoneMap != LODData.RenderSections[SectionIndex].BoneMap)
+			{
+				return false;
+			}
+		}
+	}
+
+	return true;
+}
+
 void FSkeletalMeshRenderData::InitResources(bool bNeedsVertexColors, USkinnedAsset* Owner)
 {
 	if (!bInitialized)
@@ -652,6 +685,8 @@ void FSkeletalMeshRenderData::InitResources(bool bNeedsVertexColors, USkinnedAss
 			}
 		}
 
+		bSharedLODBoneMap = ComputeLODsShareBoneMaps(LODRenderData);
+
 		check(NaniteResourcesPtr.IsValid());
 		NaniteResourcesPtr->InitResources(Owner);
 
diff --git a/Engine/Source/Runtime/Engine/Public/Animation/SkeletonRemapping.h b/Engine/Source/Runtime/Engine/Public/Animation/SkeletonRemapping.h
--- a/Engine/Source/Runtime/Engine/Public/Animation/SkeletonRemapping.h
+++ b/Engine/Source/Runtime/Engine/Public/Animation/SkeletonRemapping.h
@@ -163,6 +163,17 @@ struct FSkeletonRemapping
 		return SourceSkeleton.IsValid() && TargetSkeleton.IsValid();
 	}
 
+	/**
+	 * Get the retargeting quaternion pair for a target skeleton bone
+	 *
+	 * @param   TargetSkeletonBoneIndex     Skeleton bone index on the target skeleton
+	 * @return                              Tuple of (Q0, Q1) where RetargetedRotation = Q0 * SourceRotation * Q1
+	 */
+	const TTuple<FQuat, FQuat>& GetRetargetingQuaternions(int32 TargetSkeletonBoneIndex) const
+	{
+		return RetargetingTable[TargetSkeletonBoneIndex];
+	}
+
 private:
 	/** Internal helper function - performs mapping generation */
 	ENGINE_API void GenerateMapping();
diff --git a/Engine/Source/Runtime/Engine/Public/Animation/TransformProviderData.h b/Engine/Source/Runtime/Engine/Public/Animation/TransformProviderData.h
--- a/Engine/Source/Runtime/Engine/Public/Animation/TransformProviderData.h
+++ b/Engine/Source/Runtime/Engine/Public/Animation/TransformProviderData.h
@@ -81,6 +81,11 @@ public:
 	{
 	}
 
+	virtual UTransformProviderData* GetRootTransformProvider()
+	{
+		return this;
+	}
+
 #if WITH_EDITOR
 	void BeginCacheForCookedPlatformData(const ITargetPlatform* TargetPlatform)
 	{
diff --git a/Engine/Source/Runtime/Engine/Public/Rendering/SkeletalMeshRenderData.h b/Engine/Source/Runtime/Engine/Public/Rendering/SkeletalMeshRenderData.h
--- a/Engine/Source/Runtime/Engine/Public/Rendering/SkeletalMeshRenderData.h
+++ b/Engine/Source/Runtime/Engine/Public/Rendering/SkeletalMeshRenderData.h
@@ -43,6 +43,9 @@ public:
 	/** Whether ray tracing acceleration structures should be created for this mesh. Derived from owner USkinnedAsset. */
 	bool bSupportRayTracing;
 
+	/** True if all LODs share the same section bone map (enables instanced skinned GPU LOD without indirection) */
+	bool bSharedLODBoneMap;
+
 	/** Optional, name of the owning skeletal mesh */
 	FName OwnerName = NAME_None;
 
diff --git a/Engine/Source/Runtime/Engine/Public/SceneInterface.h b/Engine/Source/Runtime/Engine/Public/SceneInterface.h
--- a/Engine/Source/Runtime/Engine/Public/SceneInterface.h
+++ b/Engine/Source/Runtime/Engine/Public/SceneInterface.h
@@ -15,6 +15,7 @@
 
 class AWorldSettings;
 class IAnimBankTransformProvider;
+class IAnimSequenceTransformProvider;
 class IAnimRuntimeTransformProvider;
 class FArchive;
 class FFloat16Color;
@@ -434,6 +435,7 @@ public:
 
 	virtual IAnimBankTransformProvider* GetAnimBankTransformProvider() = 0;
 	virtual IAnimRuntimeTransformProvider* GetAnimRuntimeTransformProvider() = 0;
+	virtual IAnimSequenceTransformProvider* GetAnimSequenceTransformProvider() = 0;
 
 	/**
 	 * Adds the unique volumetric cloud component to the scene
diff --git a/Engine/Source/Runtime/Engine/Public/SkinningSceneExtensionProxy.h b/Engine/Source/Runtime/Engine/Public/SkinningSceneExtensionProxy.h
--- a/Engine/Source/Runtime/Engine/Public/SkinningSceneExtensionProxy.h
+++ b/Engine/Source/Runtime/Engine/Public/SkinningSceneExtensionProxy.h
@@ -148,6 +148,12 @@ protected:
 	friend class FSkinningSceneExtension;
 };
 
+class FInstancedSkinningSceneExtensionProxyUserData
+{
+public:
+	virtual ~FInstancedSkinningSceneExtensionProxyUserData() = default;
+};
+
 class FInstancedSkinningSceneExtensionProxy : public FSkinningSceneExtensionProxy
 {
 public:
@@ -168,7 +174,24 @@ public:
 		return TransformProviderProxy;
 	}
 
+	void SetUserData(FInstancedSkinningSceneExtensionProxyUserData* InUserData)
+	{
+		UserData = InUserData;
+	}
+
+	FInstancedSkinningSceneExtensionProxyUserData* GetUserData() const
+	{
+		return UserData;
+	}
+
+	FInstancedSkinningSceneExtensionProxyUserData& GetUserDataChecked() const
+	{
+		check(UserData);
+		return *UserData;
+	}
+
 protected:
 	FTransformProviderRenderProxy* TransformProviderProxy = nullptr;
+	FInstancedSkinningSceneExtensionProxyUserData* UserData = nullptr;
 	FGuid TransformProviderId;
 };
\ No newline at end of file
diff --git a/Engine/Source/Runtime/RenderCore/Private/UnifiedBuffer.cpp b/Engine/Source/Runtime/RenderCore/Private/UnifiedBuffer.cpp
--- a/Engine/Source/Runtime/RenderCore/Private/UnifiedBuffer.cpp
+++ b/Engine/Source/Runtime/RenderCore/Private/UnifiedBuffer.cpp
@@ -1056,37 +1056,38 @@ FRDGScatterUploadBuilder* FRDGScatterUploadBuilder::Create(FRDGBuilder& GraphBui
 
 UE::Tasks::FTask FRDGScatterUploadBuilder::Execute(FRDGBuilder& GraphBuilder)
 {
-	if (Passes.IsEmpty() && CommandListFunctions.IsEmpty())
-	{
-		return {};
-	}
-
-	const bool bUseAsyncTask = MaxBytes > 32 * 1024 || !CommandListFunctions.IsEmpty();
+	UE::Tasks::FTask Task;
 
-	UE::Tasks::FTask Task = GraphBuilder.AddCommandListSetupTask([Passes = Passes, CommandListFunctions = MoveTemp(CommandListFunctions)] (FRHICommandList& RHICmdList)
+	if (!Passes.IsEmpty() || !CommandListFunctions.IsEmpty())
 	{
-		TRACE_CPUPROFILER_EVENT_SCOPE(FRDGScatterUploadBuilder);
-		for (auto& Pass : Passes)
-		{
-			Pass.Uploader->Lock(RHICmdList);
-			Pass.Function(*Pass.Uploader);
-			Pass.Uploader->Unlock(RHICmdList);
-		}
+		const bool bUseAsyncTask = MaxBytes > 32 * 1024 || !CommandListFunctions.IsEmpty();
 
-		for (auto& CommandListFunction : CommandListFunctions)
+		Task = GraphBuilder.AddCommandListSetupTask([Passes = Passes, CommandListFunctions = MoveTemp(CommandListFunctions)] (FRHICommandList& RHICmdList)
 		{
-			CommandListFunction(RHICmdList);
-		}
+			TRACE_CPUPROFILER_EVENT_SCOPE(FRDGScatterUploadBuilder);
+			for (auto& Pass : Passes)
+			{
+				Pass.Uploader->Lock(RHICmdList);
+				Pass.Function(*Pass.Uploader);
+				Pass.Uploader->Unlock(RHICmdList);
+			}
+
+			for (auto& CommandListFunction : CommandListFunctions)
+			{
+				CommandListFunction(RHICmdList);
+			}
 
-	}, bUseAsyncTask);
+		}, PrerequisiteTasks, UE::Tasks::ETaskPriority::Normal, bUseAsyncTask);
+	}
 
 	for (auto& Pass : Passes)
 	{
 		Pass.UploadBuffer.End(GraphBuilder, Pass.Uploader);
 	}
 
-	Passes.Reset();
-	CommandListFunctions.Reset();
+	Passes.Empty();
+	CommandListFunctions.Empty();
+	PrerequisiteTasks.Empty();
 	return Task;
 }
 
diff --git a/Engine/Source/Runtime/RenderCore/Public/UnifiedBuffer.h b/Engine/Source/Runtime/RenderCore/Public/UnifiedBuffer.h
--- a/Engine/Source/Runtime/RenderCore/Public/UnifiedBuffer.h
+++ b/Engine/Source/Runtime/RenderCore/Public/UnifiedBuffer.h
@@ -463,6 +463,11 @@ public:
 		CommandListFunctions.Emplace(MoveTemp(InFunction));
 	}
 
+	void AddPrerequisiteTask(const UE::Tasks::FTask& Task)
+	{
+		PrerequisiteTasks.Emplace(Task);
+	}
+
 	RENDERCORE_API UE::Tasks::FTask Execute(FRDGBuilder& GraphBuilder);
 
 private:
@@ -483,6 +488,7 @@ private:
 
 	TArray<FPass, SceneRenderingAllocator> Passes;
 	TArray<FCommandListFunction, SceneRenderingAllocator> CommandListFunctions;
+	TArray<UE::Tasks::FTask, SceneRenderingAllocator> PrerequisiteTasks;
 	uint32 MaxBytes = 0;
 
 	RDG_FRIEND_ALLOCATOR_FRIEND(FRDGScatterUploadBuilder);
diff --git a/Engine/Source/Runtime/Renderer/Private/RendererScene.cpp b/Engine/Source/Runtime/Renderer/Private/RendererScene.cpp
--- a/Engine/Source/Runtime/Renderer/Private/RendererScene.cpp
+++ b/Engine/Source/Runtime/Renderer/Private/RendererScene.cpp
@@ -103,6 +103,7 @@
 #include "Animation/AnimBank.h"
 #include "Skinning/AnimBankTransformProvider.h"
 #include "Skinning/AnimRuntimeTransformProvider.h"
+#include "Skinning/AnimSequenceTransformProvider.h"
 #include "Skinning/SkinningTransformProvider.h"
 #include "PathTracing.h"
 
@@ -2769,6 +2770,11 @@ IAnimRuntimeTransformProvider* FScene::GetAnimRuntimeTransformProvider()
 	return SceneExtensions.GetExtensionPtr<FAnimRuntimeTransformProvider>();
 }
 
+IAnimSequenceTransformProvider* FScene::GetAnimSequenceTransformProvider()
+{
+	return SceneExtensions.GetExtensionPtr<FAnimSequenceTransformProvider>();
+}
+
 void FScene::GetLightIESAtlasSlot(const FLightSceneProxy* Proxy, FLightRenderParameters* Out)
 {
 	if (Proxy)
@@ -6850,6 +6856,7 @@ public:
 
 	virtual IAnimBankTransformProvider* GetAnimBankTransformProvider() { return nullptr; }
 	virtual IAnimRuntimeTransformProvider* GetAnimRuntimeTransformProvider() { return nullptr; }
+	virtual IAnimSequenceTransformProvider* GetAnimSequenceTransformProvider() { return nullptr; }
 
 	virtual void AddHairStrands(FHairStrandsInstance* Proxy) override {}
 	virtual void RemoveHairStrands(FHairStrandsInstance* Proxy) override {}
diff --git a/Engine/Source/Runtime/Renderer/Private/ScenePrivate.h b/Engine/Source/Runtime/Renderer/Private/ScenePrivate.h
--- a/Engine/Source/Runtime/Renderer/Private/ScenePrivate.h
+++ b/Engine/Source/Runtime/Renderer/Private/ScenePrivate.h
@@ -3513,6 +3513,7 @@ public:
 
 	virtual IAnimBankTransformProvider* GetAnimBankTransformProvider() override;
 	virtual IAnimRuntimeTransformProvider* GetAnimRuntimeTransformProvider() override;
+	virtual IAnimSequenceTransformProvider* GetAnimSequenceTransformProvider() override;
 
 	virtual void SetPhysicsField(class FPhysicsFieldSceneProxy* PhysicsFieldSceneProxy) override;
 	virtual void ResetPhysicsField() override;
diff --git a/Engine/Source/Runtime/Renderer/Private/Skinning/AnimSequenceTransformProvider.cpp b/Engine/Source/Runtime/Renderer/Private/Skinning/AnimSequenceTransformProvider.cpp
new file mode 100644
--- /dev/null
+++ b/Engine/Source/Runtime/Renderer/Private/Skinning/AnimSequenceTransformProvider.cpp
@@ -0,0 +1,1829 @@
+// Copyright Epic Games, Inc. All Rights Reserved.
+
+#include "AnimSequenceTransformProvider.h"
+#include "Animation/AnimBoneCompressionCodec.h"
+#include "Animation/AnimSequence.h"
+#include "Animation/AnimSequenceDecompressionContext.h"
+#include "Animation/AttributesRuntime.h"
+#include "Animation/SkeletonRemapping.h"
+#include "Animation/SkeletonRemappingRegistry.h"
+#include "Animation/Skeleton.h"
+#include "AnimationRuntime.h"
+#include "AnimEncoding.h"
+#include "ComponentRecreateRenderStateContext.h"
+#include "GlobalRenderResources.h"
+#include "GlobalShader.h"
+#include "RenderGraphBuilder.h"
+#include "RenderGraphResources.h"
+#include "RenderGraphUtils.h"
+#include "RenderingThread.h"
+#include "RenderUtils.h"
+#include "SceneInterface.h"
+#include "ScenePrivate.h"
+#include "SceneView.h"
+#include "SkeletalRenderPublic.h"
+#include "SkinningDefinitions.h"
+#include "RendererModule.h"
+
+DECLARE_STATS_GROUP(TEXT("AnimSequenceTransformProvider"), STATGROUP_AnimSequenceTransformProvider, STATCAT_Advanced);
+
+DECLARE_DWORD_ACCUMULATOR_STAT(TEXT("Sequences"), STAT_AnimSequenceTransformProvider_Sequences, STATGROUP_AnimSequenceTransformProvider);
+DECLARE_DWORD_ACCUMULATOR_STAT(TEXT("Mappings"), STAT_AnimSequenceTransformProvider_Mappings, STATGROUP_AnimSequenceTransformProvider);
+
+DECLARE_MEMORY_STAT(TEXT("Sequence Transform Buffer"), STAT_AnimSequenceTransformProvider_SequenceTransformMemory, STATGROUP_AnimSequenceTransformProvider);
+DECLARE_MEMORY_STAT(TEXT("Retargeting Data Buffer"), STAT_AnimSequenceTransformProvider_RetargetingDataMemory, STATGROUP_AnimSequenceTransformProvider);
+DECLARE_MEMORY_STAT(TEXT("Leaf Bone Path Buffer"), STAT_AnimSequenceTransformProvider_LeafBonePathMemory, STATGROUP_AnimSequenceTransformProvider);
+
+DECLARE_DWORD_COUNTER_STAT(TEXT("Active Tracks"), STAT_AnimSequenceTransformProvider_ActiveTracks, STATGROUP_AnimSequenceTransformProvider);
+DECLARE_DWORD_COUNTER_STAT(TEXT("Evaluate Blocks"), STAT_AnimSequenceTransformProvider_EvaluateBlocks, STATGROUP_AnimSequenceTransformProvider);
+DECLARE_DWORD_COUNTER_STAT(TEXT("Evaluate Blend Blocks"), STAT_AnimSequenceTransformProvider_EvaluateBlendBlocks, STATGROUP_AnimSequenceTransformProvider);
+DECLARE_DWORD_COUNTER_STAT(TEXT("Current Transform Updates"), STAT_AnimSequenceTransformProvider_CurrentTransformUpdates, STATGROUP_AnimSequenceTransformProvider);
+DECLARE_DWORD_COUNTER_STAT(TEXT("Previous Transform Updates"), STAT_AnimSequenceTransformProvider_PreviousTransformUpdates, STATGROUP_AnimSequenceTransformProvider);
+DECLARE_DWORD_COUNTER_STAT(TEXT("Sequence Frame Uploads"), STAT_AnimSequenceTransformProvider_SequenceFrameUploads, STATGROUP_AnimSequenceTransformProvider);
+
+DECLARE_GPU_STAT(AnimSequenceTransformProvider);
+
+static FGuid AnimSequenceGPUProviderId(ANIM_SEQUENCE_GPU_TRANSFORM_PROVIDER_GUID);
+static FGuid AnimSequenceCPUProviderId(ANIM_SEQUENCE_CPU_TRANSFORM_PROVIDER_GUID);
+
+IMPLEMENT_SCENE_EXTENSION(FAnimSequenceTransformProvider);
+
+static TAutoConsoleVariable<float> CVarAnimSequenceTransformProviderTimeScale(
+	TEXT("r.AnimSequenceTransformProvider.TimeScale"),
+	1.0f,
+	TEXT(""),
+	ECVF_Scalability | ECVF_RenderThreadSafe
+);
+
+static TAutoConsoleVariable<bool> CVarAnimSequenceTransformProviderGPU(
+	TEXT("r.AnimSequenceTransformProvider.GPU"),
+	false,
+	TEXT(""),
+	FConsoleVariableDelegate::CreateLambda([](IConsoleVariable* InVariable)
+	{
+		FGlobalComponentRecreateRenderStateContext Context;
+	}),
+	ECVF_Scalability | ECVF_RenderThreadSafe
+);
+
+static TAutoConsoleVariable<int32> CVarAnimSequenceTransformProviderUploadBudget(
+	TEXT("r.AnimSequenceTransformProvider.UploadBudget"),
+	32,
+	TEXT("Maximum number of sequence frames to upload per frame across all sequences.\n")
+	TEXT("Higher values reduce time to full quality but increase per-frame cost."),
+	ECVF_Scalability | ECVF_RenderThreadSafe
+);
+
+static TAutoConsoleVariable<int32> CVarAnimSequenceTransformProviderMaxSampleRateLevels(
+	TEXT("r.AnimSequenceTransformProvider.MaxSampleRateLevels"),
+	3,
+	TEXT("Number of sample rate levels for progressive sequence upload.\n")
+	TEXT("1 = no progressive upload (full rate immediately)\n")
+	TEXT("2 = two levels (half rate, then full)\n")
+	TEXT("3 = three levels (quarter, half, full) [default]"),
+	FConsoleVariableDelegate::CreateLambda([](IConsoleVariable* InVariable)
+	{
+		FGlobalComponentRecreateRenderStateContext Context;
+	}),
+	ECVF_Scalability | ECVF_RenderThreadSafe
+);
+
+static TAutoConsoleVariable<int32> CVarAnimSequenceTransformProviderForceSampleRateLevel(
+	TEXT("r.AnimSequenceTransformProvider.ForceSampleRateLevel"),
+	-1,
+	TEXT("Force a minimum sample rate level for debugging.\n")
+	TEXT("-1 = disabled (use actual completed level)\n")
+	TEXT("0 = full rate\n")
+	TEXT("1 = half rate\n")
+	TEXT("2 = quarter rate"),
+	ECVF_Scalability | ECVF_RenderThreadSafe
+);
+
+static TAutoConsoleVariable<int32> CVarAnimSequenceTransformProviderForceSampleRate(
+	TEXT("r.AnimSequenceTransformProvider.ForceSampleRate"),
+	-1,
+	TEXT("Force a specific sample rate for every anim sequence in frames. Clamped from 0 to the sequence sample rate.\n"),
+	FConsoleVariableDelegate::CreateLambda([](IConsoleVariable* InVariable)
+	{
+		FGlobalComponentRecreateRenderStateContext Context;
+	}),
+	ECVF_Scalability | ECVF_RenderThreadSafe
+);
+
+FFrameRate GetAnimSequenceSampleRate(const UAnimSequence* AnimSequence)
+{
+	const int32 ForcedSampleRate = CVarAnimSequenceTransformProviderForceSampleRate.GetValueOnRenderThread();
+	if (ForcedSampleRate > 0)
+	{
+		FFrameRate SampleRate(ForcedSampleRate, 1);
+		SampleRate.Numerator = FMath::Min(SampleRate.Numerator, AnimSequence->GetSamplingFrameRate().Numerator);
+		return SampleRate;
+	}
+	return AnimSequence->GetSamplingFrameRate();
+}
+
+uint32 GetMaxSampleRateLevels(uint32 NumFrames)
+{
+	if (NumFrames <= 1)
+	{
+		return 1;
+	}
+
+	const uint32 MaxLevels = FMath::Max(1, CVarAnimSequenceTransformProviderMaxSampleRateLevels.GetValueOnRenderThread());
+
+	return FMath::Min(MaxLevels, FMath::FloorLog2(NumFrames - 1) + 1);
+}
+
+uint32 FAnimSequenceTransformProvider::FSequence::GetSampleStep() const
+{
+	const uint32 ForcedSampleRateLevel = CVarAnimSequenceTransformProviderForceSampleRateLevel.GetValueOnRenderThread();
+	uint32 SampleRateLevel = ForcedSampleRateLevel != INDEX_NONE ? ForcedSampleRateLevel : CurrentSampleRateLevel;
+	return 1u << FMath::Clamp<uint32>(SampleRateLevel, CurrentSampleRateLevel, NumSampleRateLevels - 1);
+}
+
+void FAnimSequenceTransformProvider::FSequence::GetKeyIndicesFromTime(int32& OutKey0, int32& OutKey1, float& OutAlpha, float Time) const
+{
+	const auto AlignDown = [] (int32 Value, int32 Alignment)
+	{
+		return (Value / Alignment) * Alignment;
+	};
+
+	const int32 Step = GetSampleStep();
+	const float ExactFrame = Time * NumFrames / PlayLength;
+	const int32 LastAlignedFrame = AlignDown((int32)NumFrames - 1, Step);
+
+	OutKey0  = FMath::Min(AlignDown((int32)ExactFrame, Step), LastAlignedFrame);
+	OutKey1  = FMath::Min(OutKey0 + Step, LastAlignedFrame);
+	OutAlpha = OutKey1 > OutKey0 ? (ExactFrame - (float)OutKey0) / (float)Step : 0.0f;
+}
+
+class FEvaluateAnimSequenceCS : public FGlobalShader
+{
+public:
+	DECLARE_GLOBAL_SHADER(FEvaluateAnimSequenceCS);
+	SHADER_USE_PARAMETER_STRUCT(FEvaluateAnimSequenceCS, FGlobalShader);
+
+	static constexpr uint32 BonesPerGroup = 64u;
+
+	static bool ShouldCompilePermutation(const FGlobalShaderPermutationParameters& Parameters)
+	{
+		return true;
+	}
+
+	static void ModifyCompilationEnvironment(const FGlobalShaderPermutationParameters& Parameters, FShaderCompilerEnvironment& OutEnvironment)
+	{
+		FGlobalShader::ModifyCompilationEnvironment(Parameters, OutEnvironment);
+		OutEnvironment.SetDefine(TEXT("BONES_PER_GROUP"), BonesPerGroup);
+
+		OutEnvironment.CompilerFlags.Add(CFLAG_WarningsAsErrors);
+		OutEnvironment.CompilerFlags.Add(CFLAG_HLSL2021);
+	}
+
+	BEGIN_SHADER_PARAMETER_STRUCT(FParameters, )
+		SHADER_PARAMETER_RDG_BUFFER_SRV(ByteAddressBuffer, RetargetingDataBuffer)
+		SHADER_PARAMETER_RDG_BUFFER_SRV(ByteAddressBuffer, SequenceTransformBuffer)
+		SHADER_PARAMETER_RDG_BUFFER_SRV(ByteAddressBuffer, EvaluateBlockHeaderBuffer)
+		SHADER_PARAMETER_RDG_BUFFER_UAV(RWByteAddressBuffer, LocalSpaceTransformBufferUAV)
+	END_SHADER_PARAMETER_STRUCT()
+};
+IMPLEMENT_GLOBAL_SHADER(FEvaluateAnimSequenceCS, "/Engine/Private/Skinning/AnimSequenceTransformProvider.usf", "EvaluateAnimSequenceCS", SF_Compute);
+
+class FEvaluateAnimSequenceBlendCS : public FGlobalShader
+{
+public:
+	DECLARE_GLOBAL_SHADER(FEvaluateAnimSequenceBlendCS);
+	SHADER_USE_PARAMETER_STRUCT(FEvaluateAnimSequenceBlendCS, FGlobalShader);
+
+	static constexpr uint32 BonesPerGroup = 64u;
+
+	static bool ShouldCompilePermutation(const FGlobalShaderPermutationParameters& Parameters)
+	{
+		return true;
+	}
+
+	static void ModifyCompilationEnvironment(const FGlobalShaderPermutationParameters& Parameters, FShaderCompilerEnvironment& OutEnvironment)
+	{
+		FGlobalShader::ModifyCompilationEnvironment(Parameters, OutEnvironment);
+		OutEnvironment.SetDefine(TEXT("BONES_PER_GROUP"), BonesPerGroup);
+
+		OutEnvironment.CompilerFlags.Add(CFLAG_WarningsAsErrors);
+		OutEnvironment.CompilerFlags.Add(CFLAG_HLSL2021);
+	}
+
+	BEGIN_SHADER_PARAMETER_STRUCT(FParameters, )
+		SHADER_PARAMETER_RDG_BUFFER_SRV(ByteAddressBuffer, RetargetingDataBuffer)
+		SHADER_PARAMETER_RDG_BUFFER_SRV(ByteAddressBuffer, SequenceTransformBuffer)
+		SHADER_PARAMETER_RDG_BUFFER_SRV(ByteAddressBuffer, EvaluateBlockHeaderBuffer)
+		SHADER_PARAMETER_RDG_BUFFER_UAV(RWByteAddressBuffer, LocalSpaceTransformBufferUAV)
+	END_SHADER_PARAMETER_STRUCT()
+};
+IMPLEMENT_GLOBAL_SHADER(FEvaluateAnimSequenceBlendCS, "/Engine/Private/Skinning/AnimSequenceTransformProvider.usf", "EvaluateAnimSequenceBlendCS", SF_Compute);
+
+class FConcatenateLocalToGlobalCS : public FGlobalShader
+{
+public:
+	DECLARE_GLOBAL_SHADER(FConcatenateLocalToGlobalCS);
+	SHADER_USE_PARAMETER_STRUCT(FConcatenateLocalToGlobalCS, FGlobalShader);
+
+	static constexpr uint32 BonesPerGroup = FEvaluateAnimSequenceCS::BonesPerGroup;
+
+	static bool ShouldCompilePermutation(const FGlobalShaderPermutationParameters& Parameters)
+	{
+		return true;
+	}
+
+	static void ModifyCompilationEnvironment(const FGlobalShaderPermutationParameters& Parameters, FShaderCompilerEnvironment& OutEnvironment)
+	{
+		FGlobalShader::ModifyCompilationEnvironment(Parameters, OutEnvironment);
+		OutEnvironment.SetDefine(TEXT("BONES_PER_GROUP"), BonesPerGroup);
+
+		OutEnvironment.CompilerFlags.Add(CFLAG_WarningsAsErrors);
+		OutEnvironment.CompilerFlags.Add(CFLAG_HLSL2021);
+	}
+
+	BEGIN_SHADER_PARAMETER_STRUCT(FParameters, )
+		SHADER_PARAMETER_RDG_BUFFER_SRV(ByteAddressBuffer, LocalSpaceTransformBuffer)
+		SHADER_PARAMETER_RDG_BUFFER_SRV(ByteAddressBuffer, ConcatenateBlockHeaderBuffer)
+		SHADER_PARAMETER_RDG_BUFFER_SRV(ByteAddressBuffer, LeafBonePathBuffer)
+		SHADER_PARAMETER_RDG_BUFFER_UAV(RWByteAddressBuffer, GlobalSpaceTransformBufferUAV)
+	END_SHADER_PARAMETER_STRUCT()
+};
+IMPLEMENT_GLOBAL_SHADER(FConcatenateLocalToGlobalCS, "/Engine/Private/Skinning/AnimSequenceTransformProvider.usf", "ConcatenateLocalToGlobalCS", SF_Compute);
+
+class FScatterTransformsCS : public FGlobalShader
+{
+public:
+	DECLARE_GLOBAL_SHADER(FScatterTransformsCS);
+	SHADER_USE_PARAMETER_STRUCT(FScatterTransformsCS, FGlobalShader);
+
+	static constexpr uint32 BonesPerGroup = FEvaluateAnimSequenceCS::BonesPerGroup;
+
+	static bool ShouldCompilePermutation(const FGlobalShaderPermutationParameters& Parameters)
+	{
+		return true;
+	}
+
+	static void ModifyCompilationEnvironment(const FGlobalShaderPermutationParameters& Parameters, FShaderCompilerEnvironment& OutEnvironment)
+	{
+		FGlobalShader::ModifyCompilationEnvironment(Parameters, OutEnvironment);
+		OutEnvironment.SetDefine(TEXT("BONES_PER_GROUP"), BonesPerGroup);
+
+		OutEnvironment.CompilerFlags.Add(CFLAG_WarningsAsErrors);
+		OutEnvironment.CompilerFlags.Add(CFLAG_HLSL2021);
+	}
+
+	BEGIN_SHADER_PARAMETER_STRUCT(FParameters, )
+		SHADER_PARAMETER_RDG_BUFFER_SRV(ByteAddressBuffer, GlobalSpaceTransformBuffer)
+		SHADER_PARAMETER_RDG_BUFFER_SRV(ByteAddressBuffer, SequenceTransformBuffer)
+		SHADER_PARAMETER_RDG_BUFFER_SRV(ByteAddressBuffer, BoneMapBuffer)
+		SHADER_PARAMETER_RDG_BUFFER_SRV(ByteAddressBuffer, ScatterBlockHeaderBuffer)
+		SHADER_PARAMETER_RDG_COMPRESSED_BONE_TRANSFORM_UAV(SkinningTransformBufferUAV)
+	END_SHADER_PARAMETER_STRUCT()
+};
+IMPLEMENT_GLOBAL_SHADER(FScatterTransformsCS, "/Engine/Private/Skinning/AnimSequenceTransformProvider.usf", "ScatterTransformsCS", SF_Compute);
+
+FAnimSequenceTransformProvider::FProxyUserData::FProxyUserData(FAnimSequenceTransformProviderProxy* Proxy)
+	: Tracks(Proxy->GetTracks())
+	, SceneExtensionProxy(*Proxy->GetSceneExtensionProxy())
+	, TransformStorageMode(SceneExtensionProxy.GetBoneTransformStorageMode())
+{}
+
+FAnimSequenceTransformProvider::FAnimSequenceTransformProvider(FScene& InScene)
+	: ISceneExtension(InScene)
+	, SequenceTransformPersistentBuffer(0, TEXT("AnimSequenceTransformProvider.SequenceTransformBuffer"))
+	, RetargetingDataPersistentBuffer(0, TEXT("AnimSequenceTransformProvider.RetargetingDataBuffer"))
+	, LeafBonePathPersistentBuffer(0, TEXT("AnimSequenceTransformProvider.LeafBonePathBuffer"))
+{}
+
+bool FAnimSequenceTransformProvider::ShouldCreateExtension(FScene& InScene)
+{
+	return IsGPUSkinSceneExtensionEnabled();
+}
+
+void FAnimSequenceTransformProvider::InitExtension(FScene& InScene)
+{
+	if (auto TransformProvider = InScene.GetExtensionPtr<FSkinningTransformProvider>())
+	{
+		// Register GPU transform provider
+		TransformProvider->RegisterProvider(
+			AnimSequenceGPUProviderId,
+			FSkinningTransformProvider::FOnProvideTransforms::CreateRaw(this, &FAnimSequenceTransformProvider::ProvideGPUTransforms),
+			false /* Use skeleton batching */
+		);
+
+		// Register CPU transform provider
+		TransformProvider->RegisterProvider(
+			AnimSequenceCPUProviderId,
+			FSkinningTransformProvider::FOnProvideTransforms::CreateRaw(this, &FAnimSequenceTransformProvider::ProvideCPUTransforms),
+			false /* Use skeleton batching */
+		);
+	}
+}
+
+void FAnimSequenceTransformProvider::BuildBoneRetargetingTables(
+	FBoneRetargetingTables& BoneRetargetingTables,
+	USkeleton* AnimSkeleton,
+	const USkinnedAsset* MeshAsset)
+{
+	const USkeleton* TargetSkeleton = MeshAsset->GetSkeleton();
+	const FReferenceSkeleton& MeshRefSkeleton = MeshAsset->GetRefSkeleton();
+	const FReferenceSkeleton& TargetSkeletonRefSkeleton = TargetSkeleton->GetReferenceSkeleton();
+	const int32 NumMeshBones = MeshRefSkeleton.GetRawBoneNum();
+
+	const TArray<FTransform>& AnimRefPose = AnimSkeleton->GetRefLocalPoses();
+	const TArray<FTransform>& MeshRefPose = MeshRefSkeleton.GetRefBonePose();
+	const bool bUseSourceRetargetModes = TargetSkeleton->GetUseRetargetModesFromCompatibleSkeleton();
+
+	const FSkeletonRemapping& SkeletonRemapping = UE::Anim::FSkeletonRemappingRegistry::Get().GetRemapping(AnimSkeleton, TargetSkeleton);
+	const bool bNeedsRemapping = SkeletonRemapping.IsValid();
+	const bool bRequiresRefPoseRetarget = bNeedsRemapping && SkeletonRemapping.RequiresReferencePoseRetarget();
+
+	BoneRetargetingTables.Mode.Reset(NumMeshBones);
+	BoneRetargetingTables.Transform.Reset(NumMeshBones);
+	BoneRetargetingTables.AnimBoneIndex.Reset(NumMeshBones);
+	BoneRetargetingTables.bRequiresRefPoseRetarget = bRequiresRefPoseRetarget;
+
+	if (bRequiresRefPoseRetarget)
+	{
+		BoneRetargetingTables.RefPoseRetargetQuats.Reset(NumMeshBones);
+	}
+
+	for (int32 MeshBoneIndex = 0; MeshBoneIndex < NumMeshBones; ++MeshBoneIndex)
+	{
+		EBoneTranslationRetargetingMode::Type Mode = EBoneTranslationRetargetingMode::Animation;
+
+		const FName BoneName = MeshRefSkeleton.GetBoneName(MeshBoneIndex);
+		const int32 TargetSkeletonBoneIndex = TargetSkeletonRefSkeleton.FindBoneIndex(BoneName);
+
+		const int32 AnimBoneIndex = (bNeedsRemapping && TargetSkeletonBoneIndex != INDEX_NONE)
+			? SkeletonRemapping.GetSourceSkeletonBoneIndex(TargetSkeletonBoneIndex)
+			: TargetSkeletonBoneIndex;
+
+		const FTransform& MeshRefTransform = MeshRefPose[MeshBoneIndex];
+		TTuple<FQuat, FQuat> RefPoseRetargetQuats(FQuat::Identity, FQuat::Identity);
+		FTransform3f Transform;
+
+		if (AnimBoneIndex != INDEX_NONE)
+		{
+			if (bRequiresRefPoseRetarget)
+			{
+				RefPoseRetargetQuats = SkeletonRemapping.GetRetargetingQuaternions(TargetSkeletonBoneIndex);
+			}
+
+			Mode = FAnimationRuntime::GetBoneTranslationRetargetingMode(
+				bUseSourceRetargetModes,
+				AnimBoneIndex,
+				TargetSkeletonBoneIndex,
+				AnimSkeleton,
+				TargetSkeleton,
+				false);
+
+			if (MeshBoneIndex == 0)
+			{
+				// Lock root bone translation to disable root motion.
+				Mode = EBoneTranslationRetargetingMode::Skeleton;
+			}
+
+			const FTransform& AnimRefTransform = AnimRefPose[AnimBoneIndex];
+
+			switch (Mode)
+			{
+				case EBoneTranslationRetargetingMode::Animation:
+					break;
+				case EBoneTranslationRetargetingMode::Skeleton:
+					Transform.SetTranslation(FVector3f(MeshRefTransform.GetTranslation()));
+					if (MeshBoneIndex == 0)
+					{
+						Transform.SetRotation(FQuat4f(MeshRefTransform.GetRotation()));
+					}
+					break;
+				case EBoneTranslationRetargetingMode::AnimationScaled:
+				{
+					const float AnimTranslationLength = AnimRefTransform.GetTranslation().Size();
+					const float MeshTranslationLength = MeshRefTransform.GetTranslation().Size();
+					if (AnimTranslationLength > UE_KINDA_SMALL_NUMBER)
+					{
+						Transform.SetScale3D(FVector3f(MeshTranslationLength / AnimTranslationLength));
+					}
+				}
+				break;
+				case EBoneTranslationRetargetingMode::AnimationRelative:
+				{
+					FTransform RetargetedAnimRefTransform = AnimRefTransform;
+
+					if (bRequiresRefPoseRetarget)
+					{
+						FQuat RetargetedRotation = RefPoseRetargetQuats.Get<0>() * AnimRefTransform.GetRotation() * RefPoseRetargetQuats.Get<1>();
+						RetargetedRotation.Normalize();
+						RetargetedAnimRefTransform.SetRotation(RetargetedRotation);
+						RetargetedAnimRefTransform.SetTranslation(RefPoseRetargetQuats.Get<0>().RotateVector(AnimRefTransform.GetTranslation()));
+					}
+
+					FQuat DeltaRotation = RetargetedAnimRefTransform.GetRotation().Inverse() * MeshRefTransform.GetRotation();
+					DeltaRotation.Normalize();
+
+					Transform.SetRotation(FQuat4f(DeltaRotation));
+					Transform.SetTranslation(FVector3f(MeshRefTransform.GetTranslation() - RetargetedAnimRefTransform.GetTranslation()));
+					Transform.SetScale3D(FVector3f(MeshRefTransform.GetScale3D() * FTransform::GetSafeScaleReciprocal(RetargetedAnimRefTransform.GetScale3D())));
+				}
+				break;
+				case EBoneTranslationRetargetingMode::OrientAndScale:
+				{
+					const FVector AnimTranslation = AnimRefTransform.GetTranslation();
+					const FVector MeshTranslation = MeshRefTransform.GetTranslation();
+
+					if (!AnimTranslation.Equals(MeshTranslation, BONE_TRANS_RT_ORIENT_AND_SCALE_PRECISION))
+					{
+						const float AnimTranslationLength = AnimTranslation.Size();
+						const float MeshTranslationLength = MeshTranslation.Size();
+
+						if (!FMath::IsNearlyZero(AnimTranslationLength * MeshTranslationLength))
+						{
+							const FVector AnimDir = AnimTranslation / AnimTranslationLength;
+							const FVector MeshDir = MeshTranslation / MeshTranslationLength;
+							const FQuat DeltaRotation = FQuat::FindBetweenNormals(AnimDir, MeshDir);
+							const float Scale = MeshTranslationLength / AnimTranslationLength;
+
+							Transform.SetRotation(FQuat4f(DeltaRotation));
+							Transform.SetScale3D(FVector3f(Scale));
+						}
+					}
+				}
+				break;
+			}
+		}
+		else
+		{
+			Transform = FTransform3f(MeshRefPose[MeshBoneIndex]);
+		}
+
+		if (bRequiresRefPoseRetarget)
+		{
+			BoneRetargetingTables.RefPoseRetargetQuats.Emplace(TTuple<FQuat4f, FQuat4f>(RefPoseRetargetQuats.Get<0>(), RefPoseRetargetQuats.Get<1>()));
+		}
+
+		BoneRetargetingTables.Mode.Emplace(Mode);
+		BoneRetargetingTables.Transform.Emplace(Transform);
+		BoneRetargetingTables.AnimBoneIndex.Emplace(AnimBoneIndex);
+	}
+}
+
+void FAnimSequenceTransformProvider::BuildLeafBonePathData(FLeafBonePathData& OutData, const FReferenceSkeleton& RefSkeleton)
+{
+	const int32 NumBones = RefSkeleton.GetRawBoneNum();
+
+	struct FBoneData
+	{
+		uint16 Depth    : 14 = 0;
+		uint16 bIsLeaf  : 1  = 1;
+		uint16 bIsOwned : 1  = 0;
+	};
+
+	TArray<FBoneData, FConcurrentLinearArrayAllocator> BoneDatas;
+	BoneDatas.SetNum(NumBones);
+	int32 NumLeafBones = NumBones;
+
+	for (int32 ChildIndex = 1; ChildIndex < NumBones; ++ChildIndex)
+	{
+		const int32 ParentIndex = RefSkeleton.GetParentIndex(ChildIndex);
+		if (ParentIndex != INDEX_NONE)
+		{
+			FBoneData& ParentBone = BoneDatas[ParentIndex];
+			FBoneData& ChildBone  = BoneDatas[ChildIndex];
+
+			if (ParentBone.bIsLeaf)
+			{
+				ParentBone.bIsLeaf = 0;
+				NumLeafBones--;
+			}
+
+			ChildBone.Depth = ParentBone.Depth + 1;
+		}
+	}
+
+	int32 NumPathNodes = 0;
+	for (int32 BoneIndex = 0; BoneIndex < NumBones; ++BoneIndex)
+	{
+		if (BoneDatas[BoneIndex].bIsLeaf)
+		{
+			NumPathNodes += BoneDatas[BoneIndex].Depth;
+		}
+	}
+
+	const uint32 OffsetTableSize = NumLeafBones;
+	const uint32 HeadersSize = NumLeafBones;
+	const uint32 TotalSize = OffsetTableSize + HeadersSize + NumPathNodes;
+
+	OutData.Buffer.SetNumUninitialized(TotalSize);
+	OutData.NumLeafBones = NumLeafBones;
+
+	uint32* OffsetTable = OutData.Buffer.GetData();
+	uint32* PathData = OutData.Buffer.GetData() + OffsetTableSize;
+
+	int32 LeafIndex = 0;
+	uint32 PathDataOffset = OffsetTableSize * sizeof(uint32);
+
+	for (int32 BoneIndex = 0; BoneIndex < NumBones; ++BoneIndex)
+	{
+		if (!BoneDatas[BoneIndex].bIsLeaf)
+		{
+			continue;
+		}
+
+		const uint16 Depth = BoneDatas[BoneIndex].Depth;
+
+		OffsetTable[LeafIndex] = PathDataOffset;
+
+		FBoneData& RootBone = BoneDatas[0];
+		const bool bIsRootOwned = !RootBone.bIsOwned;
+		RootBone.bIsOwned = 1;
+
+		UE::HLSL::FLeafBonePathHeader* Header = reinterpret_cast<UE::HLSL::FLeafBonePathHeader*>(PathData);
+		Header->PathDepth = Depth;
+		Header->bIsRootOwned = bIsRootOwned;
+		PathData++;
+		PathDataOffset += sizeof(uint32);
+
+		// Write path nodes in reverse (root -> leaf), skipping root
+		if (Depth > 0)
+		{
+			UE::HLSL::FLeafBonePathNode* Nodes = reinterpret_cast<UE::HLSL::FLeafBonePathNode*>(PathData);
+
+			int32 CurrentNodeIndex = Depth - 1;
+			int32 CurrentBoneIndex = BoneIndex;
+
+			while (CurrentNodeIndex >= 0)
+			{
+				check(CurrentBoneIndex != INDEX_NONE);
+
+				FBoneData& CurrentBone = BoneDatas[CurrentBoneIndex];
+
+				UE::HLSL::FLeafBonePathNode& Node = Nodes[CurrentNodeIndex];
+				Node.BoneIndex = CurrentBoneIndex;
+				Node.bIsLeaf   = CurrentBone.bIsLeaf;
+				Node.bIsOwned  = !CurrentBone.bIsOwned;
+				CurrentBone.bIsOwned = 1;
+
+				CurrentBoneIndex = RefSkeleton.GetParentIndex(CurrentBoneIndex);
+				CurrentNodeIndex--;
+			}
+
+			PathData += Depth;
+			PathDataOffset += Depth * sizeof(uint32);
+		}
+
+		LeafIndex++;
+	}
+}
+
+void FAnimSequenceTransformProvider::RegisterProxy(FAnimSequenceTransformProviderProxy* Proxy)
+{
+	TConstArrayView<FAnimSequenceTransformProviderSequence> ProviderSequences = Proxy->GetSequences();
+	FInstancedSkinningSceneExtensionProxy* ExtensionProxy = Proxy->GetSceneExtensionProxy();
+	const USkinnedAsset* SkinnedAsset = ExtensionProxy->GetSkinnedAsset();
+	TConstArrayView<FMatrix44f> InvGlobalRefPoseTransforms = SkinnedAsset->GetRefBasesInvMatrix();
+
+	FProxyUserData* ProxyUserData = new FProxyUserData(Proxy);
+	ProxyUserData->Sequences.Reserve(ProviderSequences.Num());
+	ProxyUserData->SequenceToMappingTable.Reserve(ProviderSequences.Num());
+	ProxyUserData->NumMeshBones = InvGlobalRefPoseTransforms.Num();
+
+	BuildLeafBonePathData(ProxyUserData->LeafBonePathData, SkinnedAsset->GetRefSkeleton());
+	ProxyUserDatasToUpload.Emplace(ProxyUserData);
+
+	for (const FAnimSequenceTransformProviderSequence& ProviderSequence : ProviderSequences)
+	{
+		const UAnimSequence* AnimSequence = ProviderSequence.Sequence.Get();
+		USkeleton* SourceSkeleton = AnimSequence->GetSkeleton();
+
+		const FSequenceRegistryKey SequenceRegistryKey(AnimSequence);
+		const FMappingRegistryKey MappingRegistryKey(SourceSkeleton, SkinnedAsset);
+
+		TUniquePtr<FSequence>& Sequence = SequenceRegistry.FindOrAdd(SequenceRegistryKey);
+		if (!Sequence)
+		{
+			Sequence = MakeUnique<FSequence>();
+
+			Sequence->RegistryKey            = SequenceRegistryKey;
+			Sequence->Sequence               = AnimSequence;
+			Sequence->PlayLength             = AnimSequence->GetPlayLength();
+			Sequence->SampleRate             = GetAnimSequenceSampleRate(AnimSequence);
+			Sequence->NumFrames              = Sequence->SampleRate.AsFrameTime(Sequence->PlayLength).RoundToFrame().Value;
+			Sequence->NumAnimBones           = SourceSkeleton->GetReferenceSkeleton().GetRawBoneNum();
+			Sequence->NumSampleRateLevels    = GetMaxSampleRateLevels(Sequence->NumFrames);
+			Sequence->CurrentSampleRateLevel = Sequence->NumSampleRateLevels;
+			Sequence->TargetSampleRateLevel  = Sequence->NumSampleRateLevels - 1;
+			Sequence->NextFrameToUpload      = 0;
+			Sequence->FramesRequested.Init(false, Sequence->NumFrames);
+			Sequence->FramesUploaded.Init(false, Sequence->NumFrames);
+			Sequence->NumFramesRequested     = 0;
+			Sequence->NumFramesUploaded      = 0;
+			
+			check(!SequencesToUpload.Contains(Sequence.Get()));
+			SequencesToUpload.AddUnique(Sequence.Get());
+
+			INC_DWORD_STAT(STAT_AnimSequenceTransformProvider_Sequences);
+		}
+		++Sequence->RefCount;
+
+		TUniquePtr<FMapping>& Mapping = MappingRegistry.FindOrAdd(MappingRegistryKey);
+		if (!Mapping)
+		{
+			Mapping = MakeUnique<FMapping>();
+
+			Mapping->RegistryKey    = MappingRegistryKey;
+			Mapping->SourceSkeleton = SourceSkeleton;
+			Mapping->SkinnedAsset   = SkinnedAsset;
+
+			Mapping->NumMeshBones = ProxyUserData->NumMeshBones;
+			BuildBoneRetargetingTables(Mapping->BoneRetargetingTables, SourceSkeleton, SkinnedAsset);
+
+			Mapping->InvGlobalRefPoseTransforms.Reserve(InvGlobalRefPoseTransforms.Num());
+			for (const FMatrix44f& InvGlobalRefPose : InvGlobalRefPoseTransforms)
+			{
+				Mapping->InvGlobalRefPoseTransforms.Emplace(InvGlobalRefPose);
+			}
+
+			MappingsToUpload.Add(Mapping.Get());
+
+			INC_DWORD_STAT(STAT_AnimSequenceTransformProvider_Mappings);
+		}
+		++Mapping->RefCount;
+
+		ProxyUserData->Sequences.Emplace(Sequence.Get());
+		ProxyUserData->SequenceToMappingTable.Emplace(Mapping.Get());
+		ProxyUserData->Mappings.AddUnique(Mapping.Get());
+	}
+
+	ExtensionProxy->SetUserData(ProxyUserData);
+}
+
+void FAnimSequenceTransformProvider::UnregisterProxy(FAnimSequenceTransformProviderProxy* Proxy)
+{
+	FInstancedSkinningSceneExtensionProxy* ExtensionProxy = Proxy->GetSceneExtensionProxy();
+
+	FProxyUserData* ProxyUserData = static_cast<FProxyUserData*>(ExtensionProxy->GetUserData());
+	check(ProxyUserData);
+	ExtensionProxy->SetUserData(nullptr);
+
+	ProxyUserDatasToUpload.RemoveSingleSwap(ProxyUserData);
+
+	if (ProxyUserData->LeafBonePathBufferAllocation.IsValid())
+	{
+		LeafBonePathAllocator.Free(ProxyUserData->LeafBonePathBufferAllocation.Offset, ProxyUserData->LeafBonePathBufferAllocation.Count);
+		ProxyUserData->LeafBonePathBufferAllocation = {};
+	}
+
+	for (FSequence* Sequence : ProxyUserData->Sequences)
+	{
+		check(Sequence && Sequence->RefCount > 0);
+		--Sequence->RefCount;
+
+		if (Sequence->RefCount == 0)
+		{
+			if (Sequence->BufferAllocation.IsValid())
+			{
+				SequenceTransformAllocator.Free(Sequence->BufferAllocation.Offset, Sequence->BufferAllocation.Count);
+				Sequence->BufferAllocation = {};
+			}
+
+			SequencesToUpload.RemoveSingleSwap(Sequence, EAllowShrinking::No);
+			SequenceRegistry.Remove(Sequence->RegistryKey);
+			Sequence = nullptr;
+
+			DEC_DWORD_STAT(STAT_AnimSequenceTransformProvider_Sequences);
+		}
+	}
+
+	for (FMapping* Mapping : ProxyUserData->Mappings)
+	{
+		check(Mapping && Mapping->RefCount > 0);
+		--Mapping->RefCount;
+
+		if (Mapping->RefCount == 0)
+		{
+			if (Mapping->MappingTransformBufferAllocation.IsValid())
+			{
+				SequenceTransformAllocator.Free(Mapping->MappingTransformBufferAllocation.Offset, Mapping->MappingTransformBufferAllocation.Count);
+				Mapping->MappingTransformBufferAllocation = {};
+				Mapping->InvGlobalRefPoseTransformBufferOffset = 0;
+			}
+
+			if (Mapping->RetargetingDataBufferAllocation.IsValid())
+			{
+				RetargetingDataAllocator.Free(Mapping->RetargetingDataBufferAllocation.Offset, Mapping->RetargetingDataBufferAllocation.Count);
+				Mapping->RetargetingDataBufferAllocation = {};
+			}
+
+			MappingsToUpload.RemoveSingleSwap(Mapping, EAllowShrinking::No);
+			MappingRegistry.Remove(Mapping->RegistryKey);
+			Mapping = nullptr;
+
+			DEC_DWORD_STAT(STAT_AnimSequenceTransformProvider_Mappings);
+		}
+	}
+
+	delete ProxyUserData;
+}
+
+template <typename T>
+class TBufferUploadData
+{
+public:
+	TBufferUploadData(T* InElements, int32 InMaxElements)
+		: Elements(InElements)
+		, MaxElements(InMaxElements)
+	{}
+
+	T& Add_GetRef()
+	{
+		check(NumElements < MaxElements);
+		return Elements[NumElements++];
+	}
+
+private:
+	T* Elements;
+	int32 MaxElements;
+	int32 NumElements = 0;
+};
+
+template <typename T>
+FRDGBuffer* CreateByteAddressBuffer(FRDGBuilder& GraphBuilder, uint32 MaxNumElements, const TCHAR* Name)
+{
+	return GraphBuilder.CreateBuffer(FRDGBufferDesc::CreateByteAddressDesc(MaxNumElements * sizeof(T)), Name);
+}
+
+template <typename T>
+TBufferUploadData<T> Lock(FRDGBuilder& GraphBuilder, FRDGBuffer* Buffer)
+{
+	GraphBuilder.ConvertToExternalBuffer(Buffer);
+
+	return TBufferUploadData<T>(reinterpret_cast<T*>(GraphBuilder.RHICmdList.LockBuffer(Buffer->GetRHIUnchecked(), 0, Buffer->Desc.GetSize(), RLM_WriteOnly)), Buffer->Desc.GetSize() / sizeof(T));
+}
+
+void Unlock(FRDGBuilder& GraphBuilder, FRDGBuffer* Buffer)
+{
+	GraphBuilder.RHICmdList.UnlockBuffer(Buffer->GetRHIUnchecked());
+}
+
+void FAnimSequenceTransformProvider::ProvideGPUTransforms(FSkinningTransformProvider::FProviderContext& Context)
+{
+	TRACE_CPUPROFILER_EVENT_SCOPE(FAnimSequenceTransformProvider::ProvideGPUTransforms);
+	FRDGBuilder& GraphBuilder = Context.GraphBuilder;
+	static const uint32 BonesPerGroup = FEvaluateAnimSequenceCS::BonesPerGroup;
+
+	SequenceTransformAllocator.Consolidate();
+	RetargetingDataAllocator.Consolidate();
+	LeafBonePathAllocator.Consolidate();
+
+	uint32 MaxTotalSubTracks = 0;
+
+	for (const FSkinningTransformProvider::FProviderIndirection Indirection : Context.Indirections)
+	{
+		FInstancedSkinningSceneExtensionProxy* ExtensionProxy = static_cast<FInstancedSkinningSceneExtensionProxy*>(Context.Proxies[Indirection.Index]);
+		FProxyUserData& ProxyUserData = static_cast<FProxyUserData&>(ExtensionProxy->GetUserDataChecked());
+		MaxTotalSubTracks += 2 * ProxyUserData.Tracks.GetNumActiveTracks();
+	}
+
+	if (!MaxTotalSubTracks)
+	{
+		return;
+	}
+
+	struct FTrackUpdateRequest
+	{
+		uint8 bNeedsCurrent : 1;
+		uint8 UpdateMode    : 7;
+
+		void SetUpdateMode(EPreviousBoneTransformUpdateMode InUpdateMode)
+		{
+			UpdateMode = static_cast<uint8>(InUpdateMode);
+		}
+
+		EPreviousBoneTransformUpdateMode GetUpdateMode() const
+		{
+			return static_cast<EPreviousBoneTransformUpdateMode>(UpdateMode);
+		}
+	};
+
+	TArray<FTrackUpdateRequest, SceneRenderingAllocator> TrackUpdateRequests;
+	TrackUpdateRequests.Reserve(MaxTotalSubTracks);
+
+	uint32 NumTotalIntermediateTransforms = 0;
+	uint32 NumTotalEvaluateBlockHeaders = 0;
+	uint32 NumTotalEvaluateBlendBlockHeaders = 0;
+	uint32 NumTotalScatterBlockHeaders = 0;
+	uint32 NumTotalConcatenateBlockHeaders = 0;
+
+	for (const FSkinningTransformProvider::FProviderIndirection Indirection : Context.Indirections)
+	{
+		FInstancedSkinningSceneExtensionProxy* ExtensionProxy = static_cast<FInstancedSkinningSceneExtensionProxy*>(Context.Proxies[Indirection.Index]);
+		FProxyUserData& ProxyUserData = static_cast<FProxyUserData&>(ExtensionProxy->GetUserDataChecked());
+		FAnimSequenceTrackRenderPool& Tracks = ProxyUserData.Tracks;
+		ProxyUserData.TrackHistories.SetNum(Tracks.GetNumTracks());
+
+		const uint32 NumEvaluateBlockHeaders    = FMath::DivideAndRoundUp(ProxyUserData.NumMeshBones, BonesPerGroup);
+		const uint32 NumScatterBlockHeaders     = FMath::DivideAndRoundUp(ExtensionProxy->GetMaxBoneTransformCount(), BonesPerGroup);
+		const uint32 NumConcatenateBlockHeaders = FMath::DivideAndRoundUp(ProxyUserData.LeafBonePathData.NumLeafBones, BonesPerGroup);
+		const uint32 CurrentTransformSlot = Indirection.CurrentTransformOffset < Indirection.PreviousTransformOffset ? 0 : 1;
+	
+		int32 NumSubTracks = 0;
+		int32 NumEvaluateSubTracks = 0;
+		int32 NumEvaluateBlendSubTracks = 0;
+
+		Tracks.EnumerateActiveTracks([&] (int32 TrackIndex)
+		{
+			const FAnimSequenceTrackRenderData& Data   = Tracks.GetData(TrackIndex);
+			const FAnimSequenceTrackPackedData& Source = Data.GetSource();
+			const FAnimSequenceTrackPackedData& Target = Data.GetTarget();
+			const float SourcePosition                 = Source.GetPosition();
+			const float SourcePreviousPosition         = Source.GetPreviousPosition();
+			const float TargetPosition                 = Target.GetPosition();
+			const float TargetPreviousPosition         = Target.GetPreviousPosition();
+			const float BlendWeight                    = Data.GetBlendWeight();
+
+			// Dirty previous transforms means the whole allocation is dirty. Rebuild everything.
+			const bool bForceUpdateAll = EnumHasAnyFlags(Indirection.DirtyBoneTransforms, EDirtyBoneTransforms::Previous);
+			FTrackHistory::FUpdateResult UpdateResult = ProxyUserData.TrackHistories[TrackIndex].Update(Data, CurrentTransformSlot, bForceUpdateAll);
+
+			FTrackUpdateRequest UpdateRequest;
+			UpdateRequest.bNeedsCurrent = UpdateResult.bNeedsCurrent;
+			UpdateRequest.UpdateMode = uint8(UpdateResult.bNeedsCurrent && UpdateResult.bNeedsPrevious && TargetPosition == TargetPreviousPosition && SourcePosition == SourcePreviousPosition
+				? EPreviousBoneTransformUpdateMode::DuplicateCurrentToPrevious
+				: EPreviousBoneTransformUpdateMode::UpdatePrevious);
+
+			const int32 NumLocalSubTracks = UpdateResult.bNeedsCurrent + (UpdateRequest.GetUpdateMode() == EPreviousBoneTransformUpdateMode::UpdatePrevious ? 1 : 0);
+
+			NumSubTracks              += NumLocalSubTracks;
+			NumEvaluateSubTracks      += BlendWeight <= 0.0f ? NumLocalSubTracks : 0;
+			NumEvaluateBlendSubTracks += BlendWeight >  0.0f ? NumLocalSubTracks : 0;
+
+			TrackUpdateRequests.Emplace(UpdateRequest);
+		});
+
+		NumTotalIntermediateTransforms    += NumSubTracks * ProxyUserData.NumMeshBones;
+		NumTotalEvaluateBlockHeaders      += NumEvaluateSubTracks * NumEvaluateBlockHeaders;
+		NumTotalEvaluateBlendBlockHeaders += NumEvaluateBlendSubTracks * NumEvaluateBlockHeaders;
+		NumTotalScatterBlockHeaders       += NumSubTracks * NumScatterBlockHeaders;
+		NumTotalConcatenateBlockHeaders   += NumSubTracks * NumConcatenateBlockHeaders;
+	}
+
+	uint32 NumTotalSequenceTransforms = 0;
+	uint32 NumTotalRetargetingDatas = 0;
+	uint32 NumTotalLeafBonePathElements = 0;
+
+	for (FMapping* Mapping : MappingsToUpload)
+	{
+		const int32 NumRetargetingTransforms = (Mapping->BoneRetargetingTables.bRequiresRefPoseRetarget ? 2 : 1);
+		const int32 NumTransforms            = Mapping->NumMeshBones * (NumRetargetingTransforms + 1);
+		const int32 NumRetargetingDatas      = Mapping->NumMeshBones;
+
+		Mapping->MappingTransformBufferAllocation.Count  = NumTransforms;
+		Mapping->MappingTransformBufferAllocation.Offset = SequenceTransformAllocator.Allocate(NumTransforms);
+
+		Mapping->InvGlobalRefPoseTransformBufferOffset = Mapping->MappingTransformBufferAllocation.Offset;
+		Mapping->RetargetingTransformBufferOffset    = Mapping->MappingTransformBufferAllocation.Offset + Mapping->NumMeshBones;
+
+		Mapping->RetargetingDataBufferAllocation.Count = NumRetargetingDatas;
+		Mapping->RetargetingDataBufferAllocation.Offset = RetargetingDataAllocator.Allocate(NumRetargetingDatas);
+
+		NumTotalSequenceTransforms += NumTransforms;
+		NumTotalRetargetingDatas  += NumRetargetingDatas;
+	}
+
+	if (!MappingsToUpload.IsEmpty())
+	{
+		GraphBuilder.AddPostExecuteCallback([this]
+		{
+			MappingsToUpload.Empty();
+		});
+	}
+
+	for (FSequence* Sequence : SequencesToUpload)
+	{
+		const int32 NumTransforms = Sequence->NumAnimBones * Sequence->NumFrames;
+
+		if (!Sequence->BufferAllocation.IsValid())
+		{
+			Sequence->BufferAllocation.Count  = NumTransforms;
+			Sequence->BufferAllocation.Offset = SequenceTransformAllocator.Allocate(NumTransforms);
+		}
+
+		NumTotalSequenceTransforms += NumTransforms;
+	}
+
+	for (FProxyUserData* ProxyUserData : ProxyUserDatasToUpload)
+	{
+		const int32 NumElements = ProxyUserData->LeafBonePathData.Buffer.Num();
+
+		ProxyUserData->LeafBonePathBufferAllocation.Count = NumElements;
+		ProxyUserData->LeafBonePathBufferAllocation.Offset = LeafBonePathAllocator.Allocate(NumElements);
+
+		NumTotalLeafBonePathElements += NumElements;
+	}
+
+	int32& NumTotalFramesRequested = *GraphBuilder.AllocPOD<int32>();
+	NumTotalFramesRequested = 0;
+
+	FRDGBuffer* EvaluateBlockHeaderBuffer      = CreateByteAddressBuffer<UE::HLSL::FEvaluateAnimSequenceBlockHeader>(GraphBuilder, NumTotalEvaluateBlockHeaders, TEXT("AnimSequenceTransformProvider.EvaluateBlockHeaders"));
+	FRDGBuffer* EvaluateBlendBlockHeaderBuffer = CreateByteAddressBuffer<UE::HLSL::FEvaluateAnimSequenceBlendBlockHeader>(GraphBuilder, NumTotalEvaluateBlendBlockHeaders, TEXT("AnimSequenceTransformProvider.EvaluateBlendBlockHeaders"));
+	FRDGBuffer* ScatterBlockHeaderBuffer       = CreateByteAddressBuffer<UE::HLSL::FScatterTransformsBlockHeader>(GraphBuilder, NumTotalScatterBlockHeaders, TEXT("AnimSequenceTransformProvider.ScatterBlockHeaders"));
+	FRDGBuffer* ConcatenateBlockHeaderBuffer   = CreateByteAddressBuffer<UE::HLSL::FConcatenateBlockHeader>(GraphBuilder, NumTotalConcatenateBlockHeaders, TEXT("AnimSequenceTransformProvider.ConcatenateBlockHeaders"));
+
+	TBufferUploadData<UE::HLSL::FEvaluateAnimSequenceBlockHeader> EvaluateBlockHeaders           = Lock<UE::HLSL::FEvaluateAnimSequenceBlockHeader>(GraphBuilder, EvaluateBlockHeaderBuffer);
+	TBufferUploadData<UE::HLSL::FEvaluateAnimSequenceBlendBlockHeader> EvaluateBlendBlockHeaders = Lock<UE::HLSL::FEvaluateAnimSequenceBlendBlockHeader>(GraphBuilder, EvaluateBlendBlockHeaderBuffer);
+	TBufferUploadData<UE::HLSL::FScatterTransformsBlockHeader> ScatterBlockHeaders               = Lock<UE::HLSL::FScatterTransformsBlockHeader>(GraphBuilder, ScatterBlockHeaderBuffer);
+	TBufferUploadData<UE::HLSL::FConcatenateBlockHeader> ConcatenateBlockHeaders                 = Lock<UE::HLSL::FConcatenateBlockHeader>(GraphBuilder, ConcatenateBlockHeaderBuffer);
+
+	UE::Tasks::FTask ProcessTask = GraphBuilder.AddSetupTask(
+		[
+			  this
+			, &ViewFamily = Context.ViewFamily
+			, Indirections = Context.Indirections
+			, Proxies = Context.Proxies
+			, &NumTotalFramesRequested
+			, TrackUpdateRequests = MoveTemp(TrackUpdateRequests)
+			, EvaluateBlockHeaders
+			, EvaluateBlendBlockHeaders
+			, ScatterBlockHeaders
+			, ConcatenateBlockHeaders
+		] () mutable
+	{
+		TRACE_CPUPROFILER_EVENT_SCOPE(FAnimSequenceTransformProvider::EvaluateAnimSequences);
+
+		uint32 IntermediateTransformOffset = 0;
+		uint32 NumActiveTracks = 0;
+		uint32 NumEvaluateBlocks = 0;
+		uint32 NumEvaluateBlendBlocks = 0;
+		uint32 NumCurrentTransformUpdates = 0;
+		uint32 NumPreviousTransformUpdates = 0;
+
+		const auto EvaluateAnimSequence = [&](
+			const FSkinningTransformProvider::FProviderIndirection& Indirection,
+			const FProxyUserData& ProxyUserData,
+			const FMapping& SourceMapping,
+			const FMapping& TargetMapping,
+			FSequence& SourceSequence,
+			FSequence& TargetSequence,
+			float SourceTime,
+			float TargetTime,
+			float BlendWeight,
+			uint32 MaxBoneTransformCount,
+			uint32 DstTransformBufferOffset1,
+			uint32 DstTransformBufferOffset2)
+		{
+			const auto GetBlockTransformCount = [](uint32 BlockIndex, uint32 TotalCount, uint32 BlockSize) -> uint32
+			{
+				return FMath::Min(BlockSize, TotalCount - BlockIndex * BlockSize);
+			};
+
+			int32 SourceKeyIndex0, SourceKeyIndex1;
+			float SourceAlpha;
+			SourceSequence.GetKeyIndicesFromTime(SourceKeyIndex0, SourceKeyIndex1, SourceAlpha, SourceTime);
+
+			NumTotalFramesRequested += SourceSequence.RequestFrame(SourceKeyIndex0);
+			NumTotalFramesRequested += SourceSequence.RequestFrame(SourceKeyIndex1);
+
+			// Evaluation Pass
+			{
+				check(SourceMapping.NumMeshBones == TargetMapping.NumMeshBones);
+				const uint32 NumBlocks = FMath::DivideAndRoundUp(TargetMapping.NumMeshBones, BonesPerGroup);
+
+				const uint32 SourceSequenceTransformBufferOffset    = SourceSequence.BufferAllocation.Offset * sizeof(FBoneTransformWithScale);
+				const uint32 SourceRetargetingTransformBufferOffset = SourceMapping.RetargetingTransformBufferOffset * sizeof(FBoneTransformWithScale);
+				const uint32 SourceRetargetingDataBufferOffset      = SourceMapping.RetargetingDataBufferAllocation.Offset * sizeof(UE::HLSL::FBoneRetargetingData);
+
+				if (BlendWeight > 0.0f)
+				{
+					int32 TargetKeyIndex0, TargetKeyIndex1;
+					float TargetAlpha;
+					TargetSequence.GetKeyIndicesFromTime(TargetKeyIndex0, TargetKeyIndex1, TargetAlpha, TargetTime);
+
+					NumTotalFramesRequested += TargetSequence.RequestFrame(TargetKeyIndex0);
+					NumTotalFramesRequested += TargetSequence.RequestFrame(TargetKeyIndex1);
+
+					const uint32 TargetSequenceTransformBufferOffset    = TargetSequence.BufferAllocation.Offset * sizeof(FBoneTransformWithScale);
+					const uint32 TargetRetargetingTransformBufferOffset = TargetMapping.RetargetingTransformBufferOffset * sizeof(FBoneTransformWithScale);
+					const uint32 TargetRetargetingDataBufferOffset      = TargetMapping.RetargetingDataBufferAllocation.Offset * sizeof(UE::HLSL::FBoneRetargetingData);
+
+					for (uint32 BlockIndex = 0; BlockIndex < NumBlocks; ++BlockIndex)
+					{
+						UE::HLSL::FEvaluateAnimSequenceBlendBlockHeader& BlockHeader = EvaluateBlendBlockHeaders.Add_GetRef();
+
+						BlockHeader.SourceAlpha                            = SourceAlpha;
+						BlockHeader.TargetAlpha                            = TargetAlpha;
+						BlockHeader.BlendWeight                            = BlendWeight;
+						BlockHeader.SourceKeyIndex0                        = SourceKeyIndex0;
+						BlockHeader.SourceKeyIndex1                        = SourceKeyIndex1;
+						BlockHeader.TargetKeyIndex0                        = TargetKeyIndex0;
+						BlockHeader.TargetKeyIndex1                        = TargetKeyIndex1;
+						BlockHeader.SourceSequenceTransformBufferOffset    = SourceSequenceTransformBufferOffset;
+						BlockHeader.TargetSequenceTransformBufferOffset    = TargetSequenceTransformBufferOffset;
+						BlockHeader.TargetRetargetingTransformBufferOffset = TargetRetargetingTransformBufferOffset;
+						BlockHeader.SourceRetargetingTransformBufferOffset = SourceRetargetingTransformBufferOffset;
+						BlockHeader.TargetRetargetingDataBufferOffset      = TargetRetargetingDataBufferOffset;
+						BlockHeader.SourceRetargetingDataBufferOffset      = SourceRetargetingDataBufferOffset;
+						BlockHeader.SourceNumAnimBones                     = SourceSequence.NumAnimBones;
+						BlockHeader.TargetNumAnimBones                     = TargetSequence.NumAnimBones;
+						BlockHeader.DstTransformBufferOffset               = IntermediateTransformOffset;
+						BlockHeader.BlockLocalIndex                        = BlockIndex;
+						BlockHeader.BlockTransformCount                    = GetBlockTransformCount(BlockIndex, TargetMapping.NumMeshBones, BonesPerGroup);
+					}
+
+					NumEvaluateBlendBlocks += NumBlocks;
+				}
+				else
+				{
+					for (uint32 BlockIndex = 0; BlockIndex < NumBlocks; ++BlockIndex)
+					{
+						UE::HLSL::FEvaluateAnimSequenceBlockHeader& BlockHeader = EvaluateBlockHeaders.Add_GetRef();
+
+						BlockHeader.SequenceTransformBufferOffset    = SourceSequenceTransformBufferOffset;
+						BlockHeader.RetargetingTransformBufferOffset = SourceRetargetingTransformBufferOffset;
+						BlockHeader.RetargetingDataBufferOffset      = SourceRetargetingDataBufferOffset;
+						BlockHeader.DstTransformBufferOffset         = IntermediateTransformOffset;
+						BlockHeader.NumAnimBones                     = SourceSequence.NumAnimBones;
+						BlockHeader.KeyIndex0                        = SourceKeyIndex0;
+						BlockHeader.KeyIndex1                        = SourceKeyIndex1;
+						BlockHeader.Alpha                            = SourceAlpha;
+						BlockHeader.BlockLocalIndex                  = BlockIndex;
+						BlockHeader.BlockTransformCount              = GetBlockTransformCount(BlockIndex, TargetMapping.NumMeshBones, BonesPerGroup);
+					}
+
+					NumEvaluateBlocks += NumBlocks;
+				}
+			}
+
+			// Concatenate Local to Global Pass
+			{
+				const uint32 NumLeafBones = ProxyUserData.LeafBonePathData.NumLeafBones;
+				const uint32 NumBlocks    = FMath::DivideAndRoundUp(NumLeafBones, BonesPerGroup);
+
+				const uint32 LeafBonePathBufferOffset = ProxyUserData.LeafBonePathBufferAllocation.Offset * sizeof(uint32);
+
+				for (uint32 BlockIndex = 0; BlockIndex < NumBlocks; ++BlockIndex)
+				{
+					UE::HLSL::FConcatenateBlockHeader& BlockHeader = ConcatenateBlockHeaders.Add_GetRef();
+
+					BlockHeader.LeafBonePathBufferOffset = LeafBonePathBufferOffset;
+					BlockHeader.TransformBufferOffset    = IntermediateTransformOffset;
+					BlockHeader.BlockLocalIndex          = BlockIndex;
+					BlockHeader.BlockLeafBoneCount       = GetBlockTransformCount(BlockIndex, NumLeafBones, BonesPerGroup);
+				}
+			}
+
+			// Scatter Pass
+			{
+				const uint32 NumBlocks     = FMath::DivideAndRoundUp(MaxBoneTransformCount, BonesPerGroup);
+				const uint32 BoneMapOffset = ProxyUserData.TransformStorageMode == EBoneTransformStorageMode::BoneMap ? Indirection.HierarchyOffset : ~uint32(0);
+
+				const uint32 InvGlobalRefPoseTransformBufferOffset = SourceMapping.InvGlobalRefPoseTransformBufferOffset * sizeof(FBoneTransformWithScale);
+
+				for (uint32 BlockIndex = 0; BlockIndex < NumBlocks; ++BlockIndex)
+				{
+					UE::HLSL::FScatterTransformsBlockHeader& BlockHeader = ScatterBlockHeaders.Add_GetRef();
+
+					BlockHeader.InvGlobalRefPoseTransformBufferOffset = InvGlobalRefPoseTransformBufferOffset;
+					BlockHeader.SrcTransformBufferOffset              = IntermediateTransformOffset;
+					BlockHeader.DstTransformBufferOffset1             = DstTransformBufferOffset1;
+					BlockHeader.DstTransformBufferOffset2             = DstTransformBufferOffset2;
+					BlockHeader.BoneMapOffset                         = BoneMapOffset;
+					BlockHeader.BlockLocalIndex                       = BlockIndex;
+					BlockHeader.BlockTransformCount                   = GetBlockTransformCount(BlockIndex, MaxBoneTransformCount, BonesPerGroup);
+				}
+			}
+
+			IntermediateTransformOffset += SourceMapping.NumMeshBones * sizeof(FBoneTransformWithScale);
+		};
+
+		const float GlobalTimeScale = CVarAnimSequenceTransformProviderTimeScale.GetValueOnRenderThread();
+		const float DeltaTime = ViewFamily.Time.GetDeltaWorldTimeSeconds() * GlobalTimeScale;
+		uint32 UpdateResultIndex = 0;
+
+		for (const FSkinningTransformProvider::FProviderIndirection Indirection : Indirections)
+		{
+			FProxyUserData& ProxyUserData = static_cast<FProxyUserData&>(static_cast<FInstancedSkinningSceneExtensionProxy*>(Proxies[Indirection.Index])->GetUserDataChecked());
+			FAnimSequenceTrackRenderPool& Tracks = ProxyUserData.Tracks;
+
+			const uint32 CurrentTransformOffset  = Indirection.CurrentTransformOffset;
+			const uint32 PreviousTransformOffset = Indirection.PreviousTransformOffset;
+			const uint32 MaxBoneTransformCount   = ProxyUserData.SceneExtensionProxy.GetMaxBoneTransformCount();
+
+			Tracks.EnumerateActiveTracks([&] (int32 TrackIndex)
+			{
+				FAnimSequenceTrackRenderData& Data = Tracks.GetData(TrackIndex);
+
+				const FAnimSequenceTrackPackedData& Source = Data.GetSource();
+				const FAnimSequenceTrackPackedData& Target = Data.GetTarget();
+
+				const FMapping& SourceMapping = ProxyUserData.GetMapping(Source);
+				const FMapping& TargetMapping = ProxyUserData.GetMapping(Target);
+				FSequence& SourceSequence = ProxyUserData.GetSequence(Source);
+				FSequence& TargetSequence = ProxyUserData.GetSequence(Target);
+
+				const FTrackUpdateRequest UpdateRequest = TrackUpdateRequests[UpdateResultIndex++];
+	
+				const float SourcePosition         = Source.GetPosition();
+				const float SourcePreviousPosition = Source.GetPreviousPosition();
+				const float TargetPosition         = Target.GetPosition();
+				const float TargetPreviousPosition = Target.GetPreviousPosition();
+				const float BlendWeight            = Data.GetBlendWeight();
+
+				const uint32 TrackTransformOffset         = TrackIndex * MaxBoneTransformCount * 2 * sizeof(FCompressedBoneTransform);
+				const uint32 TrackCurrentTransformOffset  = TrackTransformOffset + CurrentTransformOffset;
+				const uint32 TrackPreviousTransformOffset = TrackTransformOffset + PreviousTransformOffset;
+
+				if (UpdateRequest.bNeedsCurrent)
+				{
+					uint32 OutputTransformOffset1 = TrackCurrentTransformOffset;
+					uint32 OutputTransformOffset2 = UpdateRequest.GetUpdateMode() == EPreviousBoneTransformUpdateMode::DuplicateCurrentToPrevious ? TrackPreviousTransformOffset : TrackCurrentTransformOffset;
+
+					EvaluateAnimSequence(
+						Indirection,
+						ProxyUserData,
+						SourceMapping, TargetMapping,
+						SourceSequence, TargetSequence,
+						SourcePosition, TargetPosition,
+						BlendWeight,
+						MaxBoneTransformCount,
+						OutputTransformOffset1, OutputTransformOffset2);
+
+					NumCurrentTransformUpdates++;
+				}
+
+				if (UpdateRequest.GetUpdateMode() == EPreviousBoneTransformUpdateMode::UpdatePrevious)
+				{
+					EvaluateAnimSequence(
+						Indirection,
+						ProxyUserData,
+						SourceMapping, TargetMapping,
+						SourceSequence, TargetSequence,
+						SourcePreviousPosition, TargetPreviousPosition,
+						BlendWeight,
+						MaxBoneTransformCount,
+						TrackPreviousTransformOffset, TrackPreviousTransformOffset);
+
+					NumPreviousTransformUpdates++;
+				}
+
+				Data.Tick(SourceSequence.PlayLength, TargetSequence.PlayLength, DeltaTime);
+				NumActiveTracks++;
+			});
+		}
+
+		INC_DWORD_STAT_BY(STAT_AnimSequenceTransformProvider_ActiveTracks, NumActiveTracks);
+		INC_DWORD_STAT_BY(STAT_AnimSequenceTransformProvider_EvaluateBlocks, NumEvaluateBlocks);
+		INC_DWORD_STAT_BY(STAT_AnimSequenceTransformProvider_EvaluateBlendBlocks, NumEvaluateBlendBlocks);
+		INC_DWORD_STAT_BY(STAT_AnimSequenceTransformProvider_CurrentTransformUpdates, NumCurrentTransformUpdates);
+		INC_DWORD_STAT_BY(STAT_AnimSequenceTransformProvider_PreviousTransformUpdates, NumPreviousTransformUpdates);
+	});
+
+	Unlock(GraphBuilder, EvaluateBlockHeaderBuffer);
+	Unlock(GraphBuilder, EvaluateBlendBlockHeaderBuffer);
+	Unlock(GraphBuilder, ScatterBlockHeaderBuffer);
+	Unlock(GraphBuilder, ConcatenateBlockHeaderBuffer);
+
+	FRDGBuffer* SequenceTransformBuffer = SequenceTransformPersistentBuffer.ResizeBufferIfNeeded(GraphBuilder, SequenceTransformAllocator.GetMaxSize());
+	FRDGBuffer* RetargetingDataBuffer = RetargetingDataPersistentBuffer.ResizeBufferIfNeeded(GraphBuilder, RetargetingDataAllocator.GetMaxSize());
+	FRDGBuffer* LeafBonePathBuffer = LeafBonePathPersistentBuffer.ResizeBufferIfNeeded(GraphBuilder, LeafBonePathAllocator.GetMaxSize());
+
+	SET_MEMORY_STAT(STAT_AnimSequenceTransformProvider_SequenceTransformMemory, SequenceTransformBuffer->Desc.GetSize());
+	SET_MEMORY_STAT(STAT_AnimSequenceTransformProvider_RetargetingDataMemory, RetargetingDataBuffer->Desc.GetSize());
+	SET_MEMORY_STAT(STAT_AnimSequenceTransformProvider_LeafBonePathMemory, LeafBonePathBuffer->Desc.GetSize());
+
+	bool bExecuteScatterUpload = false;
+
+	if (NumTotalSequenceTransforms > 0)
+	{
+		Context.ScatterUploadBuilder.AddPrerequisiteTask(ProcessTask);
+
+		Context.ScatterUploadBuilder.AddPass(GraphBuilder, SequenceTransformUploadBuffer, SequenceTransformBuffer, NumTotalSequenceTransforms, sizeof(FBoneTransformWithScale), TEXT("AnimSequenceTransformProvider.Sequences"),
+			[this, &NumTotalFramesRequested] (FRDGScatterUploader& ScatterUploader)
+		{
+			TRACE_CPUPROFILER_EVENT_SCOPE(FAnimSequenceTransformProvider::UploadSequences);
+
+			for (const FMapping* Mapping : MappingsToUpload)
+			{
+				TArrayView<FBoneTransformWithScale> DstTransforms = ScatterUploader.Add_GetRef<FBoneTransformWithScale>(Mapping->MappingTransformBufferAllocation.Offset, Mapping->MappingTransformBufferAllocation.Count);
+
+				for (uint32 MeshBoneIndex = 0; MeshBoneIndex < Mapping->NumMeshBones; ++MeshBoneIndex)
+				{
+					DstTransforms[MeshBoneIndex].SetTransform(Mapping->InvGlobalRefPoseTransforms[MeshBoneIndex]);
+
+					const uint32 RetargetTransformIndex = Mapping->NumMeshBones + MeshBoneIndex * (Mapping->BoneRetargetingTables.bRequiresRefPoseRetarget ? 2 : 1);
+					DstTransforms[RetargetTransformIndex].SetTransform(Mapping->BoneRetargetingTables.Transform[MeshBoneIndex]);
+
+					if (Mapping->BoneRetargetingTables.bRequiresRefPoseRetarget)
+					{
+						FBoneTransformWithScale& RefPoseSlot = DstTransforms[RetargetTransformIndex + 1];
+						const FQuat4f& Q0 = Mapping->BoneRetargetingTables.RefPoseRetargetQuats[MeshBoneIndex].Get<0>();
+						const FQuat4f& Q1 = Mapping->BoneRetargetingTables.RefPoseRetargetQuats[MeshBoneIndex].Get<1>();
+
+						// Pack Q0 into Rotation
+						RefPoseSlot.Rotation[0] = Q0.X;
+						RefPoseSlot.Rotation[1] = Q0.Y;
+						RefPoseSlot.Rotation[2] = Q0.Z;
+						RefPoseSlot.Rotation[3] = Q0.W;
+
+						// Pack Q1 into Translation + Scale[0]
+						RefPoseSlot.Translation[0] = Q1.X;
+						RefPoseSlot.Translation[1] = Q1.Y;
+						RefPoseSlot.Translation[2] = Q1.Z;
+						RefPoseSlot.Scale[0] = Q1.W;
+						RefPoseSlot.Scale[1] = 0.0f;
+						RefPoseSlot.Scale[2] = 0.0f;
+					}
+				}
+			}
+
+			BoneTrackArray TrackToBoneIndexMap;
+			FMemMark Mark(FMemStack::Get());
+			TArray<FTransform, TMemStackAllocator<>> AnimLocalPoses;
+			uint32 NumSequenceFrameUploads = 0;
+			int32 SequenceIndex = 0;
+
+			const int32 UploadBudget = FMath::Max<int32>(0, CVarAnimSequenceTransformProviderUploadBudget.GetValueOnRenderThread() - NumTotalFramesRequested);
+			const int32 UploadBudgetDivisor = FMath::Max(1, SequencesToUpload.Num());
+			const int32 FramesPerSequence = UploadBudget / UploadBudgetDivisor;
+			int32 NumFramesLeftover       = UploadBudget % UploadBudgetDivisor;
+
+			while (SequenceIndex < SequencesToUpload.Num())
+			{
+				FSequence& Sequence = *SequencesToUpload[SequenceIndex];
+
+				int32 NumBudgetFramesToUpload = FramesPerSequence + (NumFramesLeftover > 0 ? 1 : 0);
+				NumFramesLeftover = FMath::Max(NumFramesLeftover - 1, 0);
+
+				while (Sequence.NextFrameToUpload < Sequence.NumFrames && NumBudgetFramesToUpload > 0)
+				{
+					const int32 FrameStep = 1u << Sequence.TargetSampleRateLevel;
+
+					if (Sequence.FramesUploaded[Sequence.NextFrameToUpload] || Sequence.FramesRequested[Sequence.NextFrameToUpload])
+					{
+						Sequence.NextFrameToUpload += FrameStep;
+						continue;
+					}
+
+					Sequence.FramesRequested[Sequence.NextFrameToUpload] = true;
+					Sequence.NumFramesRequested++;
+					Sequence.NextFrameToUpload += FrameStep;
+					NumBudgetFramesToUpload--;
+				}
+
+				if (Sequence.NextFrameToUpload >= Sequence.NumFrames && Sequence.TargetSampleRateLevel > 0)
+				{
+					Sequence.TargetSampleRateLevel--;
+					Sequence.CurrentSampleRateLevel--;
+					Sequence.NextFrameToUpload = 1u << Sequence.TargetSampleRateLevel;
+				}
+
+				if (Sequence.NumFramesRequested > 0)
+				{
+					const UAnimSequence* AnimSequence = Sequence.Sequence;
+					UAnimSequence::FScopedCompressedAnimSequence CompressedAnimSequence = AnimSequence->GetCompressedData();
+					const FCompressedAnimSequence& PlatformCompressedData = CompressedAnimSequence.Get();
+
+					const TArray<FTransform>& AnimRefLocalPoses = AnimSequence->GetSkeleton()->GetRefLocalPoses();
+					check(Sequence.NumAnimBones == AnimRefLocalPoses.Num());
+
+					const int32 NumTracks = PlatformCompressedData.CompressedTrackToSkeletonMapTable.Num();
+					TrackToBoneIndexMap.Reset(NumTracks);
+
+					for (int32 TrackIndex = 0; TrackIndex < NumTracks; ++TrackIndex)
+					{
+						const int32 BoneIndex = PlatformCompressedData.GetSkeletonIndexFromTrackIndex(TrackIndex);
+
+						// We only care about raw bones.
+						if (BoneIndex < (int32)Sequence.NumAnimBones)
+						{
+							TrackToBoneIndexMap.Add(BoneTrackPair(BoneIndex, TrackIndex));
+						}
+					}
+
+					FAnimSequenceDecompressionContext DecompressionContext(
+						  Sequence.SampleRate
+						, Sequence.NumFrames
+						, AnimSequence->Interpolation
+						, AnimSequence->GetRetargetTransformsSourceName()
+						, *PlatformCompressedData.CompressedDataStructure
+						, AnimRefLocalPoses
+						, PlatformCompressedData.CompressedTrackToSkeletonMapTable
+						, AnimSequence->GetSkeleton()
+						, AnimSequence->IsValidAdditive()
+						, AnimSequence->GetAdditiveAnimType()
+					);
+
+					AnimLocalPoses.Reset(AnimRefLocalPoses.Num());
+					AnimLocalPoses.Append(AnimRefLocalPoses);
+
+					for (TConstSetBitIterator<> BitIt(Sequence.FramesRequested); BitIt; ++BitIt)
+					{
+						TRACE_CPUPROFILER_EVENT_SCOPE(FAnimSequenceTransformProvider::DecompressPose);
+						const int32 FrameIndex = BitIt.GetIndex();
+
+						const double SeekTime = Sequence.SampleRate.AsSeconds(FFrameTime(FrameIndex));
+						DecompressionContext.Seek(SeekTime);
+
+						if (PlatformCompressedData.BoneCompressionCodec)
+						{
+							TArrayView<FTransform> AnimLocalPosesView(AnimLocalPoses);
+							PlatformCompressedData.BoneCompressionCodec->DecompressPose(DecompressionContext, TrackToBoneIndexMap, TrackToBoneIndexMap, TrackToBoneIndexMap, AnimLocalPosesView);
+						}
+
+						TArrayView<FBoneTransformWithScale> DstTransforms = ScatterUploader.Add_GetRef<FBoneTransformWithScale>(Sequence.BufferAllocation.Offset + Sequence.NumAnimBones * FrameIndex, Sequence.NumAnimBones);
+
+						for (int32 BoneIndex = 0; BoneIndex < AnimLocalPoses.Num(); ++BoneIndex)
+						{
+							DstTransforms[BoneIndex].SetTransform(FTransform3f(AnimLocalPoses[BoneIndex]));
+						}
+
+						Sequence.NumFramesUploaded++;
+						Sequence.FramesUploaded[FrameIndex] = true;
+						Sequence.FramesRequested[FrameIndex] = false;
+						NumSequenceFrameUploads++;
+					}
+
+					Sequence.NumFramesRequested = 0;
+				}
+
+				if (Sequence.NumFramesUploaded == Sequence.NumFrames)
+				{
+					SequencesToUpload.RemoveAtSwap(SequenceIndex, EAllowShrinking::No);
+				}
+				else
+				{
+					SequenceIndex++;
+				}
+			}
+
+			INC_DWORD_STAT_BY(STAT_AnimSequenceTransformProvider_SequenceFrameUploads, NumSequenceFrameUploads);
+		});
+
+		bExecuteScatterUpload = true;
+	}
+
+	if (NumTotalRetargetingDatas > 0)
+	{
+		Context.ScatterUploadBuilder.AddPass(GraphBuilder, RetargetingDataUploadBuffer, RetargetingDataBuffer, NumTotalRetargetingDatas, sizeof(UE::HLSL::FBoneRetargetingData), TEXT("AnimSequenceTransformProvider.RetargetingData"),
+			[this] (FRDGScatterUploader& ScatterUploader)
+		{
+			TRACE_CPUPROFILER_EVENT_SCOPE(FAnimSequenceTransformProvider::UploadRetargetingData);
+
+			for (const FMapping* Mapping : MappingsToUpload)
+			{
+				TArrayView<UE::HLSL::FBoneRetargetingData> DstRetargetingDatas = ScatterUploader.Add_GetRef<UE::HLSL::FBoneRetargetingData>(Mapping->RetargetingDataBufferAllocation.Offset, Mapping->RetargetingDataBufferAllocation.Count);
+
+				for (uint32 MeshBoneIndex = 0; MeshBoneIndex < Mapping->NumMeshBones; ++MeshBoneIndex)
+				{
+					UE::HLSL::FBoneRetargetingData& RetargetingData = DstRetargetingDatas[MeshBoneIndex];
+					RetargetingData.RetargetMode  = Mapping->BoneRetargetingTables.Mode[MeshBoneIndex];
+					RetargetingData.AnimBoneIndex = Mapping->BoneRetargetingTables.AnimBoneIndex[MeshBoneIndex];
+					RetargetingData.bRequiresRefPoseRetarget = Mapping->BoneRetargetingTables.bRequiresRefPoseRetarget;
+				}
+			}
+		});
+
+		bExecuteScatterUpload = true;
+	}
+
+	if (NumTotalLeafBonePathElements > 0)
+	{
+		Context.ScatterUploadBuilder.AddPass(GraphBuilder, LeafBonePathUploadBuffer, LeafBonePathBuffer, NumTotalLeafBonePathElements, sizeof(uint32), TEXT("AnimSequenceTransformProvider.LeafBonePaths"),
+			[this] (FRDGScatterUploader& ScatterUploader)
+		{
+			TRACE_CPUPROFILER_EVENT_SCOPE(FAnimSequenceTransformProvider::UploadLeafBonePaths);
+
+			for (const FProxyUserData* ProxyUserData : ProxyUserDatasToUpload)
+			{
+				TArrayView<uint32> DstBuffer = ScatterUploader.Add_GetRef<uint32>(ProxyUserData->LeafBonePathBufferAllocation.Offset, ProxyUserData->LeafBonePathBufferAllocation.Count);
+
+				FMemory::Memcpy(DstBuffer.GetData(), ProxyUserData->LeafBonePathData.Buffer.GetData(), ProxyUserData->LeafBonePathData.Buffer.Num() * sizeof(uint32));
+			}
+			ProxyUserDatasToUpload.Empty();
+		});
+
+		bExecuteScatterUpload = true;
+	}
+
+	if (bExecuteScatterUpload)
+	{
+		Context.ScatterUploadBuilder.Execute(Context.GraphBuilder);
+	}
+
+	FRDGBuffer* LocalSpaceTransformBuffer = GraphBuilder.CreateBuffer(FRDGBufferDesc::CreateByteAddressDesc(NumTotalIntermediateTransforms * sizeof(FBoneTransformWithScale)), TEXT("AnimSequenceTransformProvider.LocalSpaceBoneTransforms"));
+	FRDGBuffer* GlobalSpaceTransformBuffer = GraphBuilder.CreateBuffer(FRDGBufferDesc::CreateByteAddressDesc(NumTotalIntermediateTransforms * sizeof(FBoneTransformWithScale)), TEXT("AnimSequenceTransformProvider.GlobalSpaceBoneTransforms"));
+
+	FGlobalShaderMap* ShaderMap = GetGlobalShaderMap(GMaxRHIFeatureLevel);
+
+	RDG_EVENT_SCOPE_STAT(GraphBuilder, AnimSequenceTransformProvider, "AnimSequenceTransformProvider");
+
+	// Dispatch 1.A: Evaluate animation sequences. Outputs animated local space bone transforms with retargeting.
+	if (NumTotalEvaluateBlockHeaders > 0)
+	{
+		FEvaluateAnimSequenceCS::FParameters* PassParameters = GraphBuilder.AllocParameters<FEvaluateAnimSequenceCS::FParameters>();
+		PassParameters->RetargetingDataBuffer = GraphBuilder.CreateSRV(RetargetingDataBuffer);
+		PassParameters->SequenceTransformBuffer = GraphBuilder.CreateSRV(SequenceTransformBuffer);
+		PassParameters->EvaluateBlockHeaderBuffer = GraphBuilder.CreateSRV(EvaluateBlockHeaderBuffer);
+		PassParameters->LocalSpaceTransformBufferUAV = GraphBuilder.CreateUAV(LocalSpaceTransformBuffer);
+
+		TShaderMapRef<FEvaluateAnimSequenceCS> ComputeShader(ShaderMap);
+
+		FComputeShaderUtils::AddPass(
+			GraphBuilder,
+			RDG_EVENT_NAME("EvaluateAnimSequence"),
+			ComputeShader,
+			PassParameters,
+			FIntVector(NumTotalEvaluateBlockHeaders, 1, 1)
+		);
+	}
+
+	// Dispatch 1.B: Evaluate animation sequences for blends. Outputs animated local space bone transforms with retargeting.
+	if (NumTotalEvaluateBlendBlockHeaders > 0)
+	{
+		FEvaluateAnimSequenceBlendCS::FParameters* PassParameters = GraphBuilder.AllocParameters<FEvaluateAnimSequenceBlendCS::FParameters>();
+		PassParameters->RetargetingDataBuffer = GraphBuilder.CreateSRV(RetargetingDataBuffer);
+		PassParameters->SequenceTransformBuffer = GraphBuilder.CreateSRV(SequenceTransformBuffer);
+		PassParameters->EvaluateBlockHeaderBuffer = GraphBuilder.CreateSRV(EvaluateBlendBlockHeaderBuffer);
+		PassParameters->LocalSpaceTransformBufferUAV = GraphBuilder.CreateUAV(LocalSpaceTransformBuffer);
+
+		TShaderMapRef<FEvaluateAnimSequenceBlendCS> ComputeShader(ShaderMap);
+
+		FComputeShaderUtils::AddPass(
+			GraphBuilder,
+			RDG_EVENT_NAME("EvaluateAnimSequenceBlend"),
+			ComputeShader,
+			PassParameters,
+			FIntVector(NumTotalEvaluateBlendBlockHeaders, 1, 1)
+		);
+	}
+
+	// Dispatch 2: Concatenate leaf bone transforms (local to global). Outputs global space bone transforms.
+	{
+		FConcatenateLocalToGlobalCS::FParameters* PassParameters = GraphBuilder.AllocParameters<FConcatenateLocalToGlobalCS::FParameters>();
+		PassParameters->LocalSpaceTransformBuffer = GraphBuilder.CreateSRV(LocalSpaceTransformBuffer);
+		PassParameters->ConcatenateBlockHeaderBuffer = GraphBuilder.CreateSRV(ConcatenateBlockHeaderBuffer);
+		PassParameters->LeafBonePathBuffer = GraphBuilder.CreateSRV(LeafBonePathBuffer);
+		PassParameters->GlobalSpaceTransformBufferUAV = GraphBuilder.CreateUAV(GlobalSpaceTransformBuffer);
+
+		TShaderMapRef<FConcatenateLocalToGlobalCS> ComputeShader(ShaderMap);
+
+		FComputeShaderUtils::AddPass(
+			GraphBuilder,
+			RDG_EVENT_NAME("ConcatenateLeafBoneTransforms"),
+			ComputeShader,
+			PassParameters,
+			FIntVector(NumTotalConcatenateBlockHeaders, 1, 1)
+		);
+	}
+
+	// Dispatch 3: Scatter transforms to final skinning scene buffer.
+	{
+		FScatterTransformsCS::FParameters* PassParameters = GraphBuilder.AllocParameters<FScatterTransformsCS::FParameters>();
+		PassParameters->GlobalSpaceTransformBuffer = GraphBuilder.CreateSRV(GlobalSpaceTransformBuffer);
+		PassParameters->SequenceTransformBuffer = GraphBuilder.CreateSRV(SequenceTransformBuffer);
+		PassParameters->ScatterBlockHeaderBuffer = GraphBuilder.CreateSRV(ScatterBlockHeaderBuffer);
+		PassParameters->BoneMapBuffer = Context.HierarchyBufferSRV;
+		PassParameters->SkinningTransformBufferUAV = GetCompressedBoneTransformUAV(GraphBuilder, Context.TransformBuffer);
+
+		TShaderMapRef<FScatterTransformsCS> ComputeShader(ShaderMap);
+
+		FComputeShaderUtils::AddPass(
+			GraphBuilder,
+			RDG_EVENT_NAME("ScatterTransforms"),
+			ComputeShader,
+			PassParameters,
+			FIntVector(NumTotalScatterBlockHeaders, 1, 1)
+		);
+	}
+}
+
+void FAnimSequenceTransformProvider::ProvideCPUTransforms(FSkinningTransformProvider::FProviderContext& Context)
+{
+	TRACE_CPUPROFILER_EVENT_SCOPE(FAnimSequenceTransformProvider::ProvideCPUTransforms);
+	FRDGBuilder& GraphBuilder = Context.GraphBuilder;
+
+	uint32 GlobalTransformCount = 0;
+
+	for (const FSkinningTransformProvider::FProviderIndirection Indirection : Context.Indirections)
+	{
+		const FInstancedSkinningSceneExtensionProxy* ExtensionProxy = static_cast<FInstancedSkinningSceneExtensionProxy*>(Context.Proxies[Indirection.Index]);
+		FAnimSequenceTransformProviderProxy* TransformProviderProxy = static_cast<FAnimSequenceTransformProviderProxy*>(ExtensionProxy->GetTransformProviderProxy());
+
+		const uint32 TransformCount = ExtensionProxy->GetMaxBoneTransformCount();
+		const uint32 ActiveAnimationCount = TransformProviderProxy->GetTracks().GetNumActiveTracks();
+		GlobalTransformCount += ActiveAnimationCount * TransformCount * 2; // Current and Previous
+	}
+
+	if (GlobalTransformCount == 0)
+	{
+		return;
+	}
+
+	Context.ScatterUploadBuilder.AddPass(GraphBuilder, SkinningTransformUploadBuffer, Context.TransformBuffer, GlobalTransformCount, sizeof(FCompressedBoneTransform), TEXT("Skinning.AnimTransforms"),
+		[this, Indirections = Context.Indirections, Proxies = Context.Proxies, &ViewFamily = Context.ViewFamily, GlobalTransformCount] (FRDGScatterUploader& ScatterUploader)
+	{
+		TRACE_CPUPROFILER_EVENT_SCOPE(FAnimSequenceTransformProvider::ProvideCPUTransformsTask);
+
+		BoneTrackArray TrackToBoneIndexMap;
+		TArray<FTransform, SceneRenderingAllocator> SampledLocalAnimPose;
+
+		const auto DecompressPose = [&TrackToBoneIndexMap, &SampledLocalAnimPose] (const FMapping& Mapping, const FSequence& Sequence, float TimeOffset, TArray<FTransform3f, SceneRenderingAllocator>& SampledGlobalMeshPose)
+		{
+			const UAnimSequence* AnimSequence = Sequence.Sequence;
+			UAnimSequence::FScopedCompressedAnimSequence CompressedAnimSequence = AnimSequence->GetCompressedData();
+			const FCompressedAnimSequence& PlatformCompressedData = CompressedAnimSequence.Get();
+
+			const uint32 NumCompressedTracks = PlatformCompressedData.CompressedTrackToSkeletonMapTable.Num();
+			TrackToBoneIndexMap.Reset(NumCompressedTracks);
+
+			for (uint32 CompressedTrackIndex = 0; CompressedTrackIndex < NumCompressedTracks; ++CompressedTrackIndex)
+			{
+				const uint32 BoneIndex = PlatformCompressedData.GetSkeletonIndexFromTrackIndex(CompressedTrackIndex);
+
+				// We only care about raw bones.
+				if (BoneIndex < Sequence.NumAnimBones)
+				{
+					TrackToBoneIndexMap.Add(BoneTrackPair(BoneIndex, CompressedTrackIndex));
+				}
+			}
+
+			FAnimSequenceDecompressionContext DecompressionContext(
+				  Sequence.SampleRate
+				, Sequence.NumFrames
+				, AnimSequence->Interpolation
+				, AnimSequence->GetRetargetTransformsSourceName()
+				, *PlatformCompressedData.CompressedDataStructure
+				, AnimSequence->GetSkeleton()->GetRefLocalPoses()
+				, PlatformCompressedData.CompressedTrackToSkeletonMapTable
+				, AnimSequence->GetSkeleton()
+				, AnimSequence->IsValidAdditive()
+				, AnimSequence->GetAdditiveAnimType()
+			);
+
+			SampledLocalAnimPose.Reset(DecompressionContext.GetRefLocalPoses().Num());
+			SampledLocalAnimPose.Append(DecompressionContext.GetRefLocalPoses());
+
+			DecompressionContext.Seek((double)TimeOffset);
+
+			if (PlatformCompressedData.BoneCompressionCodec)
+			{
+				TArrayView<FTransform> SampledLocalAnimPoseView(SampledLocalAnimPose);
+				PlatformCompressedData.BoneCompressionCodec->DecompressPose(DecompressionContext, TrackToBoneIndexMap, TrackToBoneIndexMap, TrackToBoneIndexMap, SampledLocalAnimPoseView);
+			}
+
+			const FReferenceSkeleton& MeshRefSkeleton = Mapping.SkinnedAsset->GetRefSkeleton();
+			TConstArrayView<FTransform> MeshLocalRefPoses = MeshRefSkeleton.GetRawRefBonePose();
+
+			for (int32 MeshBoneIndex = 0; MeshBoneIndex < (int32)Mapping.NumMeshBones; ++MeshBoneIndex)
+			{
+				const int32 AnimBoneIndex = Mapping.BoneRetargetingTables.AnimBoneIndex[MeshBoneIndex];
+
+				FTransform3f MeshLocalTransform;
+
+				if (AnimBoneIndex != INDEX_NONE)
+				{
+					const FTransform3f RetargetTransform = Mapping.BoneRetargetingTables.Transform[MeshBoneIndex];
+					MeshLocalTransform = FTransform3f(SampledLocalAnimPose[AnimBoneIndex]);
+
+					if (Mapping.BoneRetargetingTables.bRequiresRefPoseRetarget && MeshBoneIndex != 0)
+					{
+						const FQuat4f& Q0 = Mapping.BoneRetargetingTables.RefPoseRetargetQuats[MeshBoneIndex].Get<0>();
+						const FQuat4f& Q1 = Mapping.BoneRetargetingTables.RefPoseRetargetQuats[MeshBoneIndex].Get<1>();
+
+						FQuat4f RetargetedRotation = Q0 * MeshLocalTransform.GetRotation() * Q1;
+						RetargetedRotation.Normalize();
+						MeshLocalTransform.SetRotation(RetargetedRotation);
+ 
+						if (Mapping.BoneRetargetingTables.Mode[MeshBoneIndex] != EBoneTranslationRetargetingMode::Skeleton)
+						{
+							MeshLocalTransform.SetTranslation(Q0.RotateVector(MeshLocalTransform.GetTranslation()));
+						}
+					}
+
+					switch (Mapping.BoneRetargetingTables.Mode[MeshBoneIndex])
+					{
+						case EBoneTranslationRetargetingMode::Animation:
+						{
+							// Identity
+						}
+						break;
+						case EBoneTranslationRetargetingMode::Skeleton:
+						{
+							MeshLocalTransform.SetTranslation(RetargetTransform.GetTranslation());
+							if (MeshBoneIndex == 0)
+							{
+								MeshLocalTransform.SetRotation(RetargetTransform.GetRotation());
+							}
+						}
+						break;
+						case EBoneTranslationRetargetingMode::AnimationScaled:
+						{
+							MeshLocalTransform.ScaleTranslation(RetargetTransform.GetScale3D());
+						}
+						break;
+						case EBoneTranslationRetargetingMode::AnimationRelative:
+						{
+							MeshLocalTransform.SetRotation(MeshLocalTransform.GetRotation() * RetargetTransform.GetRotation());
+							MeshLocalTransform.SetTranslation(MeshLocalTransform.GetTranslation() + RetargetTransform.GetTranslation());
+							MeshLocalTransform.SetScale3D(MeshLocalTransform.GetScale3D() * RetargetTransform.GetScale3D());
+							MeshLocalTransform.NormalizeRotation();
+						}
+						break;
+						case EBoneTranslationRetargetingMode::OrientAndScale:
+						{
+							MeshLocalTransform.SetTranslation(RetargetTransform.GetRotation().RotateVector(MeshLocalTransform.GetTranslation()) * RetargetTransform.GetScale3D());
+							break;
+						}
+					}
+				}
+				else
+				{
+					MeshLocalTransform = FTransform3f(MeshLocalRefPoses[MeshBoneIndex]);
+				}
+
+				const int32 ParentIndex = MeshRefSkeleton.GetParentIndex(MeshBoneIndex);
+
+				if (SampledGlobalMeshPose.IsValidIndex(ParentIndex) && ParentIndex < MeshBoneIndex)
+				{
+					SampledGlobalMeshPose[MeshBoneIndex] = MeshLocalTransform * SampledGlobalMeshPose[ParentIndex];
+				}
+				else
+				{
+					SampledGlobalMeshPose[MeshBoneIndex] = MeshLocalTransform;
+				}
+			}
+		};
+
+		TArray<FTransform3f, SceneRenderingAllocator> SourceSampledGlobalMeshPose;
+		TArray<FTransform3f, SceneRenderingAllocator> TargetSampledGlobalMeshPose;
+		TArray<FCompressedBoneTransform, SceneRenderingAllocator> SourceTransforms;
+		TBitArray<SceneRenderingBitArrayAllocator> SourceTransformMap;
+
+		const auto EvaluateAnimSequence = [&SourceTransforms, &SourceTransformMap, &SourceSampledGlobalMeshPose, &TargetSampledGlobalMeshPose, &DecompressPose] (
+			const FInstancedSkinningSceneExtensionProxy* ExtensionProxy,
+			TArrayView<FCompressedBoneTransform> Transforms,
+			TArrayView<FCompressedBoneTransform> DuplicatePreviousTransforms,
+			const FSequence& SourceSequence,
+			const FSequence& TargetSequence,
+			const FMapping& SourceMapping,
+			const FMapping& TargetMapping,
+			float SourceTimeOffset,
+			float TargetTimeOffset,
+			float BlendWeight)
+		{
+			DecompressPose(SourceMapping, SourceSequence, SourceTimeOffset, SourceSampledGlobalMeshPose);
+
+			if (BlendWeight > 0.0f)
+			{
+				DecompressPose(TargetMapping, TargetSequence, TargetTimeOffset, TargetSampledGlobalMeshPose);
+			}
+
+			for (uint32 MeshBoneIndex = 0; MeshBoneIndex < SourceMapping.NumMeshBones; ++MeshBoneIndex)
+			{
+				if (BlendWeight > 0.0f)
+				{
+					SourceSampledGlobalMeshPose[MeshBoneIndex].BlendWith(TargetSampledGlobalMeshPose[MeshBoneIndex], BlendWeight);
+					SourceSampledGlobalMeshPose[MeshBoneIndex].NormalizeRotation();
+				}
+
+				SourceSampledGlobalMeshPose[MeshBoneIndex] = SourceMapping.InvGlobalRefPoseTransforms[MeshBoneIndex] * SourceSampledGlobalMeshPose[MeshBoneIndex];
+			}
+
+			if (ExtensionProxy->GetBoneTransformStorageMode() == EBoneTransformStorageMode::BoneMap)
+			{
+				TConstArrayView<uint32> BoneHierarchy = ExtensionProxy->GetBoneHierarchy();
+				SourceTransformMap.Init(false, TargetMapping.NumMeshBones);
+				SourceTransforms.SetNumUninitialized(TargetMapping.NumMeshBones, EAllowShrinking::No);
+				int32 DstTransformIndex = 0;
+
+				for (uint32 MeshBoneIndex : ExtensionProxy->GetBoneHierarchy())
+				{
+					if (!SourceTransformMap[MeshBoneIndex])
+					{
+						SourceTransformMap[MeshBoneIndex] = true;
+						StoreCompressedBoneTransform(SourceTransforms[MeshBoneIndex], SourceSampledGlobalMeshPose[MeshBoneIndex].ToMatrixWithScale());
+					}
+
+					Transforms[DstTransformIndex] = SourceTransforms[MeshBoneIndex];
+
+					if (!DuplicatePreviousTransforms.IsEmpty())
+					{
+						DuplicatePreviousTransforms[DstTransformIndex] = SourceTransforms[MeshBoneIndex];
+					}
+
+					DstTransformIndex++;
+				}
+			}
+			else
+			{
+				for (uint32 MeshBoneIndex = 0; MeshBoneIndex < TargetMapping.NumMeshBones; ++MeshBoneIndex)
+				{
+					StoreCompressedBoneTransform(Transforms[MeshBoneIndex], SourceSampledGlobalMeshPose[MeshBoneIndex].ToMatrixWithScale());
+
+					if (!DuplicatePreviousTransforms.IsEmpty())
+					{
+						DuplicatePreviousTransforms[MeshBoneIndex] = Transforms[MeshBoneIndex];
+					}
+				}
+			}
+		};
+
+		const float GlobalTimeScale = CVarAnimSequenceTransformProviderTimeScale.GetValueOnRenderThread();
+		const float DeltaTime = ViewFamily.Time.GetDeltaWorldTimeSeconds() * GlobalTimeScale;
+
+		uint32 NumActiveTracks = 0;
+		uint32 NumCurrentTransformUpdates = 0;
+		uint32 NumPreviousTransformUpdates = 0;
+
+		for (const FSkinningTransformProvider::FProviderIndirection Indirection : Indirections)
+		{
+			const FInstancedSkinningSceneExtensionProxy* ExtensionProxy = static_cast<FInstancedSkinningSceneExtensionProxy*>(Proxies[Indirection.Index]);
+			FProxyUserData& ProxyUserData = static_cast<FProxyUserData&>(ExtensionProxy->GetUserDataChecked());
+			FAnimSequenceTrackRenderPool& Tracks = ProxyUserData.Tracks;
+			ProxyUserData.TrackHistories.SetNum(Tracks.GetNumTracks());
+
+			const uint32 MaxTransformCount      = ExtensionProxy->GetMaxBoneTransformCount();
+			const uint32 UniqueAnimationCount   = ExtensionProxy->GetUniqueAnimationCount();
+			const uint32 MaxTotalTransformCount = MaxTransformCount * 2u; // Current and Previous
+			const uint32 CurrentTransformIndex  = Indirection.CurrentTransformOffset  / sizeof(FCompressedBoneTransform);
+			const uint32 PreviousTransformIndex = Indirection.PreviousTransformOffset / sizeof(FCompressedBoneTransform);
+			const uint32 CurrentTransformSlot   = CurrentTransformIndex < PreviousTransformIndex ? 0 : 1;
+
+			SourceSampledGlobalMeshPose.SetNumUninitialized(ProxyUserData.NumMeshBones, EAllowShrinking::No);
+			TargetSampledGlobalMeshPose.SetNumUninitialized(ProxyUserData.NumMeshBones, EAllowShrinking::No);
+	
+			FMemMark Mark(FMemStack::Get());
+
+			Tracks.EnumerateActiveTracks([&] (int32 TrackIndex)
+			{
+				FAnimSequenceTrackRenderData& Data = Tracks.GetData(TrackIndex);
+
+				const FAnimSequenceTrackPackedData& Source = Data.GetSource();
+				const FAnimSequenceTrackPackedData& Target = Data.GetTarget();
+
+				const FMapping& SourceMapping = ProxyUserData.GetMapping(Source);
+				const FMapping& TargetMapping = ProxyUserData.GetMapping(Target);
+				FSequence& SourceSequence = ProxyUserData.GetSequence(Source);
+				FSequence& TargetSequence = ProxyUserData.GetSequence(Target);
+				
+				// Dirty previous transforms means the whole allocation is dirty. Rebuild everything.
+				const bool bDirtyTransformAllocation = EnumHasAnyFlags(Indirection.DirtyBoneTransforms, EDirtyBoneTransforms::Previous);
+				const FTrackHistory::FUpdateResult UpdateResult = ProxyUserData.TrackHistories[TrackIndex].Update(Data, CurrentTransformSlot, bDirtyTransformAllocation);
+	
+				const float SourcePosition         = Source.GetPosition();
+				const float SourcePreviousPosition = Source.GetPreviousPosition();
+				const float TargetPosition         = Target.GetPosition();
+				const float TargetPreviousPosition = Target.GetPreviousPosition();
+				const float BlendWeight            = Data.GetBlendWeight();
+
+				EPreviousBoneTransformUpdateMode UpdateMode = EPreviousBoneTransformUpdateMode::None;
+				TArrayView<FCompressedBoneTransform> DstTransformsPrevious;
+
+				if (UpdateResult.bNeedsPrevious)
+				{
+					UpdateMode = UpdateResult.bNeedsCurrent && UpdateResult.bNeedsPrevious && TargetPosition == TargetPreviousPosition && SourcePosition == SourcePreviousPosition
+						? EPreviousBoneTransformUpdateMode::DuplicateCurrentToPrevious
+						: EPreviousBoneTransformUpdateMode::UpdatePrevious;
+
+					DstTransformsPrevious = ScatterUploader.Add_GetRef<FCompressedBoneTransform>(PreviousTransformIndex + TrackIndex * MaxTotalTransformCount, MaxTransformCount);
+				}
+
+				if (UpdateResult.bNeedsCurrent)
+				{
+					TArrayView<FCompressedBoneTransform> DstTransforms = ScatterUploader.Add_GetRef<FCompressedBoneTransform>(CurrentTransformIndex + TrackIndex * MaxTotalTransformCount, MaxTransformCount);
+					EvaluateAnimSequence(
+						ExtensionProxy,
+						DstTransforms,
+						UpdateMode == EPreviousBoneTransformUpdateMode::DuplicateCurrentToPrevious ? DstTransformsPrevious : TArrayView<FCompressedBoneTransform>{},
+						SourceSequence, TargetSequence,
+						SourceMapping, TargetMapping,
+						SourcePosition, TargetPosition,
+						BlendWeight);
+
+					NumCurrentTransformUpdates++;
+				}
+
+				if (UpdateMode == EPreviousBoneTransformUpdateMode::UpdatePrevious)
+				{
+					EvaluateAnimSequence(
+						ExtensionProxy,
+						DstTransformsPrevious,
+						{},
+						SourceSequence, TargetSequence,
+						SourceMapping, TargetMapping,
+						SourcePreviousPosition, TargetPreviousPosition,
+						BlendWeight);
+
+					NumPreviousTransformUpdates++;
+				}
+
+				Data.Tick(SourceSequence.PlayLength, TargetSequence.PlayLength, DeltaTime);
+				NumActiveTracks++;
+			});
+		}
+
+		INC_DWORD_STAT_BY(STAT_AnimSequenceTransformProvider_ActiveTracks, NumActiveTracks);
+		INC_DWORD_STAT_BY(STAT_AnimSequenceTransformProvider_CurrentTransformUpdates, NumCurrentTransformUpdates);
+		INC_DWORD_STAT_BY(STAT_AnimSequenceTransformProvider_PreviousTransformUpdates, NumPreviousTransformUpdates);
+	});
+}
\ No newline at end of file
diff --git a/Engine/Source/Runtime/Renderer/Private/Skinning/AnimSequenceTransformProvider.h b/Engine/Source/Runtime/Renderer/Private/Skinning/AnimSequenceTransformProvider.h
new file mode 100644
--- /dev/null
+++ b/Engine/Source/Runtime/Renderer/Private/Skinning/AnimSequenceTransformProvider.h
@@ -0,0 +1,276 @@
+// Copyright Epic Games, Inc. All Rights Reserved.
+
+#pragma once
+
+#include "Animation/AnimSequenceTransformProviderData.h"
+#include "Animation/Skeleton.h"
+#include "Delegates/Delegate.h"
+#include "Delegates/DelegateCombinations.h"
+#include "RendererPrivateUtils.h"
+#include "SceneExtensions.h"
+#include "SkinningTransformProvider.h"
+#include "SkinningSceneExtensionProxy.h"
+#include "UnifiedBuffer.h"
+#include "SpanAllocator.h"
+
+class FAnimSequenceTransformProvider
+	: public ISceneExtension
+	, public IAnimSequenceTransformProvider
+{
+	DECLARE_SCENE_EXTENSION(RENDERER_API, FAnimSequenceTransformProvider);
+
+public:
+	FAnimSequenceTransformProvider(FScene& InScene);
+
+	static bool ShouldCreateExtension(FScene& InScene);
+
+	virtual void InitExtension(FScene& InScene) override;
+
+	void RegisterProxy(FAnimSequenceTransformProviderProxy* Proxy) override;
+	void UnregisterProxy(FAnimSequenceTransformProviderProxy* Proxy) override;
+
+private:
+	void ProvideGPUTransforms(FSkinningTransformProvider::FProviderContext& Context);
+	void ProvideCPUTransforms(FSkinningTransformProvider::FProviderContext& Context);
+
+	struct FBoneRetargetingTables
+	{
+		TArray<EBoneTranslationRetargetingMode::Type> Mode;
+		TArray<FTransform3f> Transform;
+		TArray<int32> AnimBoneIndex;
+		TArray<TTuple<FQuat4f, FQuat4f>> RefPoseRetargetQuats;
+		bool bRequiresRefPoseRetarget = false;
+	};
+
+	void BuildBoneRetargetingTables(FBoneRetargetingTables& BoneRetargetingTables, USkeleton* AnimSkeleton, const USkinnedAsset* MeshAsset);
+
+	struct FLeafBonePathData
+	{
+		TArray<uint32> Buffer;
+		uint32 NumLeafBones;
+	};
+
+	void BuildLeafBonePathData(FLeafBonePathData& OutLeafBonePathData, const FReferenceSkeleton& RefSkeleton);
+
+	struct FBufferAllocation
+	{
+		bool IsValid() const
+		{
+			return Count != 0;
+		}
+
+		uint32 Offset = 0;
+		uint32 Count  = 0;
+	};
+
+	struct FBoneTransformWithScale
+	{
+		void SetTransform(const FTransform3f& Xfm)
+		{
+			SetRotation(Xfm.GetRotation());
+			SetTranslation(Xfm.GetTranslation());
+			SetScale(Xfm.GetScale3D());
+		}
+
+		void SetRotation(FQuat4f Q)
+		{
+			Rotation[0] = Q.X;
+			Rotation[1] = Q.Y;
+			Rotation[2] = Q.Z;
+			Rotation[3] = Q.W;
+		}
+
+		void SetTranslation(FVector3f P)
+		{
+			Translation[0] = P.X;
+			Translation[1] = P.Y;
+			Translation[2] = P.Z;
+		}
+
+		void SetScale(FVector3f S)
+		{
+			Scale[0] = S.X;
+			Scale[1] = S.Y;
+			Scale[2] = S.Z;
+		}
+
+		float Rotation[4];
+		float Translation[3];
+		float Scale[3];
+	};
+
+	using FSequenceRegistryKey = const UAnimSequence*;
+
+	struct FMappingRegistryKey
+	{
+		FMappingRegistryKey() = default;
+
+		FMappingRegistryKey(const USkeleton* InSourceSkeleton, const USkinnedAsset* InSkinnedAsset)
+			: SourceSkeleton(InSourceSkeleton)
+			, SkinnedAsset(InSkinnedAsset)
+		{}
+
+		const USkeleton* SourceSkeleton = nullptr;
+		const USkinnedAsset* SkinnedAsset = nullptr;
+
+		bool operator==(const FMappingRegistryKey& Other) const
+		{
+			return SourceSkeleton == Other.SourceSkeleton && SkinnedAsset == Other.SkinnedAsset;
+		}
+
+		friend uint32 GetTypeHash(const FMappingRegistryKey& Key)
+		{
+			return PointerHash(Key.SourceSkeleton, PointerHash(Key.SkinnedAsset));
+		}
+	};
+
+	struct FSequence
+	{
+		FSequenceRegistryKey RegistryKey;
+		FBufferAllocation BufferAllocation;
+		const UAnimSequence* Sequence = nullptr;
+		FFrameRate SampleRate;
+		TBitArray<> FramesRequested;
+		TBitArray<> FramesUploaded;
+		uint32 NumFramesRequested = 0;
+		uint32 NumFramesUploaded = 0;
+		float PlayLength = 0.0f;
+		uint32 RefCount = 0;
+		uint32 NumFrames = 0;
+		uint32 NumAnimBones = 0;
+		uint32 NumSampleRateLevels = 0;
+		uint32 CurrentSampleRateLevel = 0;
+		uint32 TargetSampleRateLevel = 0;
+		uint32 NextFrameToUpload = 0;
+
+		int32 RequestFrame(int32 Index)
+		{
+			// Only request if at the lowest sample rate level (i.e. ignore the budget) and we haven't uploaded it already.
+			if (!FramesUploaded[Index] && CurrentSampleRateLevel == NumSampleRateLevels)
+			{
+				auto Bit = FramesRequested[Index];
+				if (!Bit)
+				{
+					Bit = true;
+					NumFramesRequested++;
+					return 1;
+				}
+			}
+			return 0;
+		}
+
+		uint32 GetSampleStep() const;
+
+		void GetKeyIndicesFromTime(int32& OutKey0, int32& OutKey1, float& OutAlpha, float Time) const;
+	};
+
+	struct FMapping
+	{
+		FMappingRegistryKey RegistryKey;
+		FBufferAllocation MappingTransformBufferAllocation;
+		FBufferAllocation RetargetingDataBufferAllocation;
+		TArray<FTransform3f> InvGlobalRefPoseTransforms;
+		uint32 RetargetingTransformBufferOffset = 0;
+		uint32 InvGlobalRefPoseTransformBufferOffset = 0;
+		FBoneRetargetingTables BoneRetargetingTables;
+		const USkeleton* SourceSkeleton = nullptr;
+		const USkinnedAsset* SkinnedAsset = nullptr;
+		uint32 RefCount = 0;
+		uint32 NumMeshBones = 0;
+	};
+
+	class FTrackHistory
+	{
+	public:
+		static constexpr float Invalid = UE_MAX_FLT;
+
+		struct FUpdateResult
+		{
+			uint8 bNeedsCurrent  : 1 = 0;
+			uint8 bNeedsPrevious : 1 = 0;
+		};
+
+		FUpdateResult Update(const FAnimSequenceTrackRenderData& RenderData, int32 SlotIndex, bool bForceUpdate)
+		{
+			FUpdateResult Result;
+			Result.bNeedsCurrent  = bForceUpdate;
+			Result.bNeedsPrevious = bForceUpdate;
+			UpdateSlot(SourceSequenceIndex, SourceSlots, RenderData.GetSource(), SlotIndex, Result);
+			UpdateSlot(TargetSequenceIndex, TargetSlots, RenderData.GetTarget(), SlotIndex, Result);
+			return Result;
+		}
+
+	private:
+		void UpdateSlot(int32& SequenceIndex, FVector2f& Slots, const FAnimSequenceTrackPackedData& Data, int32 InCurrentSlot, FUpdateResult& Result)
+		{
+			if (SequenceIndex != Data.GetSequenceIndex())
+			{
+				SequenceIndex = Data.GetSequenceIndex();
+				Slots = FVector2f{ Invalid, Invalid };
+			}
+
+			Result.bNeedsCurrent  |= Slots[ InCurrentSlot] != Data.GetPosition();
+			Result.bNeedsPrevious |= Slots[!InCurrentSlot] != Data.GetPreviousPosition() && Data.HasPreviousPosition();
+			Slots[InCurrentSlot] = Data.GetPosition();
+
+			if (Data.HasPreviousPosition())
+			{
+				Slots[!InCurrentSlot] = Data.GetPreviousPosition();
+			}
+		}
+
+		int32 SourceSequenceIndex = INDEX_NONE;
+		int32 TargetSequenceIndex = INDEX_NONE;
+		FVector2f SourceSlots{ Invalid, Invalid };
+		FVector2f TargetSlots{ Invalid, Invalid };
+	};
+
+	class FProxyUserData : public FInstancedSkinningSceneExtensionProxyUserData
+	{
+	public:
+		FProxyUserData(FAnimSequenceTransformProviderProxy* Proxy);
+
+		const FMapping& GetMapping(const FAnimSequenceTrackPackedData& Data) const
+		{
+			const int32 SequenceIndex = SequenceToMappingTable.IsValidIndex(Data.GetSequenceIndex()) ? Data.GetSequenceIndex() : 0;
+			return *SequenceToMappingTable[SequenceIndex];
+		}
+
+		FSequence& GetSequence(const FAnimSequenceTrackPackedData& Data) const
+		{
+			return *Sequences[Sequences.IsValidIndex(Data.GetSequenceIndex()) ? Data.GetSequenceIndex() : 0];
+		}
+
+		FAnimSequenceTrackRenderPool& Tracks;
+		TArray<FTrackHistory> TrackHistories;
+		const FInstancedSkinningSceneExtensionProxy& SceneExtensionProxy;
+		TArray<FSequence*> Sequences;
+		TArray<FMapping*> SequenceToMappingTable;
+		TArray<FMapping*> Mappings;
+		FLeafBonePathData LeafBonePathData;
+		FBufferAllocation LeafBonePathBufferAllocation;
+		uint32 NumMeshBones = 0;
+		const EBoneTransformStorageMode TransformStorageMode;
+	};
+
+	TMap<FMappingRegistryKey, TUniquePtr<FMapping>> MappingRegistry;
+	TMap<FSequenceRegistryKey, TUniquePtr<FSequence>> SequenceRegistry;
+
+	TArray<FSequence*> SequencesToUpload;
+	TArray<FMapping*> MappingsToUpload;
+	TArray<FProxyUserData*> ProxyUserDatasToUpload;
+
+	FSpanAllocator SequenceTransformAllocator;
+	FRDGAsyncScatterUploadBuffer SequenceTransformUploadBuffer;
+	TPersistentByteAddressBuffer<FBoneTransformWithScale> SequenceTransformPersistentBuffer;
+
+	FSpanAllocator RetargetingDataAllocator;
+	FRDGAsyncScatterUploadBuffer RetargetingDataUploadBuffer;
+	TPersistentByteAddressBuffer<UE::HLSL::FBoneRetargetingData> RetargetingDataPersistentBuffer;
+
+	FSpanAllocator LeafBonePathAllocator;
+	FRDGAsyncScatterUploadBuffer LeafBonePathUploadBuffer;
+	TPersistentByteAddressBuffer<uint32> LeafBonePathPersistentBuffer;
+
+	FRDGAsyncScatterUploadBuffer SkinningTransformUploadBuffer;
+};
\ No newline at end of file
