[PostProcessing] Fixed VisualizePostProcessingStack not working correctly when it is the last operation in the stack.

#Synced-CL 49896507

diff --git a/Engine/Source/Runtime/Renderer/Private/PostProcess/PostProcessing.cpp b/Engine/Source/Runtime/Renderer/Private/PostProcess/PostProcessing.cpp
--- a/Engine/Source/Runtime/Renderer/Private/PostProcess/PostProcessing.cpp
+++ b/Engine/Source/Runtime/Renderer/Private/PostProcess/PostProcessing.cpp
@@ -232,7 +232,7 @@ bool IsPostProcessingWithAlphaChannelSupported()
 }
 
 #if DEBUG_POST_PROCESS_VOLUME_ENABLE
-FScreenPassTexture AddFinalPostProcessDebugInfoPasses(FRDGBuilder& GraphBuilder, const FViewInfo& View, FScreenPassTexture& ScreenPassSceneColor);
+FScreenPassTexture AddFinalPostProcessDebugInfoPasses(FRDGBuilder& GraphBuilder, const FViewInfo& View, FScreenPassTexture& ScreenPassSceneColor, FScreenPassRenderTarget& OverrideOutput);
 #endif
 
 #if !UE_BUILD_SHIPPING
@@ -1734,8 +1734,7 @@ void AddPostProcessingPasses(
 	{
 		FScreenPassRenderTarget OverrideOutput;
 		PassSequence.AcceptOverrideIfLastPass(EPass::VisualizePostProcessStack, OverrideOutput);
-		OverrideOutput = OverrideOutput.IsValid() ? OverrideOutput : FScreenPassRenderTarget(SceneColor, ERenderTargetLoadAction::ELoad);
-		SceneColor = AddFinalPostProcessDebugInfoPasses(GraphBuilder, View, OverrideOutput);
+		SceneColor = AddFinalPostProcessDebugInfoPasses(GraphBuilder, View, SceneColor, OverrideOutput);
 	}
 #endif
 
@@ -3709,11 +3708,24 @@ FRDGTextureRef AddProcessPlanarReflectionPass(
 }
 
 #if DEBUG_POST_PROCESS_VOLUME_ENABLE
-FScreenPassTexture AddFinalPostProcessDebugInfoPasses(FRDGBuilder& GraphBuilder, const FViewInfo& View, FScreenPassTexture& ScreenPassSceneColor)
+FScreenPassTexture AddFinalPostProcessDebugInfoPasses(FRDGBuilder& GraphBuilder, const FViewInfo& View, FScreenPassTexture& ScreenPassSceneColor, FScreenPassRenderTarget& OverrideOutput)
 {
 	RDG_EVENT_SCOPE(GraphBuilder, "FinalPostProcessDebugInfo");
 
-	AddDrawCanvasPass(GraphBuilder, RDG_EVENT_NAME("PostProcessDebug"), View, FScreenPassRenderTarget(ScreenPassSceneColor, ERenderTargetLoadAction::ELoad),
+	FScreenPassRenderTarget Output;
+	if (OverrideOutput.IsValid())
+	{
+		Output = OverrideOutput;
+		const FIntPoint MaxSize = Output.Texture->Desc.Extent - OverrideOutput.ViewRect.Min;
+		const FIntPoint Size{ FMath::Min(ScreenPassSceneColor.ViewRect.Size().X, MaxSize.X), FMath::Min(ScreenPassSceneColor.ViewRect.Size().Y, MaxSize.Y) };
+		AddCopyTexturePass(GraphBuilder, ScreenPassSceneColor.Texture, Output.Texture, ScreenPassSceneColor.ViewRect.Min, OverrideOutput.ViewRect.Min, Size);
+	}
+	else
+	{
+		Output = FScreenPassRenderTarget(ScreenPassSceneColor, ERenderTargetLoadAction::ELoad);
+	}
+
+	AddDrawCanvasPass(GraphBuilder, RDG_EVENT_NAME("PostProcessDebug"), View, Output,
 		[&View](FCanvas& Canvas)
 		{
 			const float DPIScale = Canvas.GetDPIScale();
@@ -3790,7 +3802,7 @@ FScreenPassTexture AddFinalPostProcessDebugInfoPasses(FRDGBuilder& GraphBuilder,
 			}
 		});
 
-	return MoveTemp(ScreenPassSceneColor);
+	return MoveTemp(Output);
 }
 #endif
 
