UE: Ensure instantiation of a couple functions in TMobileBasePassPSPolicyParamType<FUniformLightMapPolicy>. There can be link errors since they are potentially used in other translation units

#Synced-CL 50014479

diff --git a/Engine/Source/Runtime/Renderer/Private/MobileBasePass.cpp b/Engine/Source/Runtime/Renderer/Private/MobileBasePass.cpp
--- a/Engine/Source/Runtime/Renderer/Private/MobileBasePass.cpp
+++ b/Engine/Source/Runtime/Renderer/Private/MobileBasePass.cpp
@@ -267,6 +267,15 @@ template bool MobileBasePass::GetShaders<FMobileSingleLayerWaterPolicy>(
 	TShaderRef<TMobileBasePassVSPolicyParamType<FMobileSingleLayerWaterPolicy>>& VertexShader,
 	TShaderRef<TMobileBasePassPSPolicyParamType<FMobileSingleLayerWaterPolicy>>& PixelShader);
 
+// Force instantination for FUniformLightMapPolicy as it's used in a different translation unit
+template bool MobileBasePass::GetShaders<FUniformLightMapPolicy>(
+	const FUniformLightMapPolicy& RESTRICT LightMapPolicy,
+	EMobileLocalLightSetting LocalLightSetting,
+	const FMaterial& MaterialResource,
+	const FVertexFactoryType* VertexFactoryType,
+	TShaderRef<TMobileBasePassVSPolicyParamType<FUniformLightMapPolicy>>& VertexShader,
+	TShaderRef<TMobileBasePassPSPolicyParamType<FUniformLightMapPolicy>>& PixelShader);
+
 static bool UseSkyReflectionCapture(const FScene* RenderScene)
 {
 	return RenderScene
@@ -748,6 +757,16 @@ template void TMobileBasePassPSPolicyParamType<FMobileSingleLayerWaterPolicy>::G
 	const TMobileBasePassShaderElementData<FMobileSingleLayerWaterPolicy>& ShaderElementData,
 	FMeshDrawSingleShaderBindings& ShaderBindings) const;
 
+// Force instantination for FUniformLightMapPolicy as it's used in a different translation unit
+template void TMobileBasePassPSPolicyParamType<FUniformLightMapPolicy>::GetShaderBindings(
+	const FScene* Scene,
+	const FStaticFeatureLevel FeatureLevel,
+	const FPrimitiveSceneProxy* PrimitiveSceneProxy,
+	const FMaterialRenderProxy& MaterialRenderProxy,
+	const FMaterial& Material,
+	const TMobileBasePassShaderElementData<FUniformLightMapPolicy>& ShaderElementData,
+	FMeshDrawSingleShaderBindings& ShaderBindings) const;
+
 
 FMobileBasePassMeshProcessor::FMobileBasePassMeshProcessor(
 	EMeshPass::Type InMeshPassType,
