Mobile: Fixed "Lit Wireframe" debug view mode in configurations with a full depth pre-pass enabled

#Synced-CL 50034370

diff --git a/Engine/Source/Runtime/Renderer/Private/MobileShadingRenderer.cpp b/Engine/Source/Runtime/Renderer/Private/MobileShadingRenderer.cpp
--- a/Engine/Source/Runtime/Renderer/Private/MobileShadingRenderer.cpp
+++ b/Engine/Source/Runtime/Renderer/Private/MobileShadingRenderer.cpp
@@ -819,10 +819,19 @@ void FMobileSceneRenderer::RenderFullDepthPrepass(FRDGBuilder& GraphBuilder, TAr
 	FRenderViewContextArray RenderViews;
 	GetRenderViews(InViews, RenderViews);
 
+	FRDGTextureRef DepthTargetTexture = SceneTextures.Depth.Target;
+
+#if WITH_EDITOR
+	if (ViewFamily.EngineShowFlags.Wireframe)
+	{
+		DepthTargetTexture = SceneTextures.EditorPrimitiveDepth;
+	}
+#endif
+
 	TRDGUniformBufferBinding<FMobileBasePassUniformParameters> LastViewMobileBasePassUB;
 	if (DepthPass.IsRasterStencilDitherEnabled())
 	{
-		AddDitheredStencilFillPass(GraphBuilder, InViews, SceneTextures.Depth.Target, DepthPass);
+		AddDitheredStencilFillPass(GraphBuilder, InViews, DepthTargetTexture, DepthPass);
 	}
 
 	for (FRenderViewContext& ViewContext : RenderViews)
@@ -846,8 +855,9 @@ void FMobileSceneRenderer::RenderFullDepthPrepass(FRDGBuilder& GraphBuilder, TAr
 			PassParameters->RenderTargets[0] = FRenderTargetBinding(SceneTextures.Velocity, ERenderTargetLoadAction::EClear);
 		}
 		PassParameters->RenderTargets.DepthStencil = DepthPass.IsRasterStencilDitherEnabled() ?
-			FDepthStencilBinding(SceneTextures.Depth.Target, ERenderTargetLoadAction::ELoad, ERenderTargetLoadAction::ELoad, FExclusiveDepthStencil::DepthWrite_StencilWrite):
-			FDepthStencilBinding(SceneTextures.Depth.Target, ERenderTargetLoadAction::EClear, ERenderTargetLoadAction::EClear, FExclusiveDepthStencil::DepthWrite_StencilWrite);
+			FDepthStencilBinding(DepthTargetTexture, ERenderTargetLoadAction::ELoad, ERenderTargetLoadAction::ELoad, FExclusiveDepthStencil::DepthWrite_StencilWrite):
+			FDepthStencilBinding(DepthTargetTexture, ERenderTargetLoadAction::EClear, ERenderTargetLoadAction::EClear, FExclusiveDepthStencil::DepthWrite_StencilWrite);
+
 		PassParameters->View = View.GetShaderParameters();
 		PassParameters->MobileBasePass = CreateMobileBasePassUniformBuffer(GraphBuilder, View, EMobileBasePass::DepthPrePass, EMobileSceneTextureSetupMode::None);
 		//if the scenecolor isn't multiview but the app is, need to render as a single-view multiview due to shaders
@@ -875,25 +885,25 @@ void FMobileSceneRenderer::RenderFullDepthPrepass(FRDGBuilder& GraphBuilder, TAr
 			VelocityInstanceCullingDrawParams = View.ParallelMeshDrawCommandPasses[EMeshPass::DepthPass] ? &ParameterCollection->VelocityInstanceCullingDrawParams : &PassParameters->InstanceCullingDrawParams;
 			BuildMeshRenderingCommands(GraphBuilder, EMeshPass::Velocity, View, Scene->GPUScene, InstanceCullingManager, *VelocityInstanceCullingDrawParams);		}
 
-		GraphBuilder.AddPass(
-			RDG_EVENT_NAME("FullDepthPrepass"),
-			PassParameters,
-			ERDGPassFlags::Raster,
-			[this, ParameterCollection, &View, VelocityInstanceCullingDrawParams](FRHICommandList& RHICmdList)
-			{
-				auto PassParameters = &ParameterCollection->PassParameters;
-				RenderPrePass(RHICmdList, View, &PassParameters->InstanceCullingDrawParams);
-				if (VelocityInstanceCullingDrawParams)
+			GraphBuilder.AddPass(
+				RDG_EVENT_NAME("FullDepthPrepass"),
+				PassParameters,
+				ERDGPassFlags::Raster,
+				[this, ParameterCollection, &View, VelocityInstanceCullingDrawParams](FRHICommandList& RHICmdList)
 				{
-					RenderVelocityPass(RHICmdList, View, VelocityInstanceCullingDrawParams);
-				}
-			});
+					auto PassParameters = &ParameterCollection->PassParameters;
+					RenderPrePass(RHICmdList, View, &PassParameters->InstanceCullingDrawParams);
+					if (VelocityInstanceCullingDrawParams)
+					{
+						RenderVelocityPass(RHICmdList, View, VelocityInstanceCullingDrawParams);
+					}
+				});
 	}
 
 	// Dithered transition stencil mask clear, accounting for all active viewports
 	if (DepthPass.bDitheredLODTransitionsUseStencil)
 	{
-		AddDitheredStencilClearPass(GraphBuilder, SceneTextures.Depth.Target, ShaderPlatform, false);
+		AddDitheredStencilClearPass(GraphBuilder, DepthTargetTexture, ShaderPlatform, false);
 	}
 
 	FViewOcclusionQueriesPerView QueriesPerView;
@@ -911,7 +921,7 @@ void FMobileSceneRenderer::RenderFullDepthPrepass(FRDGBuilder& GraphBuilder, TAr
 			FViewInfo& View = *ViewContext.ViewInfo;
 
 			auto* PassParameters = GraphBuilder.AllocParameters<FMobileRenderPassParameters>();
-			PassParameters->RenderTargets.DepthStencil = FDepthStencilBinding(SceneTextures.Depth.Target, ERenderTargetLoadAction::ELoad, ERenderTargetLoadAction::ELoad, FExclusiveDepthStencil::DepthRead_StencilRead);
+			PassParameters->RenderTargets.DepthStencil = FDepthStencilBinding(DepthTargetTexture, ERenderTargetLoadAction::ELoad, ERenderTargetLoadAction::ELoad, FExclusiveDepthStencil::DepthRead_StencilRead);
 			PassParameters->View = View.GetShaderParameters();
 			PassParameters->MobileBasePass = LastViewMobileBasePassUB;
 			PassParameters->RenderTargets.NumOcclusionQueries = ComputeNumOcclusionQueriesToBatch();
diff --git a/Engine/Source/Runtime/Renderer/Private/SceneTextures.cpp b/Engine/Source/Runtime/Renderer/Private/SceneTextures.cpp
--- a/Engine/Source/Runtime/Renderer/Private/SceneTextures.cpp
+++ b/Engine/Source/Runtime/Renderer/Private/SceneTextures.cpp
@@ -805,7 +805,7 @@ void FSceneTextures::InitializeViewFamily(FRDGBuilder& GraphBuilder, FViewFamily
 		const FRDGTextureDesc ColorDesc(FRDGTextureDesc::Create2D(Config.Extent, PF_B8G8R8A8, FClearValueBinding::Transparent, TexCreate_ShaderResource | TexCreate_RenderTargetable, 1, Config.EditorPrimitiveNumSamples));
 		SceneTextures.EditorPrimitiveColor = GraphBuilder.CreateTexture(ColorDesc, TEXT("Editor.PrimitivesColor"));
 
-		FRDGTextureDesc DepthDesc(FRDGTextureDesc::Create2D(Config.Extent, PF_DepthStencil, FClearValueBinding::DepthFar, TexCreate_ShaderResource | TexCreate_DepthStencilTargetable, 1, Config.EditorPrimitiveNumSamples));
+		FRDGTextureDesc DepthDesc(FRDGTextureDesc::Create2D(Config.Extent, PF_DepthStencil, FClearValueBinding::DepthFar, TexCreate_ShaderResource | TexCreate_DepthStencilTargetable | TexCreate_InputAttachmentRead, 1, Config.EditorPrimitiveNumSamples));
 		DepthDesc.AliasableFormats.Add(PF_X24_G8);
 		SceneTextures.EditorPrimitiveDepth = GraphBuilder.CreateTexture(DepthDesc, TEXT("Editor.PrimitivesDepth"));
 	}
