Always include RT into the uber shader permutation for VolumetricFog scattering shader to avoid missing permutation when a game run with or without RT.

#Synced-CL 50037646

#Synced-CL 50037646

diff --git a/Engine/Source/Runtime/Renderer/Private/VolumetricFog.cpp b/Engine/Source/Runtime/Renderer/Private/VolumetricFog.cpp
--- a/Engine/Source/Runtime/Renderer/Private/VolumetricFog.cpp
+++ b/Engine/Source/Runtime/Renderer/Private/VolumetricFog.cpp
@@ -1166,7 +1166,7 @@ class FVolumetricFogLightScatteringCS : public FGlobalShader
 			PermutationVector.Set<FDistanceFieldSkyOcclusion>(true);
 			PermutationVector.Set<FLumenGI>(FDataDrivenShaderPlatformInfo::GetSupportsLumenGI(ShaderPlatform));
 			PermutationVector.Set<FVirtualShadowMap>(DoesPlatformSupportVirtualShadowMaps(ShaderPlatform));
-			PermutationVector.Set<FRaytracedShadowsVolume>(ShouldCompileRayTracingShadersForProject(ShaderPlatform));
+			PermutationVector.Set<FRaytracedShadowsVolume>(true);
 			PermutationVector.Set<FSampleLightFunctionAtlas>(true);
 			PermutationVector.Set<FMegaLights>(MegaLights::ShouldCompileShaders(ShaderPlatform));
 			PermutationVector.Set<FLightSoftFading>(true);
@@ -1215,7 +1215,7 @@ class FVolumetricFogLightScatteringCS : public FGlobalShader
 					GetIfSupported<FLumenGI>(PermutationVector, FDataDrivenShaderPlatformInfo::GetSupportsLumenGI(Parameters.Platform)) &&
 					GetIfSupported<FMegaLights>(PermutationVector, MegaLights::ShouldCompileShaders(Parameters.Platform)) &&
 					GetIfSupported<FVirtualShadowMap>(PermutationVector, DoesPlatformSupportVirtualShadowMaps(Parameters.Platform)) &&
-					GetIfSupported<FRaytracedShadowsVolume>(PermutationVector, ShouldCompileRayTracingShadersForProject(Parameters.Platform)) &&
+					PermutationVector.Get<FRaytracedShadowsVolume>() &&
 					PermutationVector.Get<FTemporalReprojection>() &&
 					PermutationVector.Get<FDistanceFieldSkyOcclusion>() &&
 					PermutationVector.Get<FSampleLightFunctionAtlas>() &&
