[PostProcessing] Shuffle tonemapper permutations around to improve fast path on low-end platforms

#Synced-CL 50022451

diff --git a/Engine/Shaders/Private/PostProcessCommon.ush b/Engine/Shaders/Private/PostProcessCommon.ush
--- a/Engine/Shaders/Private/PostProcessCommon.ush
+++ b/Engine/Shaders/Private/PostProcessCommon.ush
@@ -175,6 +175,10 @@ float3 ComputeVignetteMask(
 	float2 InvSize	= VignetteParams1.zw;
 	float Softness	= VignetteParams2;
 
+#if !USE_ADVANCED_VIGNETTE
+	VignetteType = VIGNETTE_TYPE_COSINE_FOURTH_LAW; // force vignette type
+#endif
+
 	float Weight = 1.0;
 	BRANCH
 	if (Intensity <= 0)
@@ -194,6 +198,7 @@ float3 ComputeVignetteMask(
 		Weight = lerp(1.0, Weight, Intensity);
 	}
 #endif
+#if USE_ADVANCED_VIGNETTE
 	else
 	{
 		float2 Position = InvSize * VignetteCircleSpacePos;
@@ -207,6 +212,7 @@ float3 ComputeVignetteMask(
 		Weight = pow(saturate(1.0 - Distance), Softness);
 		Weight = lerp(1.0, Weight, Intensity);
 	}
+#endif
 
 	return lerp(Color, (float3)1, Weight);
 }
diff --git a/Engine/Shaders/Private/PostProcessTonemap.usf b/Engine/Shaders/Private/PostProcessTonemap.usf
--- a/Engine/Shaders/Private/PostProcessTonemap.usf
+++ b/Engine/Shaders/Private/PostProcessTonemap.usf
@@ -73,6 +73,7 @@ float4 ColorScale0;
 float4 ColorScale1;
 
 float SharpenDiv6;
+uint Flags;
 
 float4 VignetteParams0;
 float4 VignetteParams1;
@@ -603,9 +604,10 @@ float4 TonemapCommonPS(
 	}
 	#endif
 
-	#if DIM_ALPHA_CHANNEL
+	if ((Flags & FLAG_ALPHA_CHANNEL) != 0)
+	{
 		OutColor.a = SceneColor.a;
-	#endif
+	}
 
 	return OutColor;
 }
diff --git a/Engine/Source/Runtime/Renderer/Private/PostProcess/PostProcessTonemap.cpp b/Engine/Source/Runtime/Renderer/Private/PostProcess/PostProcessTonemap.cpp
--- a/Engine/Source/Runtime/Renderer/Private/PostProcess/PostProcessTonemap.cpp
+++ b/Engine/Source/Runtime/Renderer/Private/PostProcess/PostProcessTonemap.cpp
@@ -83,21 +83,22 @@ struct FOutputLuminance
 namespace TonemapperPermutation
 {
 // Shared permutation dimensions between deferred and mobile renderer.
-class FTonemapperGammaOnlyDim      	: SHADER_PERMUTATION_BOOL("USE_GAMMA_ONLY");
-class FTonemapperLocalExposureDim  	: SHADER_PERMUTATION_INT("LOCAL_EXPOSURE_MODE", 3);
-class FTonemapperSharpenDim        	: SHADER_PERMUTATION_BOOL("USE_SHARPEN");
-class FTonemapperFilmGrainDim      	: SHADER_PERMUTATION_BOOL("USE_FILM_GRAIN");
-class FTonemapperVignetteTextureDim	: SHADER_PERMUTATION_BOOL("USE_VIGNETTE_TEXTURE");
-class FTonemapperMsaaDim           	: SHADER_PERMUTATION_BOOL("METAL_MSAA_HDR_DECODE");
-class FTonemapperAlphaChannelDim   	: SHADER_PERMUTATION_BOOL("DIM_ALPHA_CHANNEL");
+class FTonemapperGammaOnlyDim      		: SHADER_PERMUTATION_BOOL("USE_GAMMA_ONLY");
+class FTonemapperLocalExposureDim  		: SHADER_PERMUTATION_INT("LOCAL_EXPOSURE_MODE", 3);
+class FTonemapperSharpenDim        		: SHADER_PERMUTATION_BOOL("USE_SHARPEN");
+class FTonemapperFilmGrainDim      		: SHADER_PERMUTATION_BOOL("USE_FILM_GRAIN");
+class FTonemapperAdvancedVignetteDim	: SHADER_PERMUTATION_BOOL("USE_ADVANCED_VIGNETTE");
+class FTonemapperVignetteTextureDim		: SHADER_PERMUTATION_BOOL("USE_VIGNETTE_TEXTURE");
+class FTonemapperMsaaDim           		: SHADER_PERMUTATION_BOOL("METAL_MSAA_HDR_DECODE");
+
 using FCommonDomain = TShaderPermutationDomain<
 	FTonemapperGammaOnlyDim,
 	FTonemapperLocalExposureDim,
 	FTonemapperSharpenDim,
 	FTonemapperFilmGrainDim,
+	FTonemapperAdvancedVignetteDim,
 	FTonemapperVignetteTextureDim,
-	FTonemapperMsaaDim,
-	FTonemapperAlphaChannelDim>;
+	FTonemapperMsaaDim>;
 
 bool ShouldCompileCommonPermutation(const FGlobalShaderPermutationParameters& Parameters, const FCommonDomain& PermutationVector)
 {
@@ -113,6 +114,7 @@ bool ShouldCompileCommonPermutation(const FGlobalShaderPermutationParameters& Pa
 		return !PermutationVector.Get<FTonemapperLocalExposureDim>() &&
 			!PermutationVector.Get<FTonemapperSharpenDim>() &&
 			!PermutationVector.Get<FTonemapperFilmGrainDim>() &&
+			!PermutationVector.Get<FTonemapperAdvancedVignetteDim>() &&
 			!PermutationVector.Get<FTonemapperVignetteTextureDim>() &&
 			!PermutationVector.Get<FTonemapperMsaaDim>();
 	}
@@ -133,7 +135,6 @@ FCommonDomain BuildCommonPermutationDomain(const FViewInfo& View, bool bGammaOnl
 	const FSceneViewFamily* Family = View.Family;
 
 	FCommonDomain PermutationVector;
-	PermutationVector.Set<FTonemapperAlphaChannelDim>(IsPostProcessingWithAlphaChannelSupported());
 
 	// Gamma
 	if (bGammaOnly ||
@@ -148,6 +149,7 @@ FCommonDomain BuildCommonPermutationDomain(const FViewInfo& View, bool bGammaOnl
 	PermutationVector.Set<FTonemapperFilmGrainDim>(View.FinalPostProcessSettings.FilmGrainIntensity > 0.0f && View.FilmGrainTexture != nullptr);
 	PermutationVector.Set<FTonemapperSharpenDim>(GetSharpenSetting(Settings) > 0.0f);
 	PermutationVector.Set<FTonemapperMsaaDim>(bMetalMSAAHDRDecode);
+	PermutationVector.Set<FTonemapperAdvancedVignetteDim>(View.FinalPostProcessSettings.VignetteType != EVignetteType::CosineFourthLaw);
 	PermutationVector.Set<FTonemapperVignetteTextureDim>(View.FinalPostProcessSettings.VignetteType == EVignetteType::Texture);
 
 	if (bLocalExposure)
@@ -188,6 +190,7 @@ FDesktopDomain RemapPermutation(FDesktopDomain PermutationVector, ERHIFeatureLev
 	{
 		CommonPermutationVector.Set<FTonemapperFilmGrainDim>(true);
 		CommonPermutationVector.Set<FTonemapperVignetteTextureDim>(true);
+		CommonPermutationVector.Set<FTonemapperAdvancedVignetteDim>(true);
 		CommonPermutationVector.Set<FTonemapperSharpenDim>(true);
 
 		PermutationVector.Set<FTonemapperColorFringeDim>(true);
@@ -350,13 +353,14 @@ BEGIN_SHADER_PARAMETER_STRUCT(FTonemapParameters, )
 	SHADER_PARAMETER(FVector4f, ChromaticAberrationParams)
 	SHADER_PARAMETER(FVector4f, LensPrincipalPointOffsetScale)
 	SHADER_PARAMETER(FVector4f, LensPrincipalPointOffsetScaleInverse)
-	SHADER_PARAMETER(float, SharpenDiv6)
-	SHADER_PARAMETER(float, LUTSize)
-	SHADER_PARAMETER(float, InvLUTSize)
-	SHADER_PARAMETER(float, LUTScale)
-	SHADER_PARAMETER(float, LUTOffset)
-	SHADER_PARAMETER(float, EditorNITLevel)
-	SHADER_PARAMETER(float, BackbufferQuantizationDithering)
+	SHADER_PARAMETER(uint32,	Flags)
+	SHADER_PARAMETER(float, 	SharpenDiv6)
+	SHADER_PARAMETER(float, 	LUTSize)
+	SHADER_PARAMETER(float, 	InvLUTSize)
+	SHADER_PARAMETER(float, 	LUTScale)
+	SHADER_PARAMETER(float, 	LUTOffset)
+	SHADER_PARAMETER(float, 	EditorNITLevel)
+	SHADER_PARAMETER(float, 	BackbufferQuantizationDithering)
 
 	SHADER_PARAMETER(FVector4f, VignetteParams0)
 	SHADER_PARAMETER(FVector4f, VignetteParams1)
@@ -402,6 +406,18 @@ public:
 	}
 };
 
+enum class ETonemapperFlags : uint32
+{
+	AlphaChannel = 1 << 0
+};
+
+static void ModifyTonemapCompilationEnvironment(const FGlobalShaderPermutationParameters& Parameters, FShaderCompilerEnvironment& OutEnvironment)
+{
+	OutEnvironment.SetDefine(TEXT("SUPPORTS_SCENE_COLOR_APPLY_PARAMETERS"), FTonemapInputs::SupportsSceneColorApplyParametersBuffer(Parameters.Platform) ? 1 : 0);
+	OutEnvironment.SetDefine(TEXT("SUPPORTS_FILM_GRAIN"), SupportsFilmGrain(Parameters.Platform) ? 1 : 0);
+	OutEnvironment.SetDefine(TEXT("FLAG_ALPHA_CHANNEL"), static_cast<uint32>(ETonemapperFlags::AlphaChannel));
+}
+
 class FTonemapVS : public FGlobalShader
 {
 public:
@@ -416,6 +432,11 @@ public:
 	{
 		return true;
 	}
+
+	static void ModifyCompilationEnvironment(const FGlobalShaderPermutationParameters& Parameters, FShaderCompilerEnvironment& OutEnvironment)
+	{
+		ModifyTonemapCompilationEnvironment(Parameters, OutEnvironment);
+	}
 };
 
 class FTonemapPS : public FGlobalShader
@@ -442,8 +463,7 @@ public:
 
 	static void ModifyCompilationEnvironment(const FGlobalShaderPermutationParameters& Parameters, FShaderCompilerEnvironment& OutEnvironment)
 	{
-		OutEnvironment.SetDefine(TEXT("SUPPORTS_SCENE_COLOR_APPLY_PARAMETERS"), FTonemapInputs::SupportsSceneColorApplyParametersBuffer(Parameters.Platform) ? 1 : 0);
-		OutEnvironment.SetDefine(TEXT("SUPPORTS_FILM_GRAIN"), SupportsFilmGrain(Parameters.Platform) ? 1 : 0);
+		ModifyTonemapCompilationEnvironment(Parameters, OutEnvironment);
 	}
 };
 
@@ -478,8 +498,7 @@ public:
 		OutEnvironment.SetDefine(TEXT("THREADGROUP_SIZEX"), GTonemapComputeTileSizeX);
 		OutEnvironment.SetDefine(TEXT("THREADGROUP_SIZEY"), GTonemapComputeTileSizeY);
 
-		OutEnvironment.SetDefine(TEXT("SUPPORTS_SCENE_COLOR_APPLY_PARAMETERS"), FTonemapInputs::SupportsSceneColorApplyParametersBuffer(Parameters.Platform) ? 1 : 0);
-		OutEnvironment.SetDefine(TEXT("SUPPORTS_FILM_GRAIN"), SupportsFilmGrain(Parameters.Platform) ? 1 : 0);
+		ModifyTonemapCompilationEnvironment(Parameters, OutEnvironment);
 	}
 };
 
@@ -678,6 +697,9 @@ FScreenPassTexture AddTonemapPass(FRDGBuilder& GraphBuilder, const FViewInfo& Vi
 	FRHISamplerState* BilinearClampSampler = TStaticSamplerState<SF_Bilinear, AM_Clamp, AM_Clamp, AM_Clamp>::GetRHI();
 	FRHISamplerState* PointClampSampler = TStaticSamplerState<SF_Point, AM_Clamp, AM_Clamp, AM_Clamp>::GetRHI();
 
+	uint32 Flags = 0;
+	Flags |= IsPostProcessingWithAlphaChannelSupported() ? static_cast<uint32>(ETonemapperFlags::AlphaChannel) : 0;
+
 	const float SharpenDiv6 = TonemapperPermutation::GetSharpenSetting(View.FinalPostProcessSettings) / 6.0f;
 
 	FVector4f ChromaticAberrationParams;
@@ -851,6 +873,7 @@ FScreenPassTexture AddTonemapPass(FRDGBuilder& GraphBuilder, const FViewInfo& Vi
 	CommonParameters.ColorGradingLUTSampler = BilinearClampSampler;
 	CommonParameters.ColorScale0 = PostProcessSettings.SceneColorTint;
 	CommonParameters.ChromaticAberrationParams = ChromaticAberrationParams;
+	CommonParameters.Flags = Flags;
 	CommonParameters.SharpenDiv6 = SharpenDiv6;
 	CommonParameters.EditorNITLevel = EditorNITLevel;
 
@@ -1130,6 +1153,7 @@ class FMobileCustomResolvePS : public FGlobalShader
 
 	static void ModifyCompilationEnvironment(const FGlobalShaderPermutationParameters& Parameters, FShaderCompilerEnvironment& OutEnvironment)
 	{
+		ModifyTonemapCompilationEnvironment(Parameters, OutEnvironment);
 		OutEnvironment.SetDefine(TEXT("OUTPUT_DEVICE_SRGB"), 1);
 	}
 };
