Fixed flickering instance issue caused by instances being culled for a direct draw when auto instancing merged adjacent static mesh and instanced static mesh draws, which became possible with the change to a shared VF.
To guarantee correctness we break up the auto instancing when a change from direct to indirect draw is encountered.

#Synced-CL 50044433

diff --git a/Engine/Source/Runtime/Renderer/Private/InstanceCulling/InstanceCullingContext.cpp b/Engine/Source/Runtime/Renderer/Private/InstanceCulling/InstanceCullingContext.cpp
--- a/Engine/Source/Runtime/Renderer/Private/InstanceCulling/InstanceCullingContext.cpp
+++ b/Engine/Source/Runtime/Renderer/Private/InstanceCulling/InstanceCullingContext.cpp
@@ -1503,6 +1503,11 @@ void FInstanceCullingContext::SetupDrawCommands(
 	const uint32 MaxGenericBatchSize = bUsesUniformBufferView ? PLATFORM_MAX_UNIFORM_BUFFER_RANGE / UniformViewInstanceStride[0] : MAX_uint32;
 	const uint32 MaxPrimitiveBatchSize = bUsesUniformBufferView ? PLATFORM_MAX_UNIFORM_BUFFER_RANGE / UniformViewInstanceStride[1] : MAX_uint32;
 	
+	// Used to detect when we change from direct to indirect and break auto instancing there. This is a workaround for the fact that we need an indirect draw to support GPU culling and when MDCs are identical they could batch together otherwise.
+	// In the future we should rework the below to instead decide the draw command at the end, but this is complicated by the mobile GPU scene batching.
+	// Note: initial value is arbitrary because we will never use the result during the first iteration as CurrentStateBucketId is initialized to -1
+	bool bPrevUseIndirectDraw = false;
+
 	// Allocate conservatively for all commands, may not use all.
 	for (int32 DrawCommandIndex = 0; DrawCommandIndex < NumDrawCommandsIn; ++DrawCommandIndex)
 	{
@@ -1522,7 +1527,10 @@ void FInstanceCullingContext::SetupDrawCommands(
 		const bool bCompactIdenticalCommands = bInCompactIdenticalCommands && !(bUseIndirectDraw && bPreserveInstanceOrder) 
 			&& (bUsesUniformBufferView ? (CurrentAutoInstanceCount < MaxPrimitiveBatchSize && !bUseIndirectDraw) : true);
 
-		if (bCompactIdenticalCommands && CurrentStateBucketId != -1 && VisibleMeshDrawCommand.StateBucketId == CurrentStateBucketId && VisibleMeshDrawCommand.CullingPayloadFlags == CurrentCullingPayloadFlags)
+		const bool bChangedFromDirectToIndirect = !bPrevUseIndirectDraw && bUseIndirectDraw;
+		bPrevUseIndirectDraw = bUseIndirectDraw;
+
+		if (bCompactIdenticalCommands && CurrentStateBucketId != -1 && VisibleMeshDrawCommand.StateBucketId == CurrentStateBucketId && VisibleMeshDrawCommand.CullingPayloadFlags == CurrentCullingPayloadFlags && !bChangedFromDirectToIndirect)
 		{
 			// Drop since previous covers for this
 
