Changed the MDC instance count to 1 when BatchElement.bFetchInstanceCountFromScene is set to allow all MDCs from ISM to hash to the same ID.

#Synced-CL 50044433

diff --git a/Engine/Source/Runtime/Renderer/Private/MeshPassProcessor.cpp b/Engine/Source/Runtime/Renderer/Private/MeshPassProcessor.cpp
--- a/Engine/Source/Runtime/Renderer/Private/MeshPassProcessor.cpp
+++ b/Engine/Source/Runtime/Renderer/Private/MeshPassProcessor.cpp
@@ -912,7 +912,8 @@ void FMeshDrawCommand::SetDrawParametersAndFinalize(
 	IndexBuffer = BatchElement.IndexBuffer ? BatchElement.IndexBuffer->IndexBufferRHI.GetReference() : nullptr;
 	FirstIndex = BatchElement.FirstIndex;
 	NumPrimitives = BatchElement.NumPrimitives;
-	NumInstances = BatchElement.NumInstances;
+	// Note: we override the instance count to ensure we hash to the same state bucket for MDCs where we know it will be changed at draw time anyway.
+	NumInstances = BatchElement.bFetchInstanceCountFromScene ? 1 : BatchElement.NumInstances;
 
 	// If the mesh batch has a valid dynamic index buffer, use it instead
 	if (BatchElement.DynamicIndexBuffer.IsValid() &&  BatchElement.DynamicIndexBuffer.IndexBuffer->IsInitialized())
