Preselection Highlight only shown when there's a selected object behind the hovered object

#Synced-CL 50039254

diff --git a/Engine/Source/Runtime/Renderer/Private/SceneHitProxyRendering.cpp b/Engine/Source/Runtime/Renderer/Private/SceneHitProxyRendering.cpp
--- a/Engine/Source/Runtime/Renderer/Private/SceneHitProxyRendering.cpp
+++ b/Engine/Source/Runtime/Renderer/Private/SceneHitProxyRendering.cpp
@@ -1084,7 +1084,7 @@ int32 FEditorSelectionMeshProcessor::GetStencilValue(const FSceneView* View, con
 	{
 		StencilValue = *ExistingStencilValue;
 	}
-	else if (PrimitiveSceneProxy->IsIndividuallySelected())
+	else if (PrimitiveSceneProxy->IsIndividuallySelected() || PrimitiveSceneProxy->IsHovered())
 	{
 		const int Color = 0;
 		const int UniqueId = ProxyToStencilIndex.Num();
