Skip invalidations with out of range indices
This probably indicates an issue elsewhere in the engine but this avoids a crash at least

#rnx

#Synced-CL 50005473

diff --git a/Engine/Source/Runtime/Renderer/Private/VirtualShadowMaps/VirtualShadowMapCacheManager.cpp b/Engine/Source/Runtime/Renderer/Private/VirtualShadowMaps/VirtualShadowMapCacheManager.cpp
--- a/Engine/Source/Runtime/Renderer/Private/VirtualShadowMaps/VirtualShadowMapCacheManager.cpp
+++ b/Engine/Source/Runtime/Renderer/Private/VirtualShadowMaps/VirtualShadowMapCacheManager.cpp
@@ -442,7 +442,8 @@ void FVirtualShadowMapArrayCacheManager::FInvalidatingPrimitiveCollector::AddPri
 			// If something doesn't have a valid persistent primitive index, it is treated as a dynamic primitive
 			// See GetCachePrimitiveAsDynamic
 			bool bCachePrimativeAsDynamic = true;
-			if (Range.PersistentPrimitiveIndex.IsValid())
+			if (Range.PersistentPrimitiveIndex.IsValid() &&
+				Range.PersistentPrimitiveIndex.Index < ManagerCachePrimitiveAsDynamic.Num())
 			{
 				bCachePrimativeAsDynamic = ManagerCachePrimitiveAsDynamic[Range.PersistentPrimitiveIndex.Index];
 				if (Range.bMarkAsDynamic)
@@ -461,10 +462,9 @@ void FVirtualShadowMapArrayCacheManager::FInvalidatingPrimitiveCollector::AddPri
 		// Per-light invalidations
 		for (const FVirtualShadowMapInstanceRange& Range : CacheEntry.PrimitiveInstancesToInvalidate)
 		{
-			check(Range.PersistentPrimitiveIndex.IsValid());		// Should always be valid currently in this path
-
 			bool bCachePrimativeAsDynamic = true;
-			if (Range.PersistentPrimitiveIndex.IsValid())
+			if (Range.PersistentPrimitiveIndex.IsValid() &&
+				Range.PersistentPrimitiveIndex.Index < ManagerCachePrimitiveAsDynamic.Num())
 			{
 				bCachePrimativeAsDynamic = ManagerCachePrimitiveAsDynamic[Range.PersistentPrimitiveIndex.Index];
 				if (Range.bMarkAsDynamic)
