Added GetTypeHash for FPersistentPrimitiveIndex so it can easily be used various hash maps and similar.

#rnx

#Synced-CL 50202780

diff --git a/Engine/Source/Runtime/Renderer/Public/MeshPassProcessor.h b/Engine/Source/Runtime/Renderer/Public/MeshPassProcessor.h
--- a/Engine/Source/Runtime/Renderer/Public/MeshPassProcessor.h
+++ b/Engine/Source/Runtime/Renderer/Public/MeshPassProcessor.h
@@ -1553,6 +1553,11 @@ struct FPersistentPrimitiveIndex
 	{
 		return Index == B.Index;
 	}
+
+	friend uint32 GetTypeHash(FPersistentPrimitiveIndex PersistentPrimitiveIndex)
+	{
+		return GetTypeHash(PersistentPrimitiveIndex.Index);
+	}
 };
 
 /**
