GlobalDistanceField - fixed pooled buffer and texture memory leaks
* Remove TRefCountPtr from FGlobalDistanceFieldParameterData as it's a POD memset to 0
* Remove TRefCountPtr from FGlobalDistanceFieldInfo as it's a transient structure assumed to be POD (see FGlobalDistanceFieldInfo& GlobalDistanceFieldInfo = *Allocator.Create<FGlobalDistanceFieldInfo>(); in class FViewInfo)
* Remove TRefCountPtr from FNiagaraGpuComputeDispatch::FCachedDistanceFieldData as we initialize it from a transient FGlobalDistanceFieldParameterData (valid only for rendering one frame). Niagara can hold to FNiagaraGpuComputeDispatch objects without rendering them, but CachedGDFData can be only released through PreInitViews()/PostRenderOpaque() which won't be called on a such object.

Repro:
* Create two levels, one with 20k and another with 40k Lumen Scene View Distance (PPV setting)
* Disable Lumen HWRT
* Load one level in the editor and then the second one
* Observe how each load adds new resources (r.DumpBufferPoolMemory) and old ones aren't ever removed

#tests Tested sequencer timeline scrubbing with GDF collision particles and some occluders
#Synced-CL 50205666

diff --git a/Engine/Plugins/FX/Niagara/Source/Niagara/Private/NiagaraGpuComputeDispatch.h b/Engine/Plugins/FX/Niagara/Source/Niagara/Private/NiagaraGpuComputeDispatch.h
--- a/Engine/Plugins/FX/Niagara/Source/Niagara/Private/NiagaraGpuComputeDispatch.h
+++ b/Engine/Plugins/FX/Niagara/Source/Niagara/Private/NiagaraGpuComputeDispatch.h
@@ -300,11 +300,11 @@ private:
 	{
 		bool								bCacheValid = false;
 		bool								bValidForPass = false;
-		FTextureRHIRef						PageAtlasTexture;
-		FTextureRHIRef						CoverageAtlasTexture;
-		TRefCountPtr<FRDGPooledBuffer>		PageObjectGridBuffer;
-		FTextureRHIRef						PageTableTexture;
-		FTextureRHIRef						MipTexture;
+		FRHITexture*						PageAtlasTexture = nullptr;
+		FRHITexture*						CoverageAtlasTexture = nullptr;
+		FRDGPooledBuffer*					PageObjectGridBuffer = nullptr;
+		FRHITexture*						PageTableTexture = nullptr;
+		FRHITexture*						MipTexture = nullptr;
 		FGlobalDistanceFieldParameterData	GDFParameterData;
 	};
 	FCachedDistanceFieldData			CachedGDFData;
diff --git a/Engine/Source/Runtime/Renderer/Private/SceneRendering.h b/Engine/Source/Runtime/Renderer/Private/SceneRendering.h
--- a/Engine/Source/Runtime/Renderer/Private/SceneRendering.h
+++ b/Engine/Source/Runtime/Renderer/Private/SceneRendering.h
@@ -640,14 +640,14 @@ public:
 
 	FGlobalDistanceFieldParameterData ParameterData;
 
-	TRefCountPtr<FRDGPooledBuffer> PageFreeListAllocatorBuffer;
-	TRefCountPtr<FRDGPooledBuffer> PageFreeListBuffer;
-	TRefCountPtr<IPooledRenderTarget> PageAtlasTexture;
-	TRefCountPtr<IPooledRenderTarget> CoverageAtlasTexture;
-	TRefCountPtr<FRDGPooledBuffer> PageObjectGridBuffer;
-	TRefCountPtr<IPooledRenderTarget> PageTableCombinedTexture;
-	TRefCountPtr<IPooledRenderTarget> PageTableLayerTextures[GDF_Num];
-	TRefCountPtr<IPooledRenderTarget> MipTexture;
+	FRDGPooledBuffer* PageFreeListAllocatorBuffer = nullptr;
+	FRDGPooledBuffer* PageFreeListBuffer = nullptr;
+	IPooledRenderTarget* PageAtlasTexture = nullptr;
+	IPooledRenderTarget* CoverageAtlasTexture = nullptr;
+	FRDGPooledBuffer* PageObjectGridBuffer = nullptr;
+	IPooledRenderTarget* PageTableCombinedTexture = nullptr;
+	IPooledRenderTarget* PageTableLayerTextures[GDF_Num] = {};
+	IPooledRenderTarget* MipTexture = nullptr;
 
 	void UpdateParameterData(float MaxOcclusionDistance, bool bLumenEnabled, float LumenSceneViewDistance, FVector PreViewTranslation);
 };
diff --git a/Engine/Source/Runtime/Renderer/Public/GlobalDistanceFieldParameters.h b/Engine/Source/Runtime/Renderer/Public/GlobalDistanceFieldParameters.h
--- a/Engine/Source/Runtime/Renderer/Public/GlobalDistanceFieldParameters.h
+++ b/Engine/Source/Runtime/Renderer/Public/GlobalDistanceFieldParameters.h
@@ -17,7 +17,6 @@ class FShaderParameterMap;
 class FGlobalDistanceFieldParameterData
 {
 public:
-
 	FGlobalDistanceFieldParameterData()
 	{
 		FPlatformMemory::Memzero(this, sizeof(FGlobalDistanceFieldParameterData));
@@ -31,7 +30,7 @@ public:
 	float MipTransition;
 	FRHITexture* PageAtlasTexture;
 	FRHITexture* CoverageAtlasTexture;
-	TRefCountPtr<FRDGPooledBuffer> PageObjectGridBuffer;
+	FRDGPooledBuffer* PageObjectGridBuffer;
 	FRHITexture* PageTableTexture;
 	FRHITexture* MipTexture;
 	int32 ClipmapSizeInPages;
