Refactored AnimSequenceTransformProvider to use compact bone indices and build per-LOD tables. This reduces the workload for skeletal meshes with disabled bones from the target skeleton and when a higher LOD is used.

#Synced-CL 50172231

diff --git a/Engine/Shaders/Private/Skinning/AnimSequenceTransformProvider.usf b/Engine/Shaders/Private/Skinning/AnimSequenceTransformProvider.usf
--- a/Engine/Shaders/Private/Skinning/AnimSequenceTransformProvider.usf
+++ b/Engine/Shaders/Private/Skinning/AnimSequenceTransformProvider.usf
@@ -48,7 +48,7 @@ void StoreBoneTransformWithUniformScale(RWByteAddressBuffer Buffer, uint BaseOff
 }
 
 FBoneTransformWithUniformScale EvaluateAndRetargetBone(
-	uint MeshBoneIndex,
+	uint BoneIndex,
 	uint SequenceTransformBufferOffset,
 	uint RetargetingTransformBufferOffset,
 	uint RetargetingDataBufferOffset,
@@ -57,10 +57,10 @@ FBoneTransformWithUniformScale EvaluateAndRetargetBone(
 	float Alpha,
 	uint NumAnimBones)
 {
-	FBoneRetargetingData RetargetingData = LoadBoneRetargetingData(RetargetingDataBuffer, RetargetingDataBufferOffset, MeshBoneIndex);
+	FBoneRetargetingData RetargetingData = LoadBoneRetargetingData(RetargetingDataBuffer, RetargetingDataBufferOffset, BoneIndex);
 
-	const uint RetargetingTransformIndex = MeshBoneIndex * (RetargetingData.bRequiresRefPoseRetarget ? 2 : 1);
-	const bool bIsRootBone               = MeshBoneIndex == 0;
+	const uint RetargetingTransformIndex = BoneIndex * (RetargetingData.bRequiresRefPoseRetarget ? 2 : 1);
+	const bool bIsRootBone               = BoneIndex == 0;
 
 	FBoneTransformWithUniformScale RetargetingTransform = LoadBoneTransformWithUniformScale(SequenceTransformBuffer, RetargetingTransformBufferOffset, RetargetingTransformIndex);
 
@@ -129,20 +129,20 @@ FBoneTransformWithUniformScale EvaluateAndRetargetBone(
 [numthreads(BONES_PER_GROUP, 1, 1)]
 void EvaluateAnimSequenceCS(uint3 GroupId : SV_GroupID, uint GroupIndex : SV_GroupIndex)
 {
-	const uint BoneBlock          = GroupId.x;
-	const uint LocalMeshBoneIndex = GroupIndex;
+	const uint BoneBlock      = GroupId.x;
+	const uint LocalBoneIndex = GroupIndex;
 
 	const FEvaluateAnimSequenceBlockHeader BlockHeader = EvaluateBlockHeaderBuffer.Load<FEvaluateAnimSequenceBlockHeader>(BoneBlock * uint(sizeof(FEvaluateAnimSequenceBlockHeader)));
 
-	if (LocalMeshBoneIndex >= BlockHeader.BlockTransformCount)
+	if (LocalBoneIndex >= BlockHeader.BlockTransformCount)
 	{
 		return;
 	}
 
-	const uint MeshBoneIndex = LocalMeshBoneIndex + (BlockHeader.BlockLocalIndex * BONES_PER_GROUP);
+	const uint BoneIndex = LocalBoneIndex + (BlockHeader.BlockLocalIndex * BONES_PER_GROUP);
 
 	FBoneTransformWithUniformScale MeshLocalAnimTransform = EvaluateAndRetargetBone(
-		MeshBoneIndex,
+		BoneIndex,
 		BlockHeader.SequenceTransformBufferOffset,
 		BlockHeader.RetargetingTransformBufferOffset,
 		BlockHeader.RetargetingDataBufferOffset,
@@ -151,26 +151,26 @@ void EvaluateAnimSequenceCS(uint3 GroupId : SV_GroupID, uint GroupIndex : SV_Gro
 		BlockHeader.Alpha,
 		BlockHeader.NumAnimBones);
 
-	StoreBoneTransformWithUniformScale(LocalSpaceTransformBufferUAV, BlockHeader.DstTransformBufferOffset, MeshBoneIndex, MeshLocalAnimTransform);
+	StoreBoneTransformWithUniformScale(LocalSpaceTransformBufferUAV, BlockHeader.DstTransformBufferOffset, BoneIndex, MeshLocalAnimTransform);
 }
 
 [numthreads(BONES_PER_GROUP, 1, 1)]
 void EvaluateAnimSequenceBlendCS(uint3 GroupId : SV_GroupID, uint GroupIndex : SV_GroupIndex)
 {
 	const uint BoneBlock          = GroupId.x;
-	const uint LocalMeshBoneIndex = GroupIndex;
+	const uint LocalBoneIndex = GroupIndex;
 
 	const FEvaluateAnimSequenceBlendBlockHeader BlockHeader = EvaluateBlockHeaderBuffer.Load<FEvaluateAnimSequenceBlendBlockHeader>(BoneBlock * uint(sizeof(FEvaluateAnimSequenceBlendBlockHeader)));
 
-	if (LocalMeshBoneIndex >= BlockHeader.BlockTransformCount)
+	if (LocalBoneIndex >= BlockHeader.BlockTransformCount)
 	{
 		return;
 	}
 
-	const uint MeshBoneIndex = LocalMeshBoneIndex + (BlockHeader.BlockLocalIndex * BONES_PER_GROUP);
+	const uint BoneIndex = LocalBoneIndex + (BlockHeader.BlockLocalIndex * BONES_PER_GROUP);
 
 	FBoneTransformWithUniformScale SourceMeshLocalAnimTransform = EvaluateAndRetargetBone(
-		MeshBoneIndex,
+		BoneIndex,
 		BlockHeader.SourceSequenceTransformBufferOffset,
 		BlockHeader.SourceRetargetingTransformBufferOffset,
 		BlockHeader.SourceRetargetingDataBufferOffset,
@@ -180,7 +180,7 @@ void EvaluateAnimSequenceBlendCS(uint3 GroupId : SV_GroupID, uint GroupIndex : S
 		BlockHeader.SourceNumAnimBones);
 
 	FBoneTransformWithUniformScale TargetMeshLocalAnimTransform = EvaluateAndRetargetBone(
-		MeshBoneIndex,
+		BoneIndex,
 		BlockHeader.TargetSequenceTransformBufferOffset,
 		BlockHeader.TargetRetargetingTransformBufferOffset,
 		BlockHeader.TargetRetargetingDataBufferOffset,
@@ -191,7 +191,7 @@ void EvaluateAnimSequenceBlendCS(uint3 GroupId : SV_GroupID, uint GroupIndex : S
 
 	FBoneTransformWithUniformScale MeshLocalAnimTransform = BlendBoneTransformsWithUniformScale(SourceMeshLocalAnimTransform, TargetMeshLocalAnimTransform, BlockHeader.BlendWeight);
 
-	StoreBoneTransformWithUniformScale(LocalSpaceTransformBufferUAV, BlockHeader.DstTransformBufferOffset, MeshBoneIndex, MeshLocalAnimTransform);
+	StoreBoneTransformWithUniformScale(LocalSpaceTransformBufferUAV, BlockHeader.DstTransformBufferOffset, BoneIndex, MeshLocalAnimTransform);
 }
 
 ByteAddressBuffer LocalSpaceTransformBuffer;
@@ -299,22 +299,20 @@ void ScatterTransformsCS(uint3 GroupId : SV_GroupID, uint GroupIndex : SV_GroupI
 
 	const uint TransformIndex = LocalTransformIndex + (BlockHeader.BlockLocalIndex * BONES_PER_GROUP);
 
-	uint MeshBoneIndex = TransformIndex;
+	const uint CompactBoneIndex = BoneMapBuffer.Load<uint>(BlockHeader.BoneMapOffset + TransformIndex * uint(sizeof(uint)));
 
-	if (BlockHeader.BoneMapOffset != ~0u)
+	if (CompactBoneIndex != ~0u)
 	{
-		MeshBoneIndex = BoneMapBuffer.Load<uint>(BlockHeader.BoneMapOffset + TransformIndex * uint(sizeof(uint)));
-	}
-
-	FBoneTransformWithUniformScale InvGlobalRefPoseTransform = LoadBoneTransformWithUniformScale(SequenceTransformBuffer, BlockHeader.InvGlobalRefPoseTransformBufferOffset, MeshBoneIndex);
-	FBoneTransformWithUniformScale GlobalSpaceTransform = LoadBoneTransformWithUniformScale(GlobalSpaceTransformBuffer, BlockHeader.SrcTransformBufferOffset, MeshBoneIndex);
-	GlobalSpaceTransform = MultiplyBoneTransformWithUniformScale(InvGlobalRefPoseTransform, GlobalSpaceTransform);
-	float4x3 GlobalSpaceMatrix = (float4x3)MatrixFromBoneTransformWithUniformScale(GlobalSpaceTransform);
+		FBoneTransformWithUniformScale InvGlobalRefPoseTransform = LoadBoneTransformWithUniformScale(SequenceTransformBuffer, BlockHeader.InvGlobalRefPoseTransformBufferOffset, CompactBoneIndex);
+		FBoneTransformWithUniformScale GlobalSpaceTransform = LoadBoneTransformWithUniformScale(GlobalSpaceTransformBuffer, BlockHeader.SrcTransformBufferOffset, CompactBoneIndex);
+		GlobalSpaceTransform = MultiplyBoneTransformWithUniformScale(InvGlobalRefPoseTransform, GlobalSpaceTransform);
+		float4x3 OutputMatrix = (float4x3)MatrixFromBoneTransformWithUniformScale(GlobalSpaceTransform);
 
-	StoreCompressedBoneTransform(SkinningTransformBufferUAV, BlockHeader.DstTransformBufferOffset1, TransformIndex, GlobalSpaceMatrix);
+		StoreCompressedBoneTransform(SkinningTransformBufferUAV, BlockHeader.DstTransformBufferOffset1, TransformIndex, OutputMatrix);
 
-	if (BlockHeader.DstTransformBufferOffset1 != BlockHeader.DstTransformBufferOffset2)
-	{
-		StoreCompressedBoneTransform(SkinningTransformBufferUAV, BlockHeader.DstTransformBufferOffset2, TransformIndex, GlobalSpaceMatrix);
+		if (BlockHeader.DstTransformBufferOffset1 != BlockHeader.DstTransformBufferOffset2)
+		{
+			StoreCompressedBoneTransform(SkinningTransformBufferUAV, BlockHeader.DstTransformBufferOffset2, TransformIndex, OutputMatrix);
+		}
 	}
 }
\ No newline at end of file
diff --git a/Engine/Source/Runtime/Engine/Classes/Animation/AnimSequenceTransformProviderData.h b/Engine/Source/Runtime/Engine/Classes/Animation/AnimSequenceTransformProviderData.h
--- a/Engine/Source/Runtime/Engine/Classes/Animation/AnimSequenceTransformProviderData.h
+++ b/Engine/Source/Runtime/Engine/Classes/Animation/AnimSequenceTransformProviderData.h
@@ -416,6 +416,9 @@ protected:
 	bool IsValidFor(FInstancedSkinningSceneExtensionProxy* ExtensionProxy) const;
 
 private:
+	UPROPERTY(EditAnywhere, BlueprintReadWrite, Category = TransformProvider, meta = (AllowPrivateAccess = "true", ShowOnlyInnerProperties))
+	TObjectPtr<USkinnedAsset> SkinnedAsset;
+
 	UPROPERTY(EditAnywhere, BlueprintReadOnly, Category = TransformProvider, meta = (AllowPrivateAccess = "true", ShowOnlyInnerProperties))
 	TArray<struct FAnimSequenceTransformProviderSequence> Sequences;
 
diff --git a/Engine/Source/Runtime/Engine/Private/Animation/AnimSequenceTransformProviderData.cpp b/Engine/Source/Runtime/Engine/Private/Animation/AnimSequenceTransformProviderData.cpp
--- a/Engine/Source/Runtime/Engine/Private/Animation/AnimSequenceTransformProviderData.cpp
+++ b/Engine/Source/Runtime/Engine/Private/Animation/AnimSequenceTransformProviderData.cpp
@@ -368,7 +368,18 @@ bool UAnimSequenceTransformProviderData::IsReadyForFinishDestroy()
 
 bool UAnimSequenceTransformProviderData::IsValidFor(FInstancedSkinningSceneExtensionProxy* ExtensionProxy) const
 {
-	const USkinnedAsset* SkinnedAsset = ExtensionProxy->GetSkinnedAsset();
+	if (!SkinnedAsset || !ExtensionProxy->GetSkinnedAsset())
+	{
+		UE_LOG(LogAnimation, Warning, TEXT("AnimSequenceTransformProvider: null skinned asset for provider '%s'"), *GetPathName());
+		return false;
+	}
+
+	if (SkinnedAsset != ExtensionProxy->GetSkinnedAsset())
+	{
+		UE_LOG(LogAnimation, Warning, TEXT("AnimSequenceTransformProvider: SkinnedAsset '%s' in provider '%s' doesn't match SkinnedAsset on InstancedSkinnedMesh: '%s'"),
+			*SkinnedAsset->GetPathName(), *GetPathName(), *ExtensionProxy->GetSkinnedAsset()->GetPathName());
+		return false;
+	}
 
 	for (const FAnimSequenceTransformProviderSequence& Sequence : Sequences)
 	{
@@ -376,25 +387,25 @@ bool UAnimSequenceTransformProviderData::IsValidFor(FInstancedSkinningSceneExten
 		{
 			UE_LOG(LogAnimation, Warning, TEXT("AnimSequenceTransformProvider: Invalid null sequence in provider '%s' for skinned asset '%s'"),
 				*GetPathName(),
-				SkinnedAsset ? *SkinnedAsset->GetPathName() : TEXT("None"));
+				*SkinnedAsset->GetPathName());
 			return false;
 		}
 
 		const USkeleton* SourceSkeleton = Sequence.Sequence->GetSkeleton();
-		const USkeleton* TargetSkeleton = SkinnedAsset ? SkinnedAsset->GetSkeleton() : nullptr;
+		const USkeleton* TargetSkeleton = SkinnedAsset->GetSkeleton();
 
 		if (SourceSkeleton == nullptr)
 		{
 			UE_LOG(LogAnimation, Warning, TEXT("AnimSequenceTransformProvider: Sequence '%s' has no skeleton, cannot use with skinned asset '%s'"),
 				*Sequence.Sequence->GetPathName(),
-				SkinnedAsset ? *SkinnedAsset->GetPathName() : TEXT("None"));
+				*SkinnedAsset->GetPathName());
 			return false;
 		}
 
 		if (TargetSkeleton == nullptr)
 		{
 			UE_LOG(LogAnimation, Warning, TEXT("AnimSequenceTransformProvider: Skinned asset '%s' has no skeleton, cannot use with sequence '%s'"),
-				SkinnedAsset ? *SkinnedAsset->GetPathName() : TEXT("None"),
+				*SkinnedAsset->GetPathName(),
 				*Sequence.Sequence->GetPathName());
 			return false;
 		}
@@ -679,6 +690,7 @@ void UAnimSequenceTransformProviderDataInstance::SubmitChanges()
 
 void UAnimSequenceTransformProviderDataInstance::Initialize(UAnimSequenceTransformProviderData* ProviderData)
 {
+	SkinnedAsset = ProviderData->SkinnedAsset;
 	Parent    = ProviderData;
 	Sequences = ProviderData->Sequences;
 	bEnabled  = ProviderData->bEnabled;
@@ -832,4 +844,4 @@ int32 UAnimSequenceTransformProviderDataInstance::GetTrackSequenceIndex(int32 Tr
 		return Tracks.GetData(TrackIndex).GetSequenceIndex();
 	}
 	return INDEX_NONE;
-}
\ No newline at end of file
+}
diff --git a/Engine/Source/Runtime/Engine/Private/SkeletalRender.cpp b/Engine/Source/Runtime/Engine/Private/SkeletalRender.cpp
--- a/Engine/Source/Runtime/Engine/Private/SkeletalRender.cpp
+++ b/Engine/Source/Runtime/Engine/Private/SkeletalRender.cpp
@@ -67,6 +67,7 @@ FSkeletalMeshObject::FSkeletalMeshObject(
 
 FSkeletalMeshObject::FSkeletalMeshObject(const FSkinnedMeshSceneProxyDesc& InMeshDesc, FSkeletalMeshRenderData* InSkelMeshRenderData, ERHIFeatureLevel::Type InFeatureLevel)
 :	MinDesiredLODLevel(FMath::Max<int32>(InMeshDesc.GetPredictedLODLevel(), InSkelMeshRenderData->CurrentFirstLODIdx))
+,	MinLODLevel(InMeshDesc.GetMinLODLevel())
 ,	MaxDistanceFactor(0.f)
 ,	WorkingMinDesiredLODLevel(MinDesiredLODLevel)
 ,	WorkingMaxDistanceFactor(0.f)
diff --git a/Engine/Source/Runtime/Engine/Private/SkeletalRenderStatic.cpp b/Engine/Source/Runtime/Engine/Private/SkeletalRenderStatic.cpp
--- a/Engine/Source/Runtime/Engine/Private/SkeletalRenderStatic.cpp
+++ b/Engine/Source/Runtime/Engine/Private/SkeletalRenderStatic.cpp
@@ -19,7 +19,6 @@ FSkeletalMeshObjectStatic::FSkeletalMeshObjectStatic(const USkinnedMeshComponent
 
 FSkeletalMeshObjectStatic::FSkeletalMeshObjectStatic(const FSkinnedMeshSceneProxyDesc& InMeshDesc, FSkeletalMeshRenderData* InSkelMeshRenderData, ERHIFeatureLevel::Type InFeatureLevel)
 	: FSkeletalMeshObject(InMeshDesc, InSkelMeshRenderData, InFeatureLevel)
-	, MinLODLevel(InMeshDesc.MinLODLevel)
 {
 	// create LODs to match the base mesh
 	for (int32 LODIndex = 0; LODIndex < InSkelMeshRenderData->LODRenderData.Num(); LODIndex++)
diff --git a/Engine/Source/Runtime/Engine/Private/SkeletalRenderStatic.h b/Engine/Source/Runtime/Engine/Private/SkeletalRenderStatic.h
--- a/Engine/Source/Runtime/Engine/Private/SkeletalRenderStatic.h
+++ b/Engine/Source/Runtime/Engine/Private/SkeletalRenderStatic.h
@@ -108,6 +108,5 @@ private:
 
 	/** Render data for each LOD */
 	TArray<FSkeletalMeshObjectLOD> LODs;
-	int32 MinLODLevel = 0;
 };
 
diff --git a/Engine/Source/Runtime/Engine/Public/SkeletalRenderPublic.h b/Engine/Source/Runtime/Engine/Public/SkeletalRenderPublic.h
--- a/Engine/Source/Runtime/Engine/Public/SkeletalRenderPublic.h
+++ b/Engine/Source/Runtime/Engine/Public/SkeletalRenderPublic.h
@@ -417,6 +417,9 @@ public:
 	 */
 	int32 MinDesiredLODLevel;
 
+	/** The minimum allowed LOD level for this mesh. */
+	const int32 MinLODLevel;
+
 	/** 
 	 *	High (best) DistanceFactor that was desired for rendering this SkeletalMesh last frame. Represents how big this mesh was in screen space  
 	 *	This should only ever be WRITTEN by the RENDER thread (in FSkeletalMeshProxy::PreRenderView) and READ by the GAME thread (in USkeletalMeshComponent::UpdateSkelPose).
diff --git a/Engine/Source/Runtime/Renderer/Private/Skinning/AnimSequenceTransformProvider.cpp b/Engine/Source/Runtime/Renderer/Private/Skinning/AnimSequenceTransformProvider.cpp
--- a/Engine/Source/Runtime/Renderer/Private/Skinning/AnimSequenceTransformProvider.cpp
+++ b/Engine/Source/Runtime/Renderer/Private/Skinning/AnimSequenceTransformProvider.cpp
@@ -8,6 +8,7 @@
 #include "Animation/SkeletonRemapping.h"
 #include "Animation/SkeletonRemappingRegistry.h"
 #include "Animation/Skeleton.h"
+#include "Rendering/SkeletalMeshRenderData.h"
 #include "AnimationRuntime.h"
 #include "AnimEncoding.h"
 #include "ComponentRecreateRenderStateContext.h"
@@ -57,7 +58,7 @@ static TAutoConsoleVariable<float> CVarAnimSequenceTransformProviderTimeScale(
 
 static TAutoConsoleVariable<bool> CVarAnimSequenceTransformProviderGPU(
 	TEXT("r.AnimSequenceTransformProvider.GPU"),
-	true,
+	false,
 	TEXT(""),
 	FConsoleVariableDelegate::CreateLambda([](IConsoleVariable* InVariable)
 	{
@@ -320,42 +321,110 @@ void FAnimSequenceTransformProvider::InitExtension(FScene& InScene)
 	}
 }
 
-void FAnimSequenceTransformProvider::BuildBoneRetargetingTables(
-	FBoneRetargetingTables& BoneRetargetingTables,
+void FAnimSequenceTransformProvider::BuildBoneTables(
+	FBoneTables& BoneTables,
+	EBoneTransformStorageMode TransformStorageMode,
+	const USkinnedAsset* SkinnedMesh,
+	const FSkeletalMeshLODRenderData& LODRenderData)
+{
+	const USkeleton* TargetSkeleton = SkinnedMesh->GetSkeleton();
+	const FReferenceSkeleton& MeshRefSkeleton = SkinnedMesh->GetRefSkeleton();
+	const FReferenceSkeleton& TargetSkeletonRefSkeleton = TargetSkeleton->GetReferenceSkeleton();
+	const int32 NumMeshBones = MeshRefSkeleton.GetRawBoneNum();
+	const int32 NumCompactBones = LODRenderData.RequiredBones.Num();
+
+	BoneTables.MeshToCompact.Init(INDEX_NONE, NumMeshBones);
+	BoneTables.CompactToMesh.SetNumUninitialized(NumCompactBones);
+	BoneTables.CompactToParent.SetNumUninitialized(NumCompactBones);
+	BoneTables.CompactToTargetSkeleton.SetNumUninitialized(NumCompactBones);
+	BoneTables.NumMeshBones = NumMeshBones;
+	BoneTables.NumTransforms = NumMeshBones;
+	BoneTables.NumCompactBones = NumCompactBones;
+
+	for (int32 CompactBoneIndex = 0; CompactBoneIndex < NumCompactBones; ++CompactBoneIndex)
+	{
+		BoneTables.MeshToCompact[LODRenderData.RequiredBones[CompactBoneIndex]] = CompactBoneIndex;
+	}
+
+	for (int32 CompactBoneIndex = 0; CompactBoneIndex < NumCompactBones; ++CompactBoneIndex)
+	{
+		const int32 MeshBoneIndex = LODRenderData.RequiredBones[CompactBoneIndex];
+		BoneTables.CompactToMesh[CompactBoneIndex] = MeshBoneIndex;
+
+		const int32 MeshParentBoneIndex = MeshRefSkeleton.GetParentIndex(MeshBoneIndex);
+		const int32 CompactParentBoneIndex = MeshParentBoneIndex != INDEX_NONE
+			? BoneTables.MeshToCompact[MeshParentBoneIndex]
+			: INDEX_NONE;
+		BoneTables.CompactToParent[CompactBoneIndex] = CompactParentBoneIndex;
+
+		const FName BoneName = MeshRefSkeleton.GetBoneName(MeshBoneIndex);
+		const int32 TargetSkeletonBoneIndex = TargetSkeletonRefSkeleton.FindBoneIndex(BoneName);
+		BoneTables.CompactToTargetSkeleton[CompactBoneIndex] = TargetSkeletonBoneIndex;
+	}
+
+	if (TransformStorageMode == EBoneTransformStorageMode::BoneMap)
+	{
+		TConstArrayView<FSkelMeshRenderSection> Sections = LODRenderData.RenderSections;
+		uint32 NumTransforms = 0;
+
+		for (const FSkelMeshRenderSection& Section : Sections)
+		{
+			if (Section.IsValid())
+			{
+				NumTransforms += Section.BoneMap.Num();
+			}
+		}
+
+		BoneTables.TransformToMesh.Empty(NumTransforms);
+
+		for (const FSkelMeshRenderSection& Section : Sections)
+		{
+			if (Section.IsValid())
+			{
+				BoneTables.TransformToMesh.Append(Section.BoneMap);
+			}
+		}
+
+		BoneTables.NumTransforms = NumTransforms;
+	}
+}
+
+void FAnimSequenceTransformProvider::BuildRetargetingTables(
+	FRetargetingTables& RetargetingTables,
 	USkeleton* AnimSkeleton,
-	const USkinnedAsset* MeshAsset)
+	const USkinnedAsset* MeshAsset,
+	const FBoneTables& BoneTables)
 {
 	const USkeleton* TargetSkeleton = MeshAsset->GetSkeleton();
 	const FReferenceSkeleton& MeshRefSkeleton = MeshAsset->GetRefSkeleton();
 	const FReferenceSkeleton& TargetSkeletonRefSkeleton = TargetSkeleton->GetReferenceSkeleton();
-	const int32 NumMeshBones = MeshRefSkeleton.GetRawBoneNum();
 
 	const TArray<FTransform>& AnimRefPose = AnimSkeleton->GetRefLocalPoses();
 	const TArray<FTransform>& MeshRefPose = MeshRefSkeleton.GetRefBonePose();
 	const bool bUseSourceRetargetModes = TargetSkeleton->GetUseRetargetModesFromCompatibleSkeleton();
 
 	const FSkeletonRemapping& SkeletonRemapping = UE::Anim::FSkeletonRemappingRegistry::Get().GetRemapping(AnimSkeleton, TargetSkeleton);
-	const bool bNeedsRemapping = SkeletonRemapping.IsValid();
-	const bool bRequiresRefPoseRetarget = bNeedsRemapping && SkeletonRemapping.RequiresReferencePoseRetarget();
+	const bool bRequiresRefPoseRetarget = SkeletonRemapping.IsValid() && SkeletonRemapping.RequiresReferencePoseRetarget();
+	const uint32 NumAnimBones = AnimSkeleton->GetReferenceSkeleton().GetRawBoneNum();
 
-	BoneRetargetingTables.Mode.Reset(NumMeshBones);
-	BoneRetargetingTables.Transform.Reset(NumMeshBones);
-	BoneRetargetingTables.AnimBoneIndex.Reset(NumMeshBones);
-	BoneRetargetingTables.bRequiresRefPoseRetarget = bRequiresRefPoseRetarget;
+	RetargetingTables.Mode.Reset(BoneTables.NumCompactBones);
+	RetargetingTables.Transform.Reset(BoneTables.NumCompactBones);
+	RetargetingTables.AnimToCompact.Init(INDEX_NONE, NumAnimBones);
+	RetargetingTables.bRequiresRefPoseRetarget = bRequiresRefPoseRetarget;
 
 	if (bRequiresRefPoseRetarget)
 	{
-		BoneRetargetingTables.RefPoseRetargetQuats.Reset(NumMeshBones);
+		RetargetingTables.RefPoseRetargetQuats.Reset(BoneTables.NumCompactBones);
 	}
 
-	for (int32 MeshBoneIndex = 0; MeshBoneIndex < NumMeshBones; ++MeshBoneIndex)
+	for (uint32 CompactBoneIndex = 0; CompactBoneIndex < BoneTables.NumCompactBones; ++CompactBoneIndex)
 	{
 		EBoneTranslationRetargetingMode::Type Mode = EBoneTranslationRetargetingMode::Animation;
 
-		const FName BoneName = MeshRefSkeleton.GetBoneName(MeshBoneIndex);
-		const int32 TargetSkeletonBoneIndex = TargetSkeletonRefSkeleton.FindBoneIndex(BoneName);
-
-		const int32 AnimBoneIndex = (bNeedsRemapping && TargetSkeletonBoneIndex != INDEX_NONE)
+		const int32 MeshBoneIndex = BoneTables.CompactToMesh[CompactBoneIndex];
+		const int32 TargetSkeletonBoneIndex = BoneTables.CompactToTargetSkeleton[CompactBoneIndex];
+		
+		const int32 AnimBoneIndex = TargetSkeletonBoneIndex != INDEX_NONE && SkeletonRemapping.IsValid()
 			? SkeletonRemapping.GetSourceSkeletonBoneIndex(TargetSkeletonBoneIndex)
 			: TargetSkeletonBoneIndex;
 
@@ -453,26 +522,28 @@ void FAnimSequenceTransformProvider::BuildBoneRetargetingTables(
 				}
 				break;
 			}
+
+			RetargetingTables.AnimToCompact[AnimBoneIndex] = CompactBoneIndex;
 		}
 		else
 		{
-			Transform = FTransform3f(MeshRefPose[MeshBoneIndex]);
+			Transform = FTransform3f(MeshRefTransform);
 		}
 
 		if (bRequiresRefPoseRetarget)
 		{
-			BoneRetargetingTables.RefPoseRetargetQuats.Emplace(TTuple<FQuat4f, FQuat4f>(RefPoseRetargetQuats.Get<0>(), RefPoseRetargetQuats.Get<1>()));
+			RetargetingTables.RefPoseRetargetQuats.Emplace(TTuple<FQuat4f, FQuat4f>(RefPoseRetargetQuats.Get<0>(), RefPoseRetargetQuats.Get<1>()));
 		}
 
-		BoneRetargetingTables.Mode.Emplace(Mode);
-		BoneRetargetingTables.Transform.Emplace(Transform);
-		BoneRetargetingTables.AnimBoneIndex.Emplace(AnimBoneIndex);
+		RetargetingTables.Mode.Emplace(Mode);
+		RetargetingTables.Transform.Emplace(Transform);
+		RetargetingTables.CompactToAnim.Emplace(AnimBoneIndex);
 	}
 }
 
-void FAnimSequenceTransformProvider::BuildLeafBonePathData(FLeafBonePathData& OutData, const FReferenceSkeleton& RefSkeleton)
+void FAnimSequenceTransformProvider::BuildLeafBonePathData(FLeafBonePathData& OutData, const FBoneTables& BoneTables)
 {
-	const int32 NumBones = RefSkeleton.GetRawBoneNum();
+	const int32 NumBones = BoneTables.NumCompactBones;
 
 	struct FBoneData
 	{
@@ -487,7 +558,7 @@ void FAnimSequenceTransformProvider::BuildLeafBonePathData(FLeafBonePathData& Ou
 
 	for (int32 ChildIndex = 1; ChildIndex < NumBones; ++ChildIndex)
 	{
-		const int32 ParentIndex = RefSkeleton.GetParentIndex(ChildIndex);
+		const int32 ParentIndex = BoneTables.CompactToParent[ChildIndex];
 		if (ParentIndex != INDEX_NONE)
 		{
 			FBoneData& ParentBone = BoneDatas[ParentIndex];
@@ -566,7 +637,7 @@ void FAnimSequenceTransformProvider::BuildLeafBonePathData(FLeafBonePathData& Ou
 				Node.bIsOwned  = !CurrentBone.bIsOwned;
 				CurrentBone.bIsOwned = 1;
 
-				CurrentBoneIndex = RefSkeleton.GetParentIndex(CurrentBoneIndex);
+				CurrentBoneIndex = BoneTables.CompactToParent[CurrentBoneIndex];
 				CurrentNodeIndex--;
 			}
 
@@ -583,14 +654,28 @@ void FAnimSequenceTransformProvider::RegisterProxy(FAnimSequenceTransformProvide
 	TConstArrayView<FAnimSequenceTransformProviderSequence> ProviderSequences = Proxy->GetSequences();
 	FInstancedSkinningSceneExtensionProxy* ExtensionProxy = Proxy->GetSceneExtensionProxy();
 	const USkinnedAsset* SkinnedAsset = ExtensionProxy->GetSkinnedAsset();
+	const FSkeletalMeshRenderData* RenderData = SkinnedAsset->GetResourceForRendering();
+	const FSkeletalMeshObject* MeshObject = ExtensionProxy->GetMeshObject();
 	TConstArrayView<FMatrix44f> InvGlobalRefPoseTransforms = SkinnedAsset->GetRefBasesInvMatrix();
 
 	FProxyUserData* ProxyUserData = new FProxyUserData(Proxy);
 	ProxyUserData->Sequences.Reserve(ProviderSequences.Num());
 	ProxyUserData->SequenceToMappingTable.Reserve(ProviderSequences.Num());
 	ProxyUserData->NumMeshBones = InvGlobalRefPoseTransforms.Num();
+	ProxyUserData->MinLODLevel = MeshObject->MinLODLevel;
+
+	const int32 MinLODIndex = MeshObject->MinLODLevel;
+	const int32 NumLODs = RenderData->LODRenderData.Num();
+	ProxyUserData->LODs.SetNum(NumLODs);
+	ProxyUserData->SequenceToMappingTable.Reserve(ProviderSequences.Num() * NumLODs);
+
+	for (int32 LODIndex = MinLODIndex; LODIndex < NumLODs; ++LODIndex)
+	{
+		FProxyUserData::FLOD& LOD = ProxyUserData->LODs[LODIndex];
+		BuildBoneTables(LOD.BoneTables, ExtensionProxy->GetBoneTransformStorageMode(), SkinnedAsset, RenderData->LODRenderData[LODIndex]);
+		BuildLeafBonePathData(LOD.LeafBonePathData, LOD.BoneTables);
+	}
 
-	BuildLeafBonePathData(ProxyUserData->LeafBonePathData, SkinnedAsset->GetRefSkeleton());
 	ProxyUserDatasToUpload.Emplace(ProxyUserData);
 
 	for (const FAnimSequenceTransformProviderSequence& ProviderSequence : ProviderSequences)
@@ -599,7 +684,6 @@ void FAnimSequenceTransformProvider::RegisterProxy(FAnimSequenceTransformProvide
 		USkeleton* SourceSkeleton = AnimSequence->GetSkeleton();
 
 		const FSequenceRegistryKey SequenceRegistryKey(AnimSequence);
-		const FMappingRegistryKey MappingRegistryKey(SourceSkeleton, SkinnedAsset);
 
 		TUniquePtr<FSequence>& Sequence = SequenceRegistry.FindOrAdd(SequenceRegistryKey);
 		if (!Sequence)
@@ -620,7 +704,7 @@ void FAnimSequenceTransformProvider::RegisterProxy(FAnimSequenceTransformProvide
 			Sequence->FramesUploaded.Init(false, Sequence->NumFrames);
 			Sequence->NumFramesRequested     = 0;
 			Sequence->NumFramesUploaded      = 0;
-			
+
 			check(!SequencesToUpload.Contains(Sequence.Get()));
 			SequencesToUpload.AddUnique(Sequence.Get());
 
@@ -628,33 +712,45 @@ void FAnimSequenceTransformProvider::RegisterProxy(FAnimSequenceTransformProvide
 		}
 		++Sequence->RefCount;
 
-		TUniquePtr<FMapping>& Mapping = MappingRegistry.FindOrAdd(MappingRegistryKey);
-		if (!Mapping)
+		for (int32 LODIndex = 0; LODIndex < MinLODIndex; ++LODIndex)
 		{
-			Mapping = MakeUnique<FMapping>();
-
-			Mapping->RegistryKey    = MappingRegistryKey;
-			Mapping->SourceSkeleton = SourceSkeleton;
-			Mapping->SkinnedAsset   = SkinnedAsset;
-
-			Mapping->NumMeshBones = ProxyUserData->NumMeshBones;
-			BuildBoneRetargetingTables(Mapping->BoneRetargetingTables, SourceSkeleton, SkinnedAsset);
+			ProxyUserData->SequenceToMappingTable.Emplace(nullptr);
+		}
 
-			Mapping->InvGlobalRefPoseTransforms.Reserve(InvGlobalRefPoseTransforms.Num());
-			for (const FMatrix44f& InvGlobalRefPose : InvGlobalRefPoseTransforms)
+		for (int32 LODIndex = MinLODIndex; LODIndex < NumLODs; ++LODIndex)
+		{
+			const FMappingRegistryKey MappingRegistryKey(SourceSkeleton, SkinnedAsset, LODIndex);
+			TUniquePtr<FMapping>& Mapping = MappingRegistry.FindOrAdd(MappingRegistryKey);
+			if (!Mapping)
 			{
-				Mapping->InvGlobalRefPoseTransforms.Emplace(InvGlobalRefPose);
-			}
+				Mapping = MakeUnique<FMapping>();
+
+				const FProxyUserData::FLOD& LOD = ProxyUserData->LODs[LODIndex];
+				Mapping->SourceSkeleton  = SourceSkeleton;
+				Mapping->SkinnedAsset    = SkinnedAsset;
+				Mapping->NumMeshBones    = ProxyUserData->NumMeshBones;
+				Mapping->LODIndex        = LODIndex;
+				Mapping->NumCompactBones = LOD.BoneTables.NumCompactBones;
+				BuildRetargetingTables(Mapping->RetargetingTables, SourceSkeleton, SkinnedAsset, LOD.BoneTables);
+
+				Mapping->InvGlobalRefPoseTransforms.Reserve(Mapping->NumCompactBones);
+				for (uint32 CompactBoneIndex = 0; CompactBoneIndex < Mapping->NumCompactBones; ++CompactBoneIndex)
+				{
+					const int32 MeshBoneIndex = LOD.BoneTables.CompactToMesh[CompactBoneIndex];
+					const FMatrix44f& InvGlobalRefPose = InvGlobalRefPoseTransforms[MeshBoneIndex];
+					Mapping->InvGlobalRefPoseTransforms.Emplace(InvGlobalRefPose);
+				}
 
-			MappingsToUpload.Add(Mapping.Get());
+				MappingsToUpload.Add(Mapping.Get());
 
-			INC_DWORD_STAT(STAT_AnimSequenceTransformProvider_Mappings);
+				INC_DWORD_STAT(STAT_AnimSequenceTransformProvider_Mappings);
+			}
+			++Mapping->RefCount;
+			ProxyUserData->Mappings.AddUnique(Mapping.Get());
+			ProxyUserData->SequenceToMappingTable.Emplace(Mapping.Get());
 		}
-		++Mapping->RefCount;
 
 		ProxyUserData->Sequences.Emplace(Sequence.Get());
-		ProxyUserData->SequenceToMappingTable.Emplace(Mapping.Get());
-		ProxyUserData->Mappings.AddUnique(Mapping.Get());
 	}
 
 	ExtensionProxy->SetUserData(ProxyUserData);
@@ -670,10 +766,13 @@ void FAnimSequenceTransformProvider::UnregisterProxy(FAnimSequenceTransformProvi
 
 	ProxyUserDatasToUpload.RemoveSingleSwap(ProxyUserData);
 
-	if (ProxyUserData->LeafBonePathBufferAllocation.IsValid())
+	for (FProxyUserData::FLOD& LOD : ProxyUserData->LODs)
 	{
-		LeafBonePathAllocator.Free(ProxyUserData->LeafBonePathBufferAllocation.Offset, ProxyUserData->LeafBonePathBufferAllocation.Count);
-		ProxyUserData->LeafBonePathBufferAllocation = {};
+		if (LOD.LeafBonePathBufferAllocation.IsValid())
+		{
+			LeafBonePathAllocator.Free(LOD.LeafBonePathBufferAllocation.Offset, LOD.LeafBonePathBufferAllocation.Count);
+			LOD.LeafBonePathBufferAllocation = {};
+		}
 	}
 
 	for (FSequence* Sequence : ProxyUserData->Sequences)
@@ -718,7 +817,7 @@ void FAnimSequenceTransformProvider::UnregisterProxy(FAnimSequenceTransformProvi
 			}
 
 			MappingsToUpload.RemoveSingleSwap(Mapping, EAllowShrinking::No);
-			MappingRegistry.Remove(Mapping->RegistryKey);
+			MappingRegistry.Remove(FMappingRegistryKey(Mapping->SourceSkeleton, Mapping->SkinnedAsset, Mapping->LODIndex));
 			Mapping = nullptr;
 
 			DEC_DWORD_STAT(STAT_AnimSequenceTransformProvider_Mappings);
@@ -821,12 +920,14 @@ void FAnimSequenceTransformProvider::ProvideGPUTransforms(FSkinningTransformProv
 	{
 		FInstancedSkinningSceneExtensionProxy* ExtensionProxy = static_cast<FInstancedSkinningSceneExtensionProxy*>(Context.Proxies[Indirection.Index]);
 		FProxyUserData& ProxyUserData = static_cast<FProxyUserData&>(ExtensionProxy->GetUserDataChecked());
+		const uint32 LODIndex = ExtensionProxy->GetLOD();
+		const FProxyUserData::FLOD& LOD = ProxyUserData.LODs[LODIndex];
 		FAnimSequenceTrackRenderPool& Tracks = ProxyUserData.Tracks;
 		ProxyUserData.TrackHistories.SetNum(Tracks.GetNumTracks());
 
-		const uint32 NumEvaluateBlockHeaders    = FMath::DivideAndRoundUp(ProxyUserData.NumMeshBones, BonesPerGroup);
-		const uint32 NumScatterBlockHeaders     = FMath::DivideAndRoundUp(ExtensionProxy->GetMaxBoneTransformCount(), BonesPerGroup);
-		const uint32 NumConcatenateBlockHeaders = FMath::DivideAndRoundUp(ProxyUserData.LeafBonePathData.NumLeafBones, BonesPerGroup);
+		const uint32 NumEvaluateBlockHeaders    = FMath::DivideAndRoundUp(LOD.BoneTables.NumCompactBones, BonesPerGroup);
+		const uint32 NumScatterBlockHeaders     = FMath::DivideAndRoundUp(LOD.BoneTables.NumTransforms, BonesPerGroup);
+		const uint32 NumConcatenateBlockHeaders = FMath::DivideAndRoundUp(LOD.LeafBonePathData.NumLeafBones, BonesPerGroup);
 		const uint32 CurrentTransformSlot = Indirection.CurrentTransformOffset < Indirection.PreviousTransformOffset ? 0 : 1;
 	
 		int32 NumSubTracks = 0;
@@ -868,7 +969,7 @@ void FAnimSequenceTransformProvider::ProvideGPUTransforms(FSkinningTransformProv
 			TrackUpdateRequests.Emplace(UpdateRequest);
 		});
 
-		NumTotalIntermediateTransforms    += NumSubTracks * ProxyUserData.NumMeshBones;
+		NumTotalIntermediateTransforms    += NumSubTracks * LOD.BoneTables.NumCompactBones;
 		NumTotalEvaluateBlockHeaders      += NumEvaluateSubTracks * NumEvaluateBlockHeaders;
 		NumTotalEvaluateBlendBlockHeaders += NumEvaluateBlendSubTracks * NumEvaluateBlockHeaders;
 		NumTotalScatterBlockHeaders       += NumSubTracks * NumScatterBlockHeaders;
@@ -881,15 +982,16 @@ void FAnimSequenceTransformProvider::ProvideGPUTransforms(FSkinningTransformProv
 
 	for (FMapping* Mapping : MappingsToUpload)
 	{
-		const int32 NumRetargetingTransforms = (Mapping->BoneRetargetingTables.bRequiresRefPoseRetarget ? 2 : 1);
-		const int32 NumTransforms            = Mapping->NumMeshBones * (NumRetargetingTransforms + 1);
-		const int32 NumRetargetingDatas      = Mapping->NumMeshBones;
+		const int32 NumInvGlobalRefPoses     = Mapping->NumCompactBones;
+		const int32 NumRetargetingTransforms = (Mapping->RetargetingTables.bRequiresRefPoseRetarget ? 2 : 1);
+		const int32 NumTransforms            = Mapping->NumCompactBones * NumRetargetingTransforms + NumInvGlobalRefPoses;
+		const int32 NumRetargetingDatas      = Mapping->NumCompactBones;
 
 		Mapping->MappingTransformBufferAllocation.Count  = NumTransforms;
 		Mapping->MappingTransformBufferAllocation.Offset = SequenceTransformAllocator.Allocate(NumTransforms);
 
 		Mapping->InvGlobalRefPoseTransformBufferOffset = Mapping->MappingTransformBufferAllocation.Offset;
-		Mapping->RetargetingTransformBufferOffset    = Mapping->MappingTransformBufferAllocation.Offset + Mapping->NumMeshBones;
+		Mapping->RetargetingTransformBufferOffset      = Mapping->MappingTransformBufferAllocation.Offset + NumInvGlobalRefPoses;
 
 		Mapping->RetargetingDataBufferAllocation.Count = NumRetargetingDatas;
 		Mapping->RetargetingDataBufferAllocation.Offset = RetargetingDataAllocator.Allocate(NumRetargetingDatas);
@@ -921,12 +1023,21 @@ void FAnimSequenceTransformProvider::ProvideGPUTransforms(FSkinningTransformProv
 
 	for (FProxyUserData* ProxyUserData : ProxyUserDatasToUpload)
 	{
-		const int32 NumElements = ProxyUserData->LeafBonePathData.Buffer.Num();
+		for (FProxyUserData::FLOD& LOD : ProxyUserData->LODs)
+		{
+			if (LOD.IsValid())
+			{
+				const int32 NumLeafBonePaths = LOD.LeafBonePathData.Buffer.Num();
+				const int32 NumElements = NumLeafBonePaths + LOD.BoneTables.NumTransforms;
+
+				LOD.LeafBonePathBufferAllocation.Count = NumElements;
+				LOD.LeafBonePathBufferAllocation.Offset = LeafBonePathAllocator.Allocate(NumElements);
 
-		ProxyUserData->LeafBonePathBufferAllocation.Count = NumElements;
-		ProxyUserData->LeafBonePathBufferAllocation.Offset = LeafBonePathAllocator.Allocate(NumElements);
+				LOD.BoneMapBufferOffset = LOD.LeafBonePathBufferAllocation.Offset + NumLeafBonePaths;
 
-		NumTotalLeafBonePathElements += NumElements;
+				NumTotalLeafBonePathElements += NumElements;
+			}
+		}
 	}
 
 	int32& NumTotalFramesRequested = *GraphBuilder.AllocPOD<int32>();
@@ -967,7 +1078,7 @@ void FAnimSequenceTransformProvider::ProvideGPUTransforms(FSkinningTransformProv
 
 		const auto EvaluateAnimSequence = [&](
 			const FSkinningTransformProvider::FProviderIndirection& Indirection,
-			const FProxyUserData& ProxyUserData,
+			const FProxyUserData::FLOD& LOD,
 			const FMapping& SourceMapping,
 			const FMapping& TargetMapping,
 			FSequence& SourceSequence,
@@ -993,8 +1104,8 @@ void FAnimSequenceTransformProvider::ProvideGPUTransforms(FSkinningTransformProv
 
 			// Evaluation Pass
 			{
-				check(SourceMapping.NumMeshBones == TargetMapping.NumMeshBones);
-				const uint32 NumBlocks = FMath::DivideAndRoundUp(TargetMapping.NumMeshBones, BonesPerGroup);
+				check(SourceMapping.NumCompactBones == TargetMapping.NumCompactBones);
+				const uint32 NumBlocks = FMath::DivideAndRoundUp(TargetMapping.NumCompactBones, BonesPerGroup);
 
 				const uint32 SourceSequenceTransformBufferOffset    = SourceSequence.BufferAllocation.Offset * sizeof(FBoneTransformWithScale);
 				const uint32 SourceRetargetingTransformBufferOffset = SourceMapping.RetargetingTransformBufferOffset * sizeof(FBoneTransformWithScale);
@@ -1034,7 +1145,7 @@ void FAnimSequenceTransformProvider::ProvideGPUTransforms(FSkinningTransformProv
 						BlockHeader.TargetNumAnimBones                     = TargetSequence.NumAnimBones;
 						BlockHeader.DstTransformBufferOffset               = IntermediateTransformOffset;
 						BlockHeader.BlockLocalIndex                        = BlockIndex;
-						BlockHeader.BlockTransformCount                    = GetBlockTransformCount(BlockIndex, TargetMapping.NumMeshBones, BonesPerGroup);
+						BlockHeader.BlockTransformCount                    = GetBlockTransformCount(BlockIndex, TargetMapping.NumCompactBones, BonesPerGroup);
 					}
 
 					NumEvaluateBlendBlocks += NumBlocks;
@@ -1054,7 +1165,7 @@ void FAnimSequenceTransformProvider::ProvideGPUTransforms(FSkinningTransformProv
 						BlockHeader.KeyIndex1                        = SourceKeyIndex1;
 						BlockHeader.Alpha                            = SourceAlpha;
 						BlockHeader.BlockLocalIndex                  = BlockIndex;
-						BlockHeader.BlockTransformCount              = GetBlockTransformCount(BlockIndex, TargetMapping.NumMeshBones, BonesPerGroup);
+						BlockHeader.BlockTransformCount              = GetBlockTransformCount(BlockIndex, TargetMapping.NumCompactBones, BonesPerGroup);
 					}
 
 					NumEvaluateBlocks += NumBlocks;
@@ -1063,10 +1174,10 @@ void FAnimSequenceTransformProvider::ProvideGPUTransforms(FSkinningTransformProv
 
 			// Concatenate Local to Global Pass
 			{
-				const uint32 NumLeafBones = ProxyUserData.LeafBonePathData.NumLeafBones;
+				const uint32 NumLeafBones = LOD.LeafBonePathData.NumLeafBones;
 				const uint32 NumBlocks    = FMath::DivideAndRoundUp(NumLeafBones, BonesPerGroup);
 
-				const uint32 LeafBonePathBufferOffset = ProxyUserData.LeafBonePathBufferAllocation.Offset * sizeof(uint32);
+				const uint32 LeafBonePathBufferOffset = LOD.LeafBonePathBufferAllocation.Offset * sizeof(uint32);
 
 				for (uint32 BlockIndex = 0; BlockIndex < NumBlocks; ++BlockIndex)
 				{
@@ -1074,7 +1185,7 @@ void FAnimSequenceTransformProvider::ProvideGPUTransforms(FSkinningTransformProv
 
 					BlockHeader.LeafBonePathBufferOffset = LeafBonePathBufferOffset;
 					BlockHeader.TransformBufferOffset    = IntermediateTransformOffset;
-					BlockHeader.NumPrefetchBones         = FMath::Min(ProxyUserData.NumMeshBones, BonesPerGroup);
+					BlockHeader.NumPrefetchBones         = FMath::Min(LOD.BoneTables.NumCompactBones, BonesPerGroup);
 					BlockHeader.BlockLocalIndex          = BlockIndex;
 					BlockHeader.BlockLeafBoneCount       = GetBlockTransformCount(BlockIndex, NumLeafBones, BonesPerGroup);
 				}
@@ -1082,8 +1193,8 @@ void FAnimSequenceTransformProvider::ProvideGPUTransforms(FSkinningTransformProv
 
 			// Scatter Pass
 			{
-				const uint32 NumBlocks     = FMath::DivideAndRoundUp(MaxBoneTransformCount, BonesPerGroup);
-				const uint32 BoneMapOffset = ProxyUserData.TransformStorageMode == EBoneTransformStorageMode::BoneMap ? Indirection.HierarchyOffset : ~uint32(0);
+				const uint32 NumBlocks           = FMath::DivideAndRoundUp(MaxBoneTransformCount, BonesPerGroup);
+				const uint32 BoneMapBufferOffset = LOD.BoneMapBufferOffset * sizeof(uint32);
 
 				const uint32 InvGlobalRefPoseTransformBufferOffset = SourceMapping.InvGlobalRefPoseTransformBufferOffset * sizeof(FBoneTransformWithScale);
 
@@ -1095,13 +1206,13 @@ void FAnimSequenceTransformProvider::ProvideGPUTransforms(FSkinningTransformProv
 					BlockHeader.SrcTransformBufferOffset              = IntermediateTransformOffset;
 					BlockHeader.DstTransformBufferOffset1             = DstTransformBufferOffset1;
 					BlockHeader.DstTransformBufferOffset2             = DstTransformBufferOffset2;
-					BlockHeader.BoneMapOffset                         = BoneMapOffset;
+					BlockHeader.BoneMapOffset                         = BoneMapBufferOffset;
 					BlockHeader.BlockLocalIndex                       = BlockIndex;
 					BlockHeader.BlockTransformCount                   = GetBlockTransformCount(BlockIndex, MaxBoneTransformCount, BonesPerGroup);
 				}
 			}
 
-			IntermediateTransformOffset += SourceMapping.NumMeshBones * sizeof(FBoneTransformWithScale);
+			IntermediateTransformOffset += SourceMapping.NumCompactBones * sizeof(FBoneTransformWithScale);
 		};
 
 		const float GlobalTimeScale = CVarAnimSequenceTransformProviderTimeScale.GetValueOnRenderThread();
@@ -1110,12 +1221,15 @@ void FAnimSequenceTransformProvider::ProvideGPUTransforms(FSkinningTransformProv
 
 		for (const FSkinningTransformProvider::FProviderIndirection Indirection : Indirections)
 		{
-			FProxyUserData& ProxyUserData = static_cast<FProxyUserData&>(static_cast<FInstancedSkinningSceneExtensionProxy*>(Proxies[Indirection.Index])->GetUserDataChecked());
+			FInstancedSkinningSceneExtensionProxy* ExtensionProxy = static_cast<FInstancedSkinningSceneExtensionProxy*>(Proxies[Indirection.Index]);
+			const uint32 LODIndex = ExtensionProxy->GetLOD();
+
+			FProxyUserData& ProxyUserData = static_cast<FProxyUserData&>(ExtensionProxy->GetUserDataChecked());
+			FProxyUserData::FLOD& LOD = ProxyUserData.LODs[LODIndex];
 			FAnimSequenceTrackRenderPool& Tracks = ProxyUserData.Tracks;
 
 			const uint32 CurrentTransformOffset  = Indirection.CurrentTransformOffset;
 			const uint32 PreviousTransformOffset = Indirection.PreviousTransformOffset;
-			const uint32 MaxBoneTransformCount   = ProxyUserData.SceneExtensionProxy.GetMaxBoneTransformCount();
 
 			Tracks.EnumerateActiveTracks([&] (int32 TrackIndex)
 			{
@@ -1124,8 +1238,8 @@ void FAnimSequenceTransformProvider::ProvideGPUTransforms(FSkinningTransformProv
 				const FAnimSequenceTrackPackedData& Source = Data.GetSource();
 				const FAnimSequenceTrackPackedData& Target = Data.GetTarget();
 
-				const FMapping& SourceMapping = ProxyUserData.GetMapping(Source);
-				const FMapping& TargetMapping = ProxyUserData.GetMapping(Target);
+				const FMapping& SourceMapping = ProxyUserData.GetMapping(Source, LODIndex);
+				const FMapping& TargetMapping = ProxyUserData.GetMapping(Target, LODIndex);
 				FSequence& SourceSequence = ProxyUserData.GetSequence(Source);
 				FSequence& TargetSequence = ProxyUserData.GetSequence(Target);
 
@@ -1137,7 +1251,7 @@ void FAnimSequenceTransformProvider::ProvideGPUTransforms(FSkinningTransformProv
 				const float TargetPreviousPosition = Target.GetPreviousPosition();
 				const float BlendWeight            = Data.GetBlendWeight();
 
-				const uint32 TrackTransformOffset         = TrackIndex * MaxBoneTransformCount * 2 * sizeof(FCompressedBoneTransform);
+				const uint32 TrackTransformOffset         = TrackIndex * LOD.BoneTables.NumTransforms * 2 * sizeof(FCompressedBoneTransform);
 				const uint32 TrackCurrentTransformOffset  = TrackTransformOffset + CurrentTransformOffset;
 				const uint32 TrackPreviousTransformOffset = TrackTransformOffset + PreviousTransformOffset;
 
@@ -1148,12 +1262,12 @@ void FAnimSequenceTransformProvider::ProvideGPUTransforms(FSkinningTransformProv
 
 					EvaluateAnimSequence(
 						Indirection,
-						ProxyUserData,
+						LOD,
 						SourceMapping, TargetMapping,
 						SourceSequence, TargetSequence,
 						SourcePosition, TargetPosition,
 						BlendWeight,
-						MaxBoneTransformCount,
+						LOD.BoneTables.NumTransforms,
 						OutputTransformOffset1, OutputTransformOffset2);
 
 					NumCurrentTransformUpdates++;
@@ -1163,12 +1277,12 @@ void FAnimSequenceTransformProvider::ProvideGPUTransforms(FSkinningTransformProv
 				{
 					EvaluateAnimSequence(
 						Indirection,
-						ProxyUserData,
+						LOD,
 						SourceMapping, TargetMapping,
 						SourceSequence, TargetSequence,
 						SourcePreviousPosition, TargetPreviousPosition,
 						BlendWeight,
-						MaxBoneTransformCount,
+						LOD.BoneTables.NumTransforms,
 						TrackPreviousTransformOffset, TrackPreviousTransformOffset);
 
 					NumPreviousTransformUpdates++;
@@ -1214,18 +1328,18 @@ void FAnimSequenceTransformProvider::ProvideGPUTransforms(FSkinningTransformProv
 			{
 				TArrayView<FBoneTransformWithScale> DstTransforms = ScatterUploader.Add_GetRef<FBoneTransformWithScale>(Mapping->MappingTransformBufferAllocation.Offset, Mapping->MappingTransformBufferAllocation.Count);
 
-				for (uint32 MeshBoneIndex = 0; MeshBoneIndex < Mapping->NumMeshBones; ++MeshBoneIndex)
+				for (uint32 CompactBoneIndex = 0; CompactBoneIndex < Mapping->NumCompactBones; ++CompactBoneIndex)
 				{
-					DstTransforms[MeshBoneIndex].SetTransform(Mapping->InvGlobalRefPoseTransforms[MeshBoneIndex]);
+					DstTransforms[CompactBoneIndex].SetTransform(Mapping->InvGlobalRefPoseTransforms[CompactBoneIndex]);
 
-					const uint32 RetargetTransformIndex = Mapping->NumMeshBones + MeshBoneIndex * (Mapping->BoneRetargetingTables.bRequiresRefPoseRetarget ? 2 : 1);
-					DstTransforms[RetargetTransformIndex].SetTransform(Mapping->BoneRetargetingTables.Transform[MeshBoneIndex]);
+					const uint32 RetargetTransformIndex = Mapping->NumCompactBones + CompactBoneIndex * (Mapping->RetargetingTables.bRequiresRefPoseRetarget ? 2 : 1);
+					DstTransforms[RetargetTransformIndex].SetTransform(Mapping->RetargetingTables.Transform[CompactBoneIndex]);
 
-					if (Mapping->BoneRetargetingTables.bRequiresRefPoseRetarget)
+					if (Mapping->RetargetingTables.bRequiresRefPoseRetarget)
 					{
 						FBoneTransformWithScale& RefPoseSlot = DstTransforms[RetargetTransformIndex + 1];
-						const FQuat4f& Q0 = Mapping->BoneRetargetingTables.RefPoseRetargetQuats[MeshBoneIndex].Get<0>();
-						const FQuat4f& Q1 = Mapping->BoneRetargetingTables.RefPoseRetargetQuats[MeshBoneIndex].Get<1>();
+						const FQuat4f& Q0 = Mapping->RetargetingTables.RefPoseRetargetQuats[CompactBoneIndex].Get<0>();
+						const FQuat4f& Q1 = Mapping->RetargetingTables.RefPoseRetargetQuats[CompactBoneIndex].Get<1>();
 
 						// Pack Q0 into Rotation
 						RefPoseSlot.Rotation[0] = Q0.X;
@@ -1379,12 +1493,12 @@ void FAnimSequenceTransformProvider::ProvideGPUTransforms(FSkinningTransformProv
 			{
 				TArrayView<UE::HLSL::FBoneRetargetingData> DstRetargetingDatas = ScatterUploader.Add_GetRef<UE::HLSL::FBoneRetargetingData>(Mapping->RetargetingDataBufferAllocation.Offset, Mapping->RetargetingDataBufferAllocation.Count);
 
-				for (uint32 MeshBoneIndex = 0; MeshBoneIndex < Mapping->NumMeshBones; ++MeshBoneIndex)
+				for (uint32 CompactBoneIndex = 0; CompactBoneIndex < Mapping->NumCompactBones; ++CompactBoneIndex)
 				{
-					UE::HLSL::FBoneRetargetingData& RetargetingData = DstRetargetingDatas[MeshBoneIndex];
-					RetargetingData.RetargetMode  = Mapping->BoneRetargetingTables.Mode[MeshBoneIndex];
-					RetargetingData.AnimBoneIndex = Mapping->BoneRetargetingTables.AnimBoneIndex[MeshBoneIndex];
-					RetargetingData.bRequiresRefPoseRetarget = Mapping->BoneRetargetingTables.bRequiresRefPoseRetarget;
+					UE::HLSL::FBoneRetargetingData& RetargetingData = DstRetargetingDatas[CompactBoneIndex];
+					RetargetingData.RetargetMode  = Mapping->RetargetingTables.Mode[CompactBoneIndex];
+					RetargetingData.AnimBoneIndex = Mapping->RetargetingTables.CompactToAnim[CompactBoneIndex];
+					RetargetingData.bRequiresRefPoseRetarget = Mapping->RetargetingTables.bRequiresRefPoseRetarget;
 				}
 			}
 		});
@@ -1401,9 +1515,38 @@ void FAnimSequenceTransformProvider::ProvideGPUTransforms(FSkinningTransformProv
 
 			for (const FProxyUserData* ProxyUserData : ProxyUserDatasToUpload)
 			{
-				TArrayView<uint32> DstBuffer = ScatterUploader.Add_GetRef<uint32>(ProxyUserData->LeafBonePathBufferAllocation.Offset, ProxyUserData->LeafBonePathBufferAllocation.Count);
+				const FSkinningSceneExtensionProxy& ExtensionProxy = ProxyUserData->SceneExtensionProxy;
+
+				for (const FProxyUserData::FLOD& LOD : ProxyUserData->LODs)
+				{
+					if (!LOD.IsValid())
+					{
+						continue;
+					}
+
+					const FBoneTables& BoneTables = LOD.BoneTables;
+					const FLeafBonePathData& LeafBonePathData = LOD.LeafBonePathData;
+
+					TArrayView<uint32> DstBuffer = ScatterUploader.Add_GetRef<uint32>(LOD.LeafBonePathBufferAllocation.Offset, LOD.LeafBonePathBufferAllocation.Count);
+
+					const int32 NumLeafBonePaths = LeafBonePathData.Buffer.Num();
+
+					FMemory::Memcpy(DstBuffer.GetData(), LeafBonePathData.Buffer.GetData(), NumLeafBonePaths * sizeof(uint32));
 
-				FMemory::Memcpy(DstBuffer.GetData(), ProxyUserData->LeafBonePathData.Buffer.GetData(), ProxyUserData->LeafBonePathData.Buffer.Num() * sizeof(uint32));
+					if (!BoneTables.TransformToMesh.IsEmpty())
+					{
+						for (int32 TransformIndex = 0; TransformIndex < BoneTables.TransformToMesh.Num(); ++TransformIndex)
+						{
+							const int32 MeshBoneIndex = BoneTables.TransformToMesh[TransformIndex];
+							const int32 CompactBoneIndex = BoneTables.MeshToCompact[MeshBoneIndex];
+							DstBuffer[NumLeafBonePaths + TransformIndex] = CompactBoneIndex;
+						}
+					}
+					else
+					{
+						FMemory::Memcpy(DstBuffer.GetData() + NumLeafBonePaths, BoneTables.MeshToCompact.GetData(), BoneTables.MeshToCompact.Num() * sizeof(uint32));
+					}
+				}
 			}
 			ProxyUserDatasToUpload.Empty();
 		});
@@ -1488,7 +1631,7 @@ void FAnimSequenceTransformProvider::ProvideGPUTransforms(FSkinningTransformProv
 		PassParameters->GlobalSpaceTransformBuffer = GraphBuilder.CreateSRV(GlobalSpaceTransformBuffer);
 		PassParameters->SequenceTransformBuffer = GraphBuilder.CreateSRV(SequenceTransformBuffer);
 		PassParameters->ScatterBlockHeaderBuffer = GraphBuilder.CreateSRV(ScatterBlockHeaderBuffer);
-		PassParameters->BoneMapBuffer = Context.HierarchyBufferSRV;
+		PassParameters->BoneMapBuffer = GraphBuilder.CreateSRV(LeafBonePathBuffer);
 		PassParameters->SkinningTransformBufferUAV = GetCompressedBoneTransformUAV(GraphBuilder, Context.TransformBuffer);
 
 		TShaderMapRef<FScatterTransformsCS> ComputeShader(ShaderMap);
@@ -1532,27 +1675,20 @@ void FAnimSequenceTransformProvider::ProvideCPUTransforms(FSkinningTransformProv
 
 		BoneTrackArray TrackToBoneIndexMap;
 		TArray<FTransform, SceneRenderingAllocator> SampledLocalAnimPose;
-
-		const auto DecompressPose = [&TrackToBoneIndexMap, &SampledLocalAnimPose] (const FMapping& Mapping, const FSequence& Sequence, float TimeOffset, TArray<FTransform3f, SceneRenderingAllocator>& SampledGlobalMeshPose)
+		TBitArray<SceneRenderingBitArrayAllocator> SampledLocalAnimPoseBits;
+
+		const auto DecompressPose = [&TrackToBoneIndexMap, &SampledLocalAnimPose, &SampledLocalAnimPoseBits] (
+			  const FBoneTables& BoneTables
+			, const FMapping& Mapping
+			, const FSequence& Sequence
+			, float TimeOffset
+			, TArray<FTransform3f, SceneRenderingAllocator>& SampledGlobalMeshPose
+		)
 		{
 			const UAnimSequence* AnimSequence = Sequence.Sequence;
 			UAnimSequence::FScopedCompressedAnimSequence CompressedAnimSequence = AnimSequence->GetCompressedData();
 			const FCompressedAnimSequence& PlatformCompressedData = CompressedAnimSequence.Get();
 
-			const uint32 NumCompressedTracks = PlatformCompressedData.CompressedTrackToSkeletonMapTable.Num();
-			TrackToBoneIndexMap.Reset(NumCompressedTracks);
-
-			for (uint32 CompressedTrackIndex = 0; CompressedTrackIndex < NumCompressedTracks; ++CompressedTrackIndex)
-			{
-				const uint32 BoneIndex = PlatformCompressedData.GetSkeletonIndexFromTrackIndex(CompressedTrackIndex);
-
-				// We only care about raw bones.
-				if (BoneIndex < Sequence.NumAnimBones)
-				{
-					TrackToBoneIndexMap.Add(BoneTrackPair(BoneIndex, CompressedTrackIndex));
-				}
-			}
-
 			FAnimSequenceDecompressionContext DecompressionContext(
 				  Sequence.SampleRate
 				, Sequence.NumFrames
@@ -1566,8 +1702,30 @@ void FAnimSequenceTransformProvider::ProvideCPUTransforms(FSkinningTransformProv
 				, AnimSequence->GetAdditiveAnimType()
 			);
 
-			SampledLocalAnimPose.Reset(DecompressionContext.GetRefLocalPoses().Num());
-			SampledLocalAnimPose.Append(DecompressionContext.GetRefLocalPoses());
+			const uint32 NumCompressedTracks = PlatformCompressedData.CompressedTrackToSkeletonMapTable.Num();
+			TrackToBoneIndexMap.Reset(NumCompressedTracks);
+			
+			TConstArrayView<FTransform> AnimRefLocalPoses = DecompressionContext.GetRefLocalPoses();
+			SampledLocalAnimPose.SetNumUninitialized(BoneTables.NumCompactBones);
+			SampledLocalAnimPoseBits.Init(false, BoneTables.NumCompactBones);
+
+			for (uint32 CompressedTrackIndex = 0; CompressedTrackIndex < NumCompressedTracks; ++CompressedTrackIndex)
+			{
+				const int32 AnimBoneIndex = PlatformCompressedData.GetSkeletonIndexFromTrackIndex(CompressedTrackIndex);
+
+				// We only care about raw bones.
+				if (AnimBoneIndex < (int32)Sequence.NumAnimBones)
+				{
+					const int32 CompactBoneIndex = Mapping.RetargetingTables.AnimToCompact[AnimBoneIndex];
+
+					if (CompactBoneIndex != INDEX_NONE)
+					{
+						TrackToBoneIndexMap.Add(BoneTrackPair(CompactBoneIndex, CompressedTrackIndex));
+						SampledLocalAnimPose[CompactBoneIndex] = AnimRefLocalPoses[AnimBoneIndex];
+						SampledLocalAnimPoseBits[CompactBoneIndex] = true;
+					}
+				}
+			}
 
 			DecompressionContext.Seek((double)TimeOffset);
 
@@ -1580,33 +1738,42 @@ void FAnimSequenceTransformProvider::ProvideCPUTransforms(FSkinningTransformProv
 			const FReferenceSkeleton& MeshRefSkeleton = Mapping.SkinnedAsset->GetRefSkeleton();
 			TConstArrayView<FTransform> MeshLocalRefPoses = MeshRefSkeleton.GetRawRefBonePose();
 
-			for (int32 MeshBoneIndex = 0; MeshBoneIndex < (int32)Mapping.NumMeshBones; ++MeshBoneIndex)
+			for (uint32 CompactBoneIndex = 0; CompactBoneIndex < Mapping.NumCompactBones; ++CompactBoneIndex)
 			{
-				const int32 AnimBoneIndex = Mapping.BoneRetargetingTables.AnimBoneIndex[MeshBoneIndex];
+				const int32 AnimBoneIndex = Mapping.RetargetingTables.CompactToAnim[CompactBoneIndex];
+				const int32 MeshBoneIndex = BoneTables.CompactToMesh[CompactBoneIndex];
 
 				FTransform3f MeshLocalTransform;
 
 				if (AnimBoneIndex != INDEX_NONE)
 				{
-					const FTransform3f RetargetTransform = Mapping.BoneRetargetingTables.Transform[MeshBoneIndex];
-					MeshLocalTransform = FTransform3f(SampledLocalAnimPose[AnimBoneIndex]);
+					const FTransform3f RetargetTransform = Mapping.RetargetingTables.Transform[CompactBoneIndex];
+
+					if (SampledLocalAnimPoseBits[CompactBoneIndex])
+					{
+						MeshLocalTransform = FTransform3f(SampledLocalAnimPose[CompactBoneIndex]);
+					}
+					else
+					{
+						MeshLocalTransform = FTransform3f(AnimRefLocalPoses[AnimBoneIndex]);
+					}
 
-					if (Mapping.BoneRetargetingTables.bRequiresRefPoseRetarget && MeshBoneIndex != 0)
+					if (Mapping.RetargetingTables.bRequiresRefPoseRetarget && CompactBoneIndex != 0)
 					{
-						const FQuat4f& Q0 = Mapping.BoneRetargetingTables.RefPoseRetargetQuats[MeshBoneIndex].Get<0>();
-						const FQuat4f& Q1 = Mapping.BoneRetargetingTables.RefPoseRetargetQuats[MeshBoneIndex].Get<1>();
+						const FQuat4f& Q0 = Mapping.RetargetingTables.RefPoseRetargetQuats[CompactBoneIndex].Get<0>();
+						const FQuat4f& Q1 = Mapping.RetargetingTables.RefPoseRetargetQuats[CompactBoneIndex].Get<1>();
 
 						FQuat4f RetargetedRotation = Q0 * MeshLocalTransform.GetRotation() * Q1;
 						RetargetedRotation.Normalize();
 						MeshLocalTransform.SetRotation(RetargetedRotation);
  
-						if (Mapping.BoneRetargetingTables.Mode[MeshBoneIndex] != EBoneTranslationRetargetingMode::Skeleton)
+						if (Mapping.RetargetingTables.Mode[CompactBoneIndex] != EBoneTranslationRetargetingMode::Skeleton)
 						{
 							MeshLocalTransform.SetTranslation(Q0.RotateVector(MeshLocalTransform.GetTranslation()));
 						}
 					}
 
-					switch (Mapping.BoneRetargetingTables.Mode[MeshBoneIndex])
+					switch (Mapping.RetargetingTables.Mode[CompactBoneIndex])
 					{
 						case EBoneTranslationRetargetingMode::Animation:
 						{
@@ -1616,7 +1783,7 @@ void FAnimSequenceTransformProvider::ProvideCPUTransforms(FSkinningTransformProv
 						case EBoneTranslationRetargetingMode::Skeleton:
 						{
 							MeshLocalTransform.SetTranslation(RetargetTransform.GetTranslation());
-							if (MeshBoneIndex == 0)
+							if (CompactBoneIndex == 0)
 							{
 								MeshLocalTransform.SetRotation(RetargetTransform.GetRotation());
 							}
@@ -1647,15 +1814,15 @@ void FAnimSequenceTransformProvider::ProvideCPUTransforms(FSkinningTransformProv
 					MeshLocalTransform = FTransform3f(MeshLocalRefPoses[MeshBoneIndex]);
 				}
 
-				const int32 ParentIndex = MeshRefSkeleton.GetParentIndex(MeshBoneIndex);
+				const uint32 ParentIndex = BoneTables.CompactToParent[CompactBoneIndex];
 
-				if (SampledGlobalMeshPose.IsValidIndex(ParentIndex) && ParentIndex < MeshBoneIndex)
+				if (SampledGlobalMeshPose.IsValidIndex(ParentIndex) && ParentIndex < CompactBoneIndex)
 				{
-					SampledGlobalMeshPose[MeshBoneIndex] = MeshLocalTransform * SampledGlobalMeshPose[ParentIndex];
+					SampledGlobalMeshPose[CompactBoneIndex] = MeshLocalTransform * SampledGlobalMeshPose[ParentIndex];
 				}
 				else
 				{
-					SampledGlobalMeshPose[MeshBoneIndex] = MeshLocalTransform;
+					SampledGlobalMeshPose[CompactBoneIndex] = MeshLocalTransform;
 				}
 			}
 		};
@@ -1666,6 +1833,7 @@ void FAnimSequenceTransformProvider::ProvideCPUTransforms(FSkinningTransformProv
 		TBitArray<SceneRenderingBitArrayAllocator> SourceTransformMap;
 
 		const auto EvaluateAnimSequence = [&SourceTransforms, &SourceTransformMap, &SourceSampledGlobalMeshPose, &TargetSampledGlobalMeshPose, &DecompressPose] (
+			const FBoneTables& BoneTables,
 			const FInstancedSkinningSceneExtensionProxy* ExtensionProxy,
 			TArrayView<FCompressedBoneTransform> Transforms,
 			TArrayView<FCompressedBoneTransform> DuplicatePreviousTransforms,
@@ -1677,44 +1845,57 @@ void FAnimSequenceTransformProvider::ProvideCPUTransforms(FSkinningTransformProv
 			float TargetTimeOffset,
 			float BlendWeight)
 		{
-			DecompressPose(SourceMapping, SourceSequence, SourceTimeOffset, SourceSampledGlobalMeshPose);
+			DecompressPose(BoneTables, SourceMapping, SourceSequence, SourceTimeOffset, SourceSampledGlobalMeshPose);
 
 			if (BlendWeight > 0.0f)
 			{
-				DecompressPose(TargetMapping, TargetSequence, TargetTimeOffset, TargetSampledGlobalMeshPose);
+				DecompressPose(BoneTables, TargetMapping, TargetSequence, TargetTimeOffset, TargetSampledGlobalMeshPose);
 			}
 
-			for (uint32 MeshBoneIndex = 0; MeshBoneIndex < SourceMapping.NumMeshBones; ++MeshBoneIndex)
+			for (uint32 CompactBoneIndex = 0; CompactBoneIndex < TargetMapping.NumCompactBones; ++CompactBoneIndex)
 			{
 				if (BlendWeight > 0.0f)
 				{
-					SourceSampledGlobalMeshPose[MeshBoneIndex].BlendWith(TargetSampledGlobalMeshPose[MeshBoneIndex], BlendWeight);
-					SourceSampledGlobalMeshPose[MeshBoneIndex].NormalizeRotation();
+					SourceSampledGlobalMeshPose[CompactBoneIndex].BlendWith(TargetSampledGlobalMeshPose[CompactBoneIndex], BlendWeight);
+					SourceSampledGlobalMeshPose[CompactBoneIndex].NormalizeRotation();
 				}
 
-				SourceSampledGlobalMeshPose[MeshBoneIndex] = SourceMapping.InvGlobalRefPoseTransforms[MeshBoneIndex] * SourceSampledGlobalMeshPose[MeshBoneIndex];
+				SourceSampledGlobalMeshPose[CompactBoneIndex] = TargetMapping.InvGlobalRefPoseTransforms[CompactBoneIndex] * SourceSampledGlobalMeshPose[CompactBoneIndex];
 			}
 
-			if (ExtensionProxy->GetBoneTransformStorageMode() == EBoneTransformStorageMode::BoneMap)
+			if (!BoneTables.TransformToMesh.IsEmpty())
 			{
-				TConstArrayView<uint32> BoneHierarchy = ExtensionProxy->GetBoneHierarchy();
-				SourceTransformMap.Init(false, TargetMapping.NumMeshBones);
-				SourceTransforms.SetNumUninitialized(TargetMapping.NumMeshBones, EAllowShrinking::No);
+				SourceTransformMap.Init(false, TargetMapping.NumCompactBones);
+				SourceTransforms.SetNumUninitialized(TargetMapping.NumCompactBones, EAllowShrinking::No);
+				FCompressedBoneTransform IdentityTransform;
+				SetCompressedBoneTransformIdentity(IdentityTransform);
 				int32 DstTransformIndex = 0;
 
-				for (uint32 MeshBoneIndex : ExtensionProxy->GetBoneHierarchy())
+				for (uint32 MeshBoneIndex : BoneTables.TransformToMesh)
 				{
-					if (!SourceTransformMap[MeshBoneIndex])
+					const int32 CompactBoneIndex = BoneTables.MeshToCompact[MeshBoneIndex];
+					FCompressedBoneTransform* OutputTransform;
+
+					if (CompactBoneIndex != INDEX_NONE)
+					{
+						OutputTransform = &SourceTransforms[CompactBoneIndex];
+
+						if (!SourceTransformMap[CompactBoneIndex])
+						{
+							SourceTransformMap[CompactBoneIndex] = true;
+							StoreCompressedBoneTransform(OutputTransform, SourceSampledGlobalMeshPose[CompactBoneIndex].ToMatrixWithScale());
+						}
+					}
+					else
 					{
-						SourceTransformMap[MeshBoneIndex] = true;
-						StoreCompressedBoneTransform(SourceTransforms[MeshBoneIndex], SourceSampledGlobalMeshPose[MeshBoneIndex].ToMatrixWithScale());
+						OutputTransform = &IdentityTransform;
 					}
 
-					Transforms[DstTransformIndex] = SourceTransforms[MeshBoneIndex];
+					Transforms[DstTransformIndex] = *OutputTransform;
 
 					if (!DuplicatePreviousTransforms.IsEmpty())
 					{
-						DuplicatePreviousTransforms[DstTransformIndex] = SourceTransforms[MeshBoneIndex];
+						DuplicatePreviousTransforms[DstTransformIndex] = *OutputTransform;
 					}
 
 					DstTransformIndex++;
@@ -1724,11 +1905,23 @@ void FAnimSequenceTransformProvider::ProvideCPUTransforms(FSkinningTransformProv
 			{
 				for (uint32 MeshBoneIndex = 0; MeshBoneIndex < TargetMapping.NumMeshBones; ++MeshBoneIndex)
 				{
-					StoreCompressedBoneTransform(Transforms[MeshBoneIndex], SourceSampledGlobalMeshPose[MeshBoneIndex].ToMatrixWithScale());
+					const int32 CompactBoneIndex = BoneTables.MeshToCompact[MeshBoneIndex];
+					FCompressedBoneTransform OutputTransform;
+
+					if (CompactBoneIndex != INDEX_NONE)
+					{
+						StoreCompressedBoneTransform(OutputTransform, SourceSampledGlobalMeshPose[CompactBoneIndex].ToMatrixWithScale());
+					}
+					else
+					{
+						SetCompressedBoneTransformIdentity(OutputTransform);
+					}
+
+					Transforms[MeshBoneIndex] = OutputTransform;
 
 					if (!DuplicatePreviousTransforms.IsEmpty())
 					{
-						DuplicatePreviousTransforms[MeshBoneIndex] = Transforms[MeshBoneIndex];
+						DuplicatePreviousTransforms[MeshBoneIndex] = OutputTransform;
 					}
 				}
 			}
@@ -1748,15 +1941,17 @@ void FAnimSequenceTransformProvider::ProvideCPUTransforms(FSkinningTransformProv
 			FAnimSequenceTrackRenderPool& Tracks = ProxyUserData.Tracks;
 			ProxyUserData.TrackHistories.SetNum(Tracks.GetNumTracks());
 
-			const uint32 MaxTransformCount      = ExtensionProxy->GetMaxBoneTransformCount();
+			const uint32 LODIndex     = ExtensionProxy->GetLOD();
+			FProxyUserData::FLOD& LOD = ProxyUserData.LODs[LODIndex];
+
 			const uint32 UniqueAnimationCount   = ExtensionProxy->GetUniqueAnimationCount();
-			const uint32 MaxTotalTransformCount = MaxTransformCount * 2u; // Current and Previous
+			const uint32 MaxTotalTransformCount = LOD.BoneTables.NumTransforms * 2u; // Current and Previous
 			const uint32 CurrentTransformIndex  = Indirection.CurrentTransformOffset  / sizeof(FCompressedBoneTransform);
 			const uint32 PreviousTransformIndex = Indirection.PreviousTransformOffset / sizeof(FCompressedBoneTransform);
 			const uint32 CurrentTransformSlot   = CurrentTransformIndex < PreviousTransformIndex ? 0 : 1;
 
-			SourceSampledGlobalMeshPose.SetNumUninitialized(ProxyUserData.NumMeshBones, EAllowShrinking::No);
-			TargetSampledGlobalMeshPose.SetNumUninitialized(ProxyUserData.NumMeshBones, EAllowShrinking::No);
+			SourceSampledGlobalMeshPose.SetNumUninitialized(LOD.BoneTables.NumCompactBones, EAllowShrinking::No);
+			TargetSampledGlobalMeshPose.SetNumUninitialized(LOD.BoneTables.NumCompactBones, EAllowShrinking::No);
 	
 			FMemMark Mark(FMemStack::Get());
 
@@ -1767,8 +1962,8 @@ void FAnimSequenceTransformProvider::ProvideCPUTransforms(FSkinningTransformProv
 				const FAnimSequenceTrackPackedData& Source = Data.GetSource();
 				const FAnimSequenceTrackPackedData& Target = Data.GetTarget();
 
-				const FMapping& SourceMapping = ProxyUserData.GetMapping(Source);
-				const FMapping& TargetMapping = ProxyUserData.GetMapping(Target);
+				const FMapping& SourceMapping = ProxyUserData.GetMapping(Source, LODIndex);
+				const FMapping& TargetMapping = ProxyUserData.GetMapping(Target, LODIndex);
 				FSequence& SourceSequence = ProxyUserData.GetSequence(Source);
 				FSequence& TargetSequence = ProxyUserData.GetSequence(Target);
 				
@@ -1791,13 +1986,14 @@ void FAnimSequenceTransformProvider::ProvideCPUTransforms(FSkinningTransformProv
 						? EPreviousBoneTransformUpdateMode::DuplicateCurrentToPrevious
 						: EPreviousBoneTransformUpdateMode::UpdatePrevious;
 
-					DstTransformsPrevious = ScatterUploader.Add_GetRef<FCompressedBoneTransform>(PreviousTransformIndex + TrackIndex * MaxTotalTransformCount, MaxTransformCount);
+					DstTransformsPrevious = ScatterUploader.Add_GetRef<FCompressedBoneTransform>(PreviousTransformIndex + TrackIndex * MaxTotalTransformCount, LOD.BoneTables.NumTransforms);
 				}
 
 				if (UpdateResult.bNeedsCurrent)
 				{
-					TArrayView<FCompressedBoneTransform> DstTransforms = ScatterUploader.Add_GetRef<FCompressedBoneTransform>(CurrentTransformIndex + TrackIndex * MaxTotalTransformCount, MaxTransformCount);
+					TArrayView<FCompressedBoneTransform> DstTransforms = ScatterUploader.Add_GetRef<FCompressedBoneTransform>(CurrentTransformIndex + TrackIndex * MaxTotalTransformCount, LOD.BoneTables.NumTransforms);
 					EvaluateAnimSequence(
+						LOD.BoneTables,
 						ExtensionProxy,
 						DstTransforms,
 						UpdateMode == EPreviousBoneTransformUpdateMode::DuplicateCurrentToPrevious ? DstTransformsPrevious : TArrayView<FCompressedBoneTransform>{},
@@ -1812,6 +2008,7 @@ void FAnimSequenceTransformProvider::ProvideCPUTransforms(FSkinningTransformProv
 				if (UpdateMode == EPreviousBoneTransformUpdateMode::UpdatePrevious)
 				{
 					EvaluateAnimSequence(
+						LOD.BoneTables,
 						ExtensionProxy,
 						DstTransformsPrevious,
 						{},
@@ -1832,4 +2029,4 @@ void FAnimSequenceTransformProvider::ProvideCPUTransforms(FSkinningTransformProv
 		INC_DWORD_STAT_BY(STAT_AnimSequenceTransformProvider_CurrentTransformUpdates, NumCurrentTransformUpdates);
 		INC_DWORD_STAT_BY(STAT_AnimSequenceTransformProvider_PreviousTransformUpdates, NumPreviousTransformUpdates);
 	});
-}
+}
\ No newline at end of file
diff --git a/Engine/Source/Runtime/Renderer/Private/Skinning/AnimSequenceTransformProvider.h b/Engine/Source/Runtime/Renderer/Private/Skinning/AnimSequenceTransformProvider.h
--- a/Engine/Source/Runtime/Renderer/Private/Skinning/AnimSequenceTransformProvider.h
+++ b/Engine/Source/Runtime/Renderer/Private/Skinning/AnimSequenceTransformProvider.h
@@ -33,16 +33,39 @@ private:
 	void ProvideGPUTransforms(FSkinningTransformProvider::FProviderContext& Context);
 	void ProvideCPUTransforms(FSkinningTransformProvider::FProviderContext& Context);
 
-	struct FBoneRetargetingTables
+	struct FBoneTables
+	{
+		TArray<int32> TransformToMesh;
+		TArray<int32> MeshToCompact;
+		TArray<int32> CompactToMesh;
+		TArray<int32> CompactToParent;
+		TArray<int32> CompactToTargetSkeleton;
+		uint32 NumTransforms = 0;
+		uint32 NumMeshBones = 0;
+		uint32 NumCompactBones = 0;
+	};
+
+	void BuildBoneTables(
+		FBoneTables& BoneTables,
+		EBoneTransformStorageMode TransformStorageMode,
+		const USkinnedAsset* SkinnedMesh,
+		const FSkeletalMeshLODRenderData& LODRenderData);
+
+	struct FRetargetingTables
 	{
 		TArray<EBoneTranslationRetargetingMode::Type> Mode;
 		TArray<FTransform3f> Transform;
-		TArray<int32> AnimBoneIndex;
 		TArray<TTuple<FQuat4f, FQuat4f>> RefPoseRetargetQuats;
+		TArray<int32> CompactToAnim;
+		TArray<int32> AnimToCompact;
 		bool bRequiresRefPoseRetarget = false;
 	};
 
-	void BuildBoneRetargetingTables(FBoneRetargetingTables& BoneRetargetingTables, USkeleton* AnimSkeleton, const USkinnedAsset* MeshAsset);
+	void BuildRetargetingTables(
+		FRetargetingTables& RetargetingTables,
+		USkeleton* AnimSkeleton,
+		const USkinnedAsset* MeshAsset,
+		const FBoneTables& BoneTables);
 
 	struct FLeafBonePathData
 	{
@@ -50,7 +73,7 @@ private:
 		uint32 NumLeafBones;
 	};
 
-	void BuildLeafBonePathData(FLeafBonePathData& OutLeafBonePathData, const FReferenceSkeleton& RefSkeleton);
+	void BuildLeafBonePathData(FLeafBonePathData& OutLeafBonePathData, const FBoneTables& RemappingTables);
 
 	struct FBufferAllocation
 	{
@@ -103,22 +126,24 @@ private:
 	{
 		FMappingRegistryKey() = default;
 
-		FMappingRegistryKey(const USkeleton* InSourceSkeleton, const USkinnedAsset* InSkinnedAsset)
+		FMappingRegistryKey(const USkeleton* InSourceSkeleton, const USkinnedAsset* InSkinnedAsset, uint32 InLODIndex)
 			: SourceSkeleton(InSourceSkeleton)
 			, SkinnedAsset(InSkinnedAsset)
+			, LODIndex(InLODIndex)
 		{}
 
 		const USkeleton* SourceSkeleton = nullptr;
 		const USkinnedAsset* SkinnedAsset = nullptr;
+		uint32 LODIndex = 0;
 
 		bool operator==(const FMappingRegistryKey& Other) const
 		{
-			return SourceSkeleton == Other.SourceSkeleton && SkinnedAsset == Other.SkinnedAsset;
+			return SourceSkeleton == Other.SourceSkeleton && SkinnedAsset == Other.SkinnedAsset && LODIndex == Other.LODIndex;
 		}
 
 		friend uint32 GetTypeHash(const FMappingRegistryKey& Key)
 		{
-			return PointerHash(Key.SourceSkeleton, PointerHash(Key.SkinnedAsset));
+			return PointerHash(Key.SourceSkeleton, PointerHash(Key.SkinnedAsset, Key.LODIndex));
 		}
 	};
 
@@ -164,17 +189,18 @@ private:
 
 	struct FMapping
 	{
-		FMappingRegistryKey RegistryKey;
 		FBufferAllocation MappingTransformBufferAllocation;
 		FBufferAllocation RetargetingDataBufferAllocation;
 		TArray<FTransform3f> InvGlobalRefPoseTransforms;
 		uint32 RetargetingTransformBufferOffset = 0;
 		uint32 InvGlobalRefPoseTransformBufferOffset = 0;
-		FBoneRetargetingTables BoneRetargetingTables;
+		FRetargetingTables RetargetingTables;
 		const USkeleton* SourceSkeleton = nullptr;
 		const USkinnedAsset* SkinnedAsset = nullptr;
-		uint32 RefCount = 0;
+		uint32 LODIndex = 0;
 		uint32 NumMeshBones = 0;
+		uint32 NumCompactBones = 0;
+		uint32 RefCount = 0;
 	};
 
 	class FTrackHistory
@@ -228,25 +254,39 @@ private:
 	public:
 		FProxyUserData(FAnimSequenceTransformProviderProxy* Proxy);
 
-		const FMapping& GetMapping(const FAnimSequenceTrackPackedData& Data) const
+		FSequence& GetSequence(const FAnimSequenceTrackPackedData& Data) const
 		{
-			const int32 SequenceIndex = SequenceToMappingTable.IsValidIndex(Data.GetSequenceIndex()) ? Data.GetSequenceIndex() : 0;
-			return *SequenceToMappingTable[SequenceIndex];
+			return *Sequences[Sequences.IsValidIndex(Data.GetSequenceIndex()) ? Data.GetSequenceIndex() : 0];
 		}
 
-		FSequence& GetSequence(const FAnimSequenceTrackPackedData& Data) const
+		const FMapping& GetMapping(const FAnimSequenceTrackPackedData& Data, uint32 LODIndex) const
 		{
-			return *Sequences[Sequences.IsValidIndex(Data.GetSequenceIndex()) ? Data.GetSequenceIndex() : 0];
+			const int32 SequenceIndex = SequenceToMappingTable.IsValidIndex(Data.GetSequenceIndex()) ? Data.GetSequenceIndex() : 0;
+			return *SequenceToMappingTable[SequenceIndex * LODs.Num() + LODIndex];
 		}
 
+		struct FLOD
+		{
+			FBoneTables BoneTables;
+			FLeafBonePathData LeafBonePathData;
+			FBufferAllocation LeafBonePathBufferAllocation;
+			int32 BoneMapBufferOffset = 0;
+			TArray<int32> SectionBoneMaps;
+
+			bool IsValid() const
+			{
+				return BoneTables.NumCompactBones != 0;
+			}
+		};
+
 		FAnimSequenceTrackRenderPool& Tracks;
 		TArray<FTrackHistory> TrackHistories;
 		const FInstancedSkinningSceneExtensionProxy& SceneExtensionProxy;
 		TArray<FSequence*> Sequences;
-		TArray<FMapping*> SequenceToMappingTable;
 		TArray<FMapping*> Mappings;
-		FLeafBonePathData LeafBonePathData;
-		FBufferAllocation LeafBonePathBufferAllocation;
+		TArray<FMapping*> SequenceToMappingTable;
+		TArray<FLOD> LODs;
+		uint32 MinLODLevel = 0;
 		uint32 NumMeshBones = 0;
 		const EBoneTransformStorageMode TransformStorageMode;
 	};
