Custom Render Pass:  Decal bug fixes, allowing decals to be re-enabled.
* CRP views need to run decal visibility separately (not use the main view's visibility).  Decal visibility class modified to accept "AllViews" (which includes CRPs), and decision on whether decal visibility needs to run is decided per view (as CRPs have their own view family, which may have different decal related settings).  Later code that uses decal visibility results needs to take into account that the per-view data is no longer one-to-one, searching for the matching visibility data and offseting to that data in the array.
* DBuffers need to be cleared between CRP and main view renders.  Detecting when to clear used the "bProduced" flag -- code was switched to use an external flag, which can be reset to allow clearing to run again.
* Mobile shading renderer CRP DBuffer support added.
* Water info rendering needs to disable decals.
* Fixed console specific crash that required CRP decals to be disabled (using GBuffer texture references from SceneTextures.UniformBuffer which weren't initialized, rather than the ones in the base SceneTextures).
* Fixed additional console specific crash for decal rendering code paths using a downsampled DBufferMask (as opposed to per-pixel), where the DBuffer mask is incorrectly bound as an MRT due to a CRP having previously initialized it.  Instead of using the presence of the DBuffer mask texture in the struct to decide whether to bind the texture, check the original mask technique setting.

#rnx

#Synced-CL 50224996

diff --git a/Engine/Plugins/Experimental/Water/Source/Runtime/Private/WaterInfoRendering.cpp b/Engine/Plugins/Experimental/Water/Source/Runtime/Private/WaterInfoRendering.cpp
--- a/Engine/Plugins/Experimental/Water/Source/Runtime/Private/WaterInfoRendering.cpp
+++ b/Engine/Plugins/Experimental/Water/Source/Runtime/Private/WaterInfoRendering.cpp
@@ -575,6 +575,7 @@ void UpdateWaterInfoRendering_CustomRenderPass(
 	PassInput.ViewRotationMatrix.RemoveScaling();
 	PassInput.ProjectionMatrix = BuildOrthoMatrix(ZoneExtent.X, ZoneExtent.Y);
 	PassInput.ViewActor = Context.ZoneToRender;
+	PassInput.EngineShowFlags.Decals = 0;
 
 	TSet<FPrimitiveComponentId> ComponentsToRenderInDepthPass;
 	if (Context.GroundPrimitiveComponents.Num() > 0)
diff --git a/Engine/Source/Runtime/Renderer/Private/CompositionLighting/CompositionLighting.cpp b/Engine/Source/Runtime/Renderer/Private/CompositionLighting/CompositionLighting.cpp
--- a/Engine/Source/Runtime/Renderer/Private/CompositionLighting/CompositionLighting.cpp
+++ b/Engine/Source/Runtime/Renderer/Private/CompositionLighting/CompositionLighting.cpp
@@ -492,9 +492,14 @@ FCompositionLighting::FCompositionLighting(FDecalVisibilityTaskData* InDecalVisi
 
 FCompositionLighting::~FCompositionLighting()
 {
-	if (DecalVisibility)
+	const int32 DecalViewOffset = DecalVisibility ? DecalVisibility->GetViewOffset(Views) : INDEX_NONE;
+
+	if (DecalViewOffset != INDEX_NONE)
 	{
-		DecalVisibility->Finish();
+		for (int32 ViewIndex = 0; ViewIndex < Views.Num(); ++ViewIndex)
+		{
+			DecalVisibility->Finish(DecalViewOffset, ViewIndex);
+		}
 	}
 }
 
@@ -549,6 +554,8 @@ void FCompositionLighting::ProcessBeforeBasePass(FRDGBuilder& GraphBuilder, FDBu
 
 	RDG_EVENT_SCOPE_STAT(GraphBuilder, CompositionBeforeBasePass, "CompositionBeforeBasePass");
 
+	const int32 DecalViewOffset = DecalVisibility ? DecalVisibility->GetViewOffset(Views) : INDEX_NONE;
+
 	for (int32 ViewIndex = 0; ViewIndex < Views.Num(); ++ViewIndex)
 	{
 		FViewInfo& View = Views[ViewIndex];
@@ -556,7 +563,7 @@ void FCompositionLighting::ProcessBeforeBasePass(FRDGBuilder& GraphBuilder, FDBu
 
 		const bool bRunDbufferBeforeBasePass = !Substrate::IsDBufferPassEnabled(View.GetShaderPlatform());
 		const bool bEnableSSAO = ViewConfig.SSAOLocation == ESSAOLocation::BeforeBasePass;
-		const bool bEnableDecals = bRunDbufferBeforeBasePass && DecalVisibility && DecalVisibility->HasStage(ViewIndex, EDecalRenderStage::BeforeBasePass);
+		const bool bEnableDecals = bRunDbufferBeforeBasePass && DecalViewOffset != INDEX_NONE && DecalVisibility->HasStage(DecalViewOffset, ViewIndex, EDecalRenderStage::BeforeBasePass);
 
 		if (!bEnableDecals && !bEnableSSAO)
 		{
@@ -572,7 +579,7 @@ void FCompositionLighting::ProcessBeforeBasePass(FRDGBuilder& GraphBuilder, FDBu
 		if (bEnableDecals)
 		{
 			FDeferredDecalPassTextures DecalPassTextures = GetDeferredDecalPassTextures(GraphBuilder, View, SceneTextures, &DBufferTextures, EDecalRenderStage::BeforeBasePass);
-			TConstArrayView<const FVisibleDecal*> SortedDecals = DecalVisibility->FinishRelevantDecals(ViewIndex, EDecalRenderStage::BeforeBasePass);
+			TConstArrayView<const FVisibleDecal*> SortedDecals = DecalVisibility->FinishRelevantDecals(DecalViewOffset, ViewIndex, EDecalRenderStage::BeforeBasePass);
 			AddDeferredDecalPass(GraphBuilder, View, SortedDecals, DecalPassTextures, InstanceCullingManager, EDecalRenderStage::BeforeBasePass);
 		}
 
@@ -601,6 +608,8 @@ void FCompositionLighting::ProcessAfterBasePass(FRDGBuilder& GraphBuilder, FInst
 
 	RDG_EVENT_SCOPE_STAT(GraphBuilder, CompositionPreLighting, "LightCompositionTasks_PreLighting");
 
+	const int32 DecalViewOffset = DecalVisibility ? DecalVisibility->GetViewOffset(Views) : INDEX_NONE;
+
 	for (int32 ViewIndex = 0; ViewIndex < Views.Num(); ++ViewIndex)
 	{
 		FViewInfo& View = Views[ViewIndex];
@@ -613,12 +622,12 @@ void FCompositionLighting::ProcessAfterBasePass(FRDGBuilder& GraphBuilder, FInst
 
 		const auto AddDeferredDecalPassLambda = [&] (const FDeferredDecalPassTextures& DecalPassTextures, EDecalRenderStage Stage)
 		{
-			AddDeferredDecalPass(GraphBuilder, View, DecalVisibility->FinishRelevantDecals(ViewIndex, Stage), DecalPassTextures, InstanceCullingManager, Stage);
+			AddDeferredDecalPass(GraphBuilder, View, DecalVisibility->FinishRelevantDecals(DecalViewOffset, ViewIndex, Stage), DecalPassTextures, InstanceCullingManager, Stage);
 		};
 
-		if (DecalVisibility)
+		if (DecalViewOffset != INDEX_NONE)
 		{
-			if (DecalVisibility->HasStage(ViewIndex, EDecalRenderStage::BeforeLighting) && Mode != EProcessAfterBasePassMode::SkipBeforeLightingDecals)
+			if (DecalVisibility->HasStage(DecalViewOffset, ViewIndex, EDecalRenderStage::BeforeLighting) && Mode != EProcessAfterBasePassMode::SkipBeforeLightingDecals)
 			{
 				// We can disable this pass if using DBuffer decals
 				// Decals are before AmbientOcclusion so the decal can output a normal that AO is affected by
@@ -626,7 +635,7 @@ void FCompositionLighting::ProcessAfterBasePass(FRDGBuilder& GraphBuilder, FInst
 				AddDeferredDecalPassLambda(DecalPassTextures, EDecalRenderStage::BeforeLighting);
 			}
 
-			if (DecalVisibility->HasStage(ViewIndex, EDecalRenderStage::Emissive) && Mode != EProcessAfterBasePassMode::OnlyBeforeLightingDecals)
+			if (DecalVisibility->HasStage(DecalViewOffset, ViewIndex, EDecalRenderStage::Emissive) && Mode != EProcessAfterBasePassMode::OnlyBeforeLightingDecals)
 			{
 				// DBuffer decals with emissive component
 				FDeferredDecalPassTextures DecalPassTextures = GetDeferredDecalPassTextures(GraphBuilder, View, SceneTextures, nullptr, EDecalRenderStage::Emissive);
@@ -655,7 +664,7 @@ void FCompositionLighting::ProcessAfterBasePass(FRDGBuilder& GraphBuilder, FInst
 					AmbientOcclusion = AddPostProcessingGTAOPostAsync(GraphBuilder, View, Parameters, GTAOHorizons, FinalTarget);
 
 					ensureMsgf(
-						!DecalVisibility || !DecalVisibility->HasStage(ViewIndex, EDecalRenderStage::AmbientOcclusion),
+						DecalViewOffset == INDEX_NONE || !DecalVisibility->HasStage(DecalViewOffset, ViewIndex, EDecalRenderStage::AmbientOcclusion),
 						TEXT("Ambient occlusion decals are not supported with Async compute SSAO."));
 				}
 				else
@@ -671,7 +680,7 @@ void FCompositionLighting::ProcessAfterBasePass(FRDGBuilder& GraphBuilder, FInst
 						AmbientOcclusion = AddPostProcessingAmbientOcclusion(GraphBuilder, View, Parameters, FinalTarget);
 					}
 
-					if (DecalVisibility)
+					if (DecalViewOffset != INDEX_NONE)
 					{
 						FDeferredDecalPassTextures DecalPassTextures = GetDeferredDecalPassTextures(GraphBuilder, View, SceneTextures, nullptr, EDecalRenderStage::AmbientOcclusion);
 						DecalPassTextures.ScreenSpaceAO = AmbientOcclusion.Texture;
diff --git a/Engine/Source/Runtime/Renderer/Private/CompositionLighting/PostProcessDeferredDecals.cpp b/Engine/Source/Runtime/Renderer/Private/CompositionLighting/PostProcessDeferredDecals.cpp
--- a/Engine/Source/Runtime/Renderer/Private/CompositionLighting/PostProcessDeferredDecals.cpp
+++ b/Engine/Source/Runtime/Renderer/Private/CompositionLighting/PostProcessDeferredDecals.cpp
@@ -18,6 +18,7 @@
 #include "RenderCore.h"
 #include "VariableRateShadingImageManager.h"
 #include "PSOPrecacheValidation.h"
+#include "Rendering/CustomRenderPass.h"
 
 static TAutoConsoleVariable<float> CVarStencilSizeThreshold(
 	TEXT("r.Decal.StencilSizeThreshold"),
@@ -55,16 +56,16 @@ static TAutoConsoleVariable<bool> CVarDBufferDecalNormalReprojectionEnabled(
 	TEXT("in depth prepass is enabled as well (r.VelocityOutputPass=0). Otherwise the fallback is the normal extracted from the depth buffer."),
 	ECVF_RenderThreadSafe);
 
-bool AreDecalsEnabled(const FSceneViewFamily& ViewFamily)
+bool AreDecalsEnabled(const FViewInfo& View)
 {
-	return ViewFamily.EngineShowFlags.Decals && !ViewFamily.EngineShowFlags.VisualizeLightCulling;
+	const FSceneViewFamily& ViewFamily = *View.Family;
+	return ViewFamily.EngineShowFlags.Decals && !ViewFamily.EngineShowFlags.VisualizeLightCulling && !ViewFamily.EngineShowFlags.HitProxies &&
+		(!View.CustomRenderPass || View.CustomRenderPass->GetRenderMode() == FCustomRenderPassBase::ERenderMode::DepthAndBasePass);
 }
 
 bool IsDBufferEnabled(const FSceneViewFamily& ViewFamily, EShaderPlatform ShaderPlatform)
 {
-	return IsUsingDBuffers(ShaderPlatform)
-		&& AreDecalsEnabled(ViewFamily)
-		&& !ViewFamily.EngineShowFlags.ShaderComplexity;
+	return IsUsingDBuffers(ShaderPlatform) && !ViewFamily.EngineShowFlags.ShaderComplexity;
 }
 
 IMPLEMENT_STATIC_UNIFORM_BUFFER_STRUCT(FDecalPassUniformParameters, "DecalPass", SceneTextures);
@@ -114,10 +115,10 @@ FDeferredDecalPassTextures GetDeferredDecalPassTextures(
 	// Mobile deferred renderer does not use dbuffer 
 	if (!bIsMobile)
 	{
-		PassTextures.GBufferA = (*SceneTextures.UniformBuffer)->GBufferATexture;
-		PassTextures.GBufferB = (*SceneTextures.UniformBuffer)->GBufferBTexture;
-		PassTextures.GBufferC = (*SceneTextures.UniformBuffer)->GBufferCTexture;
-		PassTextures.GBufferE = (*SceneTextures.UniformBuffer)->GBufferETexture;
+		PassTextures.GBufferA = SceneTextures.GBufferA;
+		PassTextures.GBufferB = SceneTextures.GBufferB;
+		PassTextures.GBufferC = SceneTextures.GBufferC;
+		PassTextures.GBufferE = SceneTextures.GBufferE;
 	}
 
 	PassTextures.DBufferTextures = DBufferTextures;
@@ -250,20 +251,34 @@ void GetDeferredDecalPassParameters(
 	{
 		check(Textures.DBufferTextures);
 
-		const FDBufferTextures& DBufferTextures = *Textures.DBufferTextures;
+		FDBufferTextures& DBufferTextures = *Textures.DBufferTextures;
 
-		const bool bDBufferAProduced = DBufferTextures.DBufferA ? DBufferTextures.DBufferA->HasBeenProduced() : false;
-		const bool bDBufferTexArrayAProduced = DBufferTextures.DBufferATexArray ? DBufferTextures.DBufferATexArray->HasBeenProduced() : false;
 		const bool bUseTextureArrays = View.bIsMobileMultiViewEnabled || View.Aspects.IsMobileMultiViewEnabled();
-		const ERenderTargetLoadAction LoadAction = (bUseTextureArrays ? bDBufferTexArrayAProduced : bDBufferAProduced)
+		const ERenderTargetLoadAction LoadAction = (bUseTextureArrays ? DBufferTextures.bDBufferTexArrayProduced : DBufferTextures.bDBufferProduced)
 			? ERenderTargetLoadAction::ELoad
 			: ERenderTargetLoadAction::EClear;
 
+		if (bUseTextureArrays)
+		{
+			if (DBufferTextures.DBufferATexArray)
+			{
+				DBufferTextures.bDBufferTexArrayProduced = true;
+			}
+		}
+		else
+		{
+			if (DBufferTextures.DBufferA)
+			{
+				DBufferTextures.bDBufferProduced = true;
+			}
+		}
+
 		AddColorTarget(DBufferTextures.DBufferA, LoadAction, DBufferTextures.DBufferATexArray, bUseTextureArrays);
 		AddColorTarget(DBufferTextures.DBufferB, LoadAction, DBufferTextures.DBufferBTexArray, bUseTextureArrays);
 		AddColorTarget(DBufferTextures.DBufferC, LoadAction, DBufferTextures.DBufferCTexArray, bUseTextureArrays);
 
-		if (DBufferTextures.DBufferMask)
+		const EDecalDBufferMaskTechnique DBufferMaskTechnique = GetDBufferMaskTechnique(View.Family->GetShaderPlatform());
+		if (DBufferMaskTechnique == EDecalDBufferMaskTechnique::PerPixel)
 		{
 			AddColorTarget(DBufferTextures.DBufferMask, LoadAction);
 		}
diff --git a/Engine/Source/Runtime/Renderer/Private/CompositionLighting/PostProcessDeferredDecals.h b/Engine/Source/Runtime/Renderer/Private/CompositionLighting/PostProcessDeferredDecals.h
--- a/Engine/Source/Runtime/Renderer/Private/CompositionLighting/PostProcessDeferredDecals.h
+++ b/Engine/Source/Runtime/Renderer/Private/CompositionLighting/PostProcessDeferredDecals.h
@@ -17,7 +17,7 @@ class FViewInfo;
 
 DECLARE_UNIFORM_BUFFER_STRUCT(FSceneUniformParameters, RENDERER_API)
 
-bool AreDecalsEnabled(const FSceneViewFamily& ViewFamily);
+bool AreDecalsEnabled(const FViewInfo& View);
 bool IsDBufferEnabled(const FSceneViewFamily& ViewFamily, EShaderPlatform ShaderPlatform);
 
 BEGIN_GLOBAL_SHADER_PARAMETER_STRUCT(FDecalPassUniformParameters, )
diff --git a/Engine/Source/Runtime/Renderer/Private/DBufferTextures.cpp b/Engine/Source/Runtime/Renderer/Private/DBufferTextures.cpp
--- a/Engine/Source/Runtime/Renderer/Private/DBufferTextures.cpp
+++ b/Engine/Source/Runtime/Renderer/Private/DBufferTextures.cpp
@@ -12,7 +12,7 @@
 bool FDBufferTextures::IsValid() const
 {
 	check((!DBufferA || (DBufferB && DBufferC)) || (!DBufferATexArray || (DBufferBTexArray && DBufferCTexArray)));
-	return HasBeenProduced(DBufferA) || HasBeenProduced(DBufferATexArray);
+	return bDBufferProduced || bDBufferTexArrayProduced;
 }
 
 EDecalDBufferMaskTechnique GetDBufferMaskTechnique(EShaderPlatform ShaderPlatform)
diff --git a/Engine/Source/Runtime/Renderer/Private/DBufferTextures.h b/Engine/Source/Runtime/Renderer/Private/DBufferTextures.h
--- a/Engine/Source/Runtime/Renderer/Private/DBufferTextures.h
+++ b/Engine/Source/Runtime/Renderer/Private/DBufferTextures.h
@@ -42,6 +42,12 @@ struct FDBufferTextures
 	FRDGTextureRef DBufferBTexArray = nullptr;
 	FRDGTextureRef DBufferCTexArray = nullptr;
 	FRDGTextureRef DBufferMask = nullptr;
+
+	// We use an external flag to test whether the render targets have been produced, rather than "HasBeenProduced()", for the purpose of deciding whether
+	// the render targets need to be cleared.  This is necessary for Custom Render Passes, where the render targets need to be cleared again when the main
+	// view renders.  Without this, decals from Custom Render Passes will still be in the buffer and incorrectly render on top of the main view.
+	bool bDBufferProduced = false;
+	bool bDBufferTexArrayProduced = false;
 };
 
 FDBufferTexturesDesc GetDBufferTexturesDesc(FIntPoint Extent, EShaderPlatform ShaderPlatform);
diff --git a/Engine/Source/Runtime/Renderer/Private/DecalRenderingShared.cpp b/Engine/Source/Runtime/Renderer/Private/DecalRenderingShared.cpp
--- a/Engine/Source/Runtime/Renderer/Private/DecalRenderingShared.cpp
+++ b/Engine/Source/Runtime/Renderer/Private/DecalRenderingShared.cpp
@@ -46,28 +46,29 @@ FVisibleDecal::FVisibleDecal(const FDeferredDecalProxy& InDecalProxy, float InCo
 	BlendDesc = DecalRendering::ComputeDecalBlendDesc(ShaderPlatform, MaterialResource);
 }
 
-FDecalVisibilityTaskData* FDecalVisibilityTaskData::Launch(FRDGBuilder& GraphBuilder, const FScene& Scene, TConstArrayView<FViewInfo> Views)
+FDecalVisibilityTaskData* FDecalVisibilityTaskData::Launch(FRDGBuilder& GraphBuilder, const FScene& Scene, TConstArrayView<FViewInfo*> Views)
 {
-	const FSceneViewFamily& ViewFamily = *Views[0].Family;
-	if (AreDecalsEnabled(ViewFamily) && !HasRayTracedOverlay(ViewFamily))
+	// Allocate decal visibility task data if any view needs decals
+	for (FViewInfo* View : Views)
 	{
-		const bool bDBufferEnabled = ::IsDBufferEnabled(ViewFamily, ViewFamily.GetShaderPlatform());
-		const bool bGBufferEnabled = IsUsingGBuffers(ViewFamily.GetShaderPlatform());
-
-		return GraphBuilder.AllocObject<FDecalVisibilityTaskData>(Scene, Views, bDBufferEnabled, bGBufferEnabled);
+		if (AreDecalsEnabled(*View) && !HasRayTracedOverlay(*View->Family))
+		{
+			return GraphBuilder.AllocObject<FDecalVisibilityTaskData>(Scene, Views);
+		}
 	}
 	return nullptr;
 }
 
-FDecalVisibilityTaskData::FDecalVisibilityTaskData(const FScene& Scene, TConstArrayView<FViewInfo> Views, bool bInDBufferEnabled, bool bInGBufferEnabled)
-	: bDBufferEnabled(bInDBufferEnabled)
-	, bGBufferEnabled(bInGBufferEnabled)
+FDecalVisibilityTaskData::FDecalVisibilityTaskData(const FScene& Scene, TConstArrayView<FViewInfo*> Views)
 {
 	ViewPackets.Reserve(Views.Num());
 
-	for (const FViewInfo& View : Views)
+	for (const FViewInfo* View : Views)
 	{
-		ViewPackets.Emplace(*this, Scene, View);
+		if (AreDecalsEnabled(*View) && !HasRayTracedOverlay(*View->Family))
+		{
+			ViewPackets.Emplace(*this, Scene, *View);
+		}
 	}
 }
 
@@ -97,14 +98,18 @@ FDecalVisibilityViewPacket::FDecalVisibilityViewPacket(const FDecalVisibilityTas
 		AllTasksEvent.AddPrerequisites(RelevantDecals.Task);
 	};
 
+	EShaderPlatform ShaderPlatform = View.Family->GetShaderPlatform();
+	const bool bDBufferEnabled = ::IsDBufferEnabled(*View.Family, ShaderPlatform);
+	const bool bGBufferEnabled = IsUsingGBuffers(ShaderPlatform);
+
 	// DBuffer Passes
-	if (TaskData.bDBufferEnabled)
+	if (bDBufferEnabled)
 	{
 		LaunchRelevantDecalTask(EDecalRenderStage::BeforeBasePass);
 		LaunchRelevantDecalTask(EDecalRenderStage::Emissive);
 	}
 	// GBuffer Passes
-	else if (TaskData.bGBufferEnabled)
+	else if (bGBufferEnabled)
 	{
 		LaunchRelevantDecalTask(EDecalRenderStage::BeforeLighting);
 	}
diff --git a/Engine/Source/Runtime/Renderer/Private/DecalRenderingShared.h b/Engine/Source/Runtime/Renderer/Private/DecalRenderingShared.h
--- a/Engine/Source/Runtime/Renderer/Private/DecalRenderingShared.h
+++ b/Engine/Source/Runtime/Renderer/Private/DecalRenderingShared.h
@@ -89,51 +89,52 @@ private:
 	TMap<EDecalRenderStage, FRelevantDecals> RelevantDecalsMap;
 	UE::Tasks::FTaskEvent AllTasksEvent{ UE_SOURCE_LOCATION };
 	bool bFinishCalled = false;
+
+	friend class FDecalVisibilityTaskData;
 };
 
 class FDecalVisibilityTaskData
 {
 public:
-	static FDecalVisibilityTaskData* Launch(FRDGBuilder& GraphBuilder, const FScene& Scene, TConstArrayView<FViewInfo> Views);
-
-	TConstArrayView<FVisibleDecal> FinishVisibleDecals(int32 ViewIndex)
-	{
-		return ViewPackets[ViewIndex].FinishVisibleDecals();
-	}
-
-	TConstArrayView<const FVisibleDecal*> FinishRelevantDecals(int32 ViewIndex, EDecalRenderStage Stage)
-	{
-		return ViewPackets[ViewIndex].FinishRelevantDecals(Stage);
-	}
-
-	bool HasStage(int32 ViewIndex, EDecalRenderStage Stage) const
+	static FDecalVisibilityTaskData* Launch(FRDGBuilder& GraphBuilder, const FScene& Scene, TConstArrayView<FViewInfo*> Views);
+
+	// Custom Render Passes have their own array of FViewInfo, outside the scene renderer's Views array.  When accessing decal visibility
+	// for a specific view array, an offset must be added to account for where the relevant data is located in the visibility task data.
+	// To enforce this, the view offset is passed separately to the accessor functions below.  INDEX_NONE is returned if decals are disabled
+	// for the given array of Views, where the task data constructor won't have added them to the list.  It's up to the caller to check for
+	// INDEX_NONE and skip decal rendering (the accessors will assert if INDEX_NONE is passed in for ViewOffset).
+	int32 GetViewOffset(TConstArrayView<FViewInfo> Views) const
 	{
-		return ViewPackets[ViewIndex].HasStage(Stage);
+		for (int32 ViewIndex = 0; ViewIndex < ViewPackets.Num(); ViewIndex++)
+		{
+			if (&ViewPackets[ViewIndex].View == Views.GetData())
+			{
+				return ViewIndex;
+			}
+		}
+		return INDEX_NONE;
 	}
 
-	void Finish()
+	TConstArrayView<const FVisibleDecal*> FinishRelevantDecals(int32 ViewOffset, int32 ViewIndex, EDecalRenderStage Stage)
 	{
-		for (FDecalVisibilityViewPacket& ViewPacket : ViewPackets)
-		{
-			ViewPacket.Finish();
-		}
+		check(ViewOffset != INDEX_NONE);
+		return ViewPackets[ViewOffset + ViewIndex].FinishRelevantDecals(Stage);
 	}
 
-	bool IsDBufferEnabled() const
+	bool HasStage(int32 ViewOffset, int32 ViewIndex, EDecalRenderStage Stage) const
 	{
-		return bDBufferEnabled;
+		check(ViewOffset != INDEX_NONE);
+		return ViewPackets[ViewOffset + ViewIndex].HasStage(Stage);
 	}
 
-	bool IsGBufferEnabled() const
+	void Finish(int32 ViewOffset, int32 ViewIndex)
 	{
-		return bGBufferEnabled;
+		check(ViewOffset != INDEX_NONE);
+		ViewPackets[ViewOffset + ViewIndex].Finish();
 	}
 
 private:
-	FDecalVisibilityTaskData(const FScene& Scene, TConstArrayView<FViewInfo> Views, bool bInDBufferEnabled, bool bInGBufferEnabled);
-
-	const bool bDBufferEnabled;
-	const bool bGBufferEnabled;
+	FDecalVisibilityTaskData(const FScene& Scene, TConstArrayView<FViewInfo*> Views);
 
 	TArray<FDecalVisibilityViewPacket, FRDGArrayAllocator> ViewPackets;
 
@@ -143,12 +144,14 @@ private:
 
 struct FDecalViewInfo
 {
-	FDecalViewInfo(int32 InViewIndex, EDecalRenderStage InRenderStage, EDecalRenderTargetMode InRenderTargetMode) :
+	FDecalViewInfo(int32 InViewOffset, int32 InViewIndex, EDecalRenderStage InRenderStage, EDecalRenderTargetMode InRenderTargetMode) :
+		ViewOffset(InViewOffset),
 		ViewIndex(InViewIndex),
 		RenderStage(InRenderStage),
 		RenderTargetMode(InRenderTargetMode)
 	{
 	}
+	int32 ViewOffset;
 	int32 ViewIndex;
 	EDecalRenderStage RenderStage;
 	EDecalRenderTargetMode RenderTargetMode;
diff --git a/Engine/Source/Runtime/Renderer/Private/DeferredShadingRenderer.cpp b/Engine/Source/Runtime/Renderer/Private/DeferredShadingRenderer.cpp
--- a/Engine/Source/Runtime/Renderer/Private/DeferredShadingRenderer.cpp
+++ b/Engine/Source/Runtime/Renderer/Private/DeferredShadingRenderer.cpp
@@ -2465,17 +2465,13 @@ void FDeferredShadingSceneRenderer::Render(FRDGBuilder& GraphBuilder, const FSce
 						Nanite::BuildShadingCommands(GraphBuilder, *Scene, ENaniteMeshPass::BasePass, NaniteBasePassShadingCommands, Nanite::EBuildShadingCommandsMode::Custom);
 					}
 
-#if 0	// TODO:  temporary fix for PLAY-109743 while investigating, comment out decal rendering in Custom Render Passes
 					// This class handles decal rendering
 					FCompositionLighting CustomRenderPassCompositionLighting(InitViewTaskDatas.Decals, CustomRenderPassViews, *CustomRenderPassSceneTextures, [this](int32 ViewIndex) { return false; });
 					CustomRenderPassCompositionLighting.ProcessBeforeBasePass(GraphBuilder, *CustomRenderPassSceneTextures->DBufferTextures, InstanceCullingManager, *CustomRenderPassSceneTextures->SubstrateSceneData);
-#endif
 
 					RenderBasePass(*this, GraphBuilder, CustomRenderPassViews, *CustomRenderPassSceneTextures, BasePassDepthStencilAccess, /*ForwardScreenSpaceShadowMaskTexture=*/nullptr, InstanceCullingManager, bNaniteEnabled, NaniteBasePassShadingCommands, NaniteRasterResults);
 
-#if 0	// TODO:  temporary fix for PLAY-109743
 					CustomRenderPassCompositionLighting.ProcessAfterBasePass(GraphBuilder, InstanceCullingManager, FCompositionLighting::EProcessAfterBasePassMode::All, *CustomRenderPassSceneTextures->SubstrateSceneData);
-#endif
 
 					if (ShouldRenderSingleLayerWater(CustomRenderPassViews))
 					{
@@ -2578,6 +2574,10 @@ void FDeferredShadingSceneRenderer::Render(FRDGBuilder& GraphBuilder, const FSce
 			// Restore original scene texture uniforms
 			CustomRenderPassSceneTextures->SetupMode = OriginalSceneTextureSetupMode;
 			CustomRenderPassSceneTextures->UniformBuffer = OriginalSceneTextureUniformBuffer;
+
+			// Clear DBuffer produced flags after Custom Render Pass rendering, so they will be cleared when they render again
+			CustomRenderPassSceneTextures->DBufferTextures->bDBufferProduced = false;
+			CustomRenderPassSceneTextures->DBufferTextures->bDBufferTexArrayProduced = false;
 		}
 	}
 
diff --git a/Engine/Source/Runtime/Renderer/Private/MobileDecalRendering.cpp b/Engine/Source/Runtime/Renderer/Private/MobileDecalRendering.cpp
--- a/Engine/Source/Runtime/Renderer/Private/MobileDecalRendering.cpp
+++ b/Engine/Source/Runtime/Renderer/Private/MobileDecalRendering.cpp
@@ -47,7 +47,7 @@ static bool DoesPlatformSupportDecals(EShaderPlatform ShaderPlatform)
 
 void FMobileSceneRenderer::RenderDecals(FRHICommandList& RHICmdList, FViewInfo& View, const FDecalViewInfo& DecalViewInfo, const FInstanceCullingDrawParams* InstanceCullingDrawParams)
 {
-	if (!DoesPlatformSupportDecals(View.GetShaderPlatform()) || !ViewFamily.EngineShowFlags.Decals || View.bIsPlanarReflection)
+	if (!DoesPlatformSupportDecals(View.GetShaderPlatform()) || !View.Family->EngineShowFlags.Decals || View.bIsPlanarReflection)
 	{
 		return;
 	}
@@ -73,7 +73,7 @@ void FMobileSceneRenderer::RenderDecals(FRHICommandList& RHICmdList, FViewInfo&
 
 void FMobileSceneRenderer::RenderDeferredDecalsMobile(FRHICommandList& RHICmdList, const FViewInfo& View, const FDecalViewInfo& DecalViewInfo)
 {
-	TConstArrayView<const FVisibleDecal*> RelevantDecals = DecalVisibilityTaskData->FinishRelevantDecals(DecalViewInfo.ViewIndex, DecalViewInfo.RenderStage);
+	TConstArrayView<const FVisibleDecal*> RelevantDecals = DecalVisibilityTaskData->FinishRelevantDecals(DecalViewInfo.ViewOffset, DecalViewInfo.ViewIndex, DecalViewInfo.RenderStage);
 	FRelevantDecalList SortedDecals;
 
 	for (const FVisibleDecal* VisibleDecal : RelevantDecals)
@@ -140,25 +140,27 @@ void FMobileSceneRenderer::RenderDeferredDecalsMobile(FRHICommandList& RHICmdLis
 	}
 }
 
-void FMobileSceneRenderer::RenderDBuffer(FRDGBuilder& GraphBuilder, FSceneTextures& SceneTextures, FDBufferTextures& DBufferTextures, FInstanceCullingManager& InstanceCullingManager)
+void FMobileSceneRenderer::RenderDBuffer(FRDGBuilder& GraphBuilder, TArray<FViewInfo>& InViews, FSceneTextures& SceneTextures, FDBufferTextures& DBufferTextures, FInstanceCullingManager& InstanceCullingManager)
 {
 	RDG_EVENT_SCOPE(GraphBuilder, "RenderDBuffer");
 	QUICK_SCOPE_CYCLE_COUNTER(STAT_RenderDBuffer);
 
-	const EShaderPlatform Platform = GetViewFamilyInfo(Views).GetShaderPlatform();
+	const EShaderPlatform Platform = GetViewFamilyInfo(InViews).GetShaderPlatform();
 
-	if (DecalVisibilityTaskData != nullptr)
-	{ 
-		for (int32 ViewIndex = 0; ViewIndex < Views.Num(); ++ViewIndex)
+	int32 DecalViewOffset = GetDecalViewOffset(InViews);
+
+	if (DecalViewOffset != INDEX_NONE)
+	{
+		for (int32 ViewIndex = 0; ViewIndex < InViews.Num(); ++ViewIndex)
 		{
-			FViewInfo& View = Views[ViewIndex];
+			FViewInfo& View = InViews[ViewIndex];
 
 			if (!View.ShouldRenderView())
 			{
 				continue;
 			}
 
-			TConstArrayView<const FVisibleDecal*> SortedDecals = DecalVisibilityTaskData->FinishRelevantDecals(ViewIndex, EDecalRenderStage::BeforeBasePass);
+			TConstArrayView<const FVisibleDecal*> SortedDecals = DecalVisibilityTaskData->FinishRelevantDecals(DecalViewOffset, ViewIndex, EDecalRenderStage::BeforeBasePass);
 			FDeferredDecalPassTextures DecalPassTextures = GetDeferredDecalPassTextures(GraphBuilder, View, SceneTextures, &DBufferTextures, EDecalRenderStage::BeforeBasePass);
 			AddDeferredDecalPass(GraphBuilder, View, SortedDecals, DecalPassTextures, InstanceCullingManager, EDecalRenderStage::BeforeBasePass);
 		}
diff --git a/Engine/Source/Runtime/Renderer/Private/MobileShadingRenderer.cpp b/Engine/Source/Runtime/Renderer/Private/MobileShadingRenderer.cpp
--- a/Engine/Source/Runtime/Renderer/Private/MobileShadingRenderer.cpp
+++ b/Engine/Source/Runtime/Renderer/Private/MobileShadingRenderer.cpp
@@ -295,12 +295,13 @@ struct FRenderViewContext
 {
 	FViewInfo* ViewInfo;
 	int32 ViewIndex;
+	int32 DecalViewOffset;
 	bool bIsFirstView;
 	bool bIsLastView;
 };
 using FRenderViewContextArray = TArray<FRenderViewContext, TInlineAllocator<2, SceneRenderingAllocator>>;
 
-static void GetRenderViews(TArrayView<FViewInfo> InViews, FRenderViewContextArray& RenderViews)
+static void GetRenderViews(TArrayView<FViewInfo> InViews, int32 DecalViewOffset, FRenderViewContextArray& RenderViews)
 {
 	for (int32 ViewIndex = 0; ViewIndex < InViews.Num(); ViewIndex++)
 	{
@@ -310,6 +311,7 @@ static void GetRenderViews(TArrayView<FViewInfo> InViews, FRenderViewContextArra
 			FRenderViewContext RenderView;
 			RenderView.ViewInfo = &View;
 			RenderView.ViewIndex = ViewIndex;
+			RenderView.DecalViewOffset = DecalViewOffset;
 			RenderView.bIsFirstView = (RenderViews.Num() == 0);
 			RenderView.bIsLastView = false;
 
@@ -777,14 +779,7 @@ void FMobileSceneRenderer::InitViews(
 		UpdateDirectionalLightUniformBuffers(GraphBuilder, Views[0]);
 	}
 
-	if(!ViewFamily.EngineShowFlags.HitProxies)
-	{ 
-		DecalVisibilityTaskData = FDecalVisibilityTaskData::Launch(GraphBuilder, *Scene, Views);
-	}
-	else
-	{
-		DecalVisibilityTaskData = nullptr;
-	}
+	DecalVisibilityTaskData = FDecalVisibilityTaskData::Launch(GraphBuilder, *Scene, AllViews);
 }
 
 static void BeginOcclusionScope(FRDGBuilder& GraphBuilder, TArray<FViewInfo>& Views)
@@ -817,7 +812,7 @@ static void EndOcclusionScope(FRDGBuilder& GraphBuilder, TArray<FViewInfo>& View
 void FMobileSceneRenderer::RenderFullDepthPrepass(FRDGBuilder& GraphBuilder, TArrayView<FViewInfo> InViews, FSceneTextures& SceneTextures, FInstanceCullingManager& InstanceCullingManager, bool bIsSceneCaptureRenderPass)
 {
 	FRenderViewContextArray RenderViews;
-	GetRenderViews(InViews, RenderViews);
+	GetRenderViews(InViews, GetDecalViewOffset(InViews), RenderViews);
 
 	FRDGTextureRef DepthTargetTexture = SceneTextures.Depth.Target;
 
@@ -954,7 +949,7 @@ void FMobileSceneRenderer::RenderMaskedPrePass(FRHICommandList& RHICmdList, cons
 	}
 }
 
-void FMobileSceneRenderer::RenderCustomRenderPassBasePass(FRDGBuilder& GraphBuilder, TArrayView<FViewInfo> InViews, FRDGTextureRef ViewFamilyTexture, FSceneTextures& SceneTextures, FInstanceCullingManager& InstanceCullingManager, bool bIncludeTranslucent)
+void FMobileSceneRenderer::RenderCustomRenderPassBasePass(FRDGBuilder& GraphBuilder, TArrayView<FViewInfo> InViews, FRDGTextureRef ViewFamilyTexture, FSceneTextures& SceneTextures, FMobileBasePassTextures& BasePassTextures, FInstanceCullingManager& InstanceCullingManager, bool bIncludeTranslucent)
 {
 	FRenderTargetBindingSlots BasePassRenderTargets;
 	// Use the same subpass hints as main render, to avoid generating new PSOs 
@@ -977,7 +972,7 @@ void FMobileSceneRenderer::RenderCustomRenderPassBasePass(FRDGBuilder& GraphBuil
 	}
 
 	FRenderViewContextArray RenderViews;
-	GetRenderViews(InViews, RenderViews);
+	GetRenderViews(InViews, GetDecalViewOffset(InViews), RenderViews);
 
 	for (FRenderViewContext& ViewContext : RenderViews)
 	{
@@ -994,7 +989,7 @@ void FMobileSceneRenderer::RenderCustomRenderPassBasePass(FRDGBuilder& GraphBuil
 		EMobileSceneTextureSetupMode SetupMode = bIsFullDepthPrepassEnabled ? EMobileSceneTextureSetupMode::SceneDepth : EMobileSceneTextureSetupMode::None;
 		auto PassParameters = &ParameterCollection->PassParameters;
 		PassParameters->View = View.GetShaderParameters();
-		PassParameters->MobileBasePass = CreateMobileBasePassUniformBuffer(GraphBuilder, View, EMobileBasePass::Opaque, SetupMode);
+		PassParameters->MobileBasePass = CreateMobileBasePassUniformBuffer(GraphBuilder, View, EMobileBasePass::Opaque, SetupMode, BasePassTextures);
 		PassParameters->RenderTargets = BasePassRenderTargets;
 
 		FMobileRenderPassParameters* TranslucencyPassParameters = nullptr;
@@ -1028,7 +1023,7 @@ void FMobileSceneRenderer::RenderCustomRenderPassBasePass(FRDGBuilder& GraphBuil
 			}
 		}
 
-		FDecalViewInfo DecalViewInfo = FDecalViewInfo(ViewContext.ViewIndex,
+		FDecalViewInfo DecalViewInfo = FDecalViewInfo(ViewContext.DecalViewOffset, ViewContext.ViewIndex,
 			bRequiresDBufferDecals ? EDecalRenderStage::Emissive : EDecalRenderStage::Mobile,
 			bDeferredShading ? EDecalRenderTargetMode::SceneColorAndGBuffer : EDecalRenderTargetMode::SceneColor);
 
@@ -1062,7 +1057,7 @@ void FMobileSceneRenderer::RenderCustomRenderPassBasePass(FRDGBuilder& GraphBuil
 			BasePassRenderTargets.DepthStencil.SetStencilLoadAction(ERenderTargetLoadAction::ELoad);
 			BasePassRenderTargets.DepthStencil.SetDepthStencilAccess(FExclusiveDepthStencil::DepthRead_StencilWrite);
 
-			RenderDeferredMultiPassDecals(GraphBuilder, View, ViewContext.ViewIndex, SceneTextures, InstanceCullingManager);
+			RenderDeferredMultiPassDecals(GraphBuilder, View, ViewContext.DecalViewOffset, ViewContext.ViewIndex, SceneTextures, InstanceCullingManager);
 		}
 
 		if (bIncludeTranslucent)
@@ -1388,7 +1383,12 @@ void FMobileSceneRenderer::Render(FRDGBuilder& GraphBuilder, const FSceneRenderU
 	const FRDGSystemTextures& SystemTextures = FRDGSystemTextures::Get(GraphBuilder);
 
 	RenderWaterInfoTexture(GraphBuilder, *this, Scene);
-	
+
+	FViewInfo* MainView = Views.Num() > 0 ? &Views[0] : nullptr;
+	const bool bIsMobileMultiView = SceneTextures.Config.bRequireMultiView || (MainView && MainView->Aspects.IsMobileMultiViewEnabled());
+
+	FMobileBasePassTextures BasePassTextures{};
+
 	if (CustomRenderPassInfos.Num() > 0)
 	{
 		QUICK_SCOPE_CYCLE_COUNTER(STAT_CustomRenderPasses);
@@ -1455,12 +1455,33 @@ void FMobileSceneRenderer::Render(FRDGBuilder& GraphBuilder, const FSceneRenderU
 				// Render base pass if the custom pass requires it. Otherwise if full depth prepass is not enabled, then depth is generated in the base pass.
 				if (CustomRenderPass->GetRenderMode() == FCustomRenderPassBase::ERenderMode::DepthAndBasePass || (CustomRenderPass->GetRenderMode() == FCustomRenderPassBase::ERenderMode::DepthPass && !bIsFullDepthPrepassEnabled))
 				{
-					if (CustomRenderPassViewFamily.IsSceneTexturesOwner() && CustomRenderPass->GetRenderMode() == FCustomRenderPassBase::ERenderMode::DepthAndBasePass)
+					FMobileBasePassTextures CustomRenderPassBasePassTextures{};
+
+					if (CustomRenderPass->GetRenderMode() == FCustomRenderPassBase::ERenderMode::DepthAndBasePass && bRequiresDBufferDecals)
 					{
-						*CustomRenderPassSceneTextures.DBufferTextures = CreateDBufferTextures(GraphBuilder, CustomRenderPassSceneTextures.Config.Extent, ShaderPlatform);
+						// If mobile multi-view is enabled, force separate DBuffer textures to be allocated.  Multi-view causes the DBuffer textures to be allocated as array
+						// textures, and the concern is that could cause issues with Custom Render Passes.  We also need separate DBuffer textures if this view allocates its
+						// own scene textures, as the scene texture resolution will be different than the main view in that case.
+						if (CustomRenderPassViewFamily.IsSceneTexturesOwner() || bIsMobileMultiView)
+						{
+							CustomRenderPassBasePassTextures.DBufferTextures = CreateDBufferTextures(GraphBuilder, CustomRenderPassSceneTextures.Config.Extent, ShaderPlatform);
+						}
+						else
+						{
+							// Sharing main renderer DBuffer textures, allocate them if not yet allocated.  Note that scene texture extents for custom render passes are the
+							// same as the main view(s), as upstream scene texture extent calculation logic ensures the extent is a union of all views rendered in the frame,
+							// including Scene Captures and Custom Render Passes.  The exception is if bExcludeFromSceneTextureExtents is set, which excludes the given
+							// Scene Capture (or Custom Render Pass) from extent calculations.  In that case, IsSceneTexturesOwner() will be true, and the code above runs.
+							if (!BasePassTextures.DBufferTextures.DBufferA && !BasePassTextures.DBufferTextures.DBufferATexArray)
+							{
+								BasePassTextures.DBufferTextures = CreateDBufferTextures(GraphBuilder, SceneTextures.Config.Extent, ShaderPlatform, bIsMobileMultiView);
+							}
+							CustomRenderPassBasePassTextures.DBufferTextures = BasePassTextures.DBufferTextures;
+						}
+						RenderDBuffer(GraphBuilder, CustomRenderPassViews, CustomRenderPassSceneTextures, CustomRenderPassBasePassTextures.DBufferTextures, InstanceCullingManager);
 					}
 
-					RenderCustomRenderPassBasePass(GraphBuilder, CustomRenderPassViews, ViewFamilyTexture, CustomRenderPassSceneTextures, InstanceCullingManager, CustomRenderPass->IsTranslucentIncluded());
+					RenderCustomRenderPassBasePass(GraphBuilder, CustomRenderPassViews, ViewFamilyTexture, CustomRenderPassSceneTextures, CustomRenderPassBasePassTextures, InstanceCullingManager, CustomRenderPass->IsTranslucentIncluded());
 				}
 
 				CustomRenderPassSceneTextures.MobileSetupMode = EMobileSceneTextureSetupMode::SceneColor | EMobileSceneTextureSetupMode::SceneDepth | EMobileSceneTextureSetupMode::SceneDepthAux;
@@ -1535,6 +1556,21 @@ void FMobileSceneRenderer::Render(FRDGBuilder& GraphBuilder, const FSceneRenderU
 
 			CustomRenderPassSceneTextures.MobileSetupMode = OriginalSceneTextureSetupMode;
 			CustomRenderPassSceneTextures.MobileUniformBuffer = OriginalSceneTextureUniformBuffer;
+
+			// Clear DBuffer produced flags after Custom Render Pass rendering, so they will be cleared when they render again
+			CustomRenderPassSceneTextures.DBufferTextures->bDBufferProduced = false;
+			CustomRenderPassSceneTextures.DBufferTextures->bDBufferTexArrayProduced = false;
+
+			// Call finish on decal visibility task data.  Note that this is handled by the FCompositionLighting destructor for Custom Render Passes
+			// in the deferred renderer, but the mobile renderer doesn't use that class.
+			int32 DecalViewOffset = GetDecalViewOffset(CustomRenderPassViews);
+			if (DecalViewOffset != INDEX_NONE)
+			{
+				for (int32 ViewIndex = 0; ViewIndex < CustomRenderPassViews.Num(); ViewIndex++)
+				{
+					DecalVisibilityTaskData->Finish(DecalViewOffset, ViewIndex);
+				}
+			}
 		}
 	}
 
@@ -1572,7 +1608,6 @@ void FMobileSceneRenderer::Render(FRDGBuilder& GraphBuilder, const FSceneRenderU
 	}
 
 
-	FMobileBasePassTextures BasePassTextures{};
 	if (bIsFullDepthPrepassEnabled)
 	{
 		// When renderer is in ERendererOutput::DepthPrepassOnly mode, bShouldRenderHZB is set to false in InitViews()
@@ -1594,9 +1629,6 @@ void FMobileSceneRenderer::Render(FRDGBuilder& GraphBuilder, const FSceneRenderU
 			RenderDistanceFieldLighting(GraphBuilder, SceneTextures, FDistanceFieldAOParameters(OcclusionMaxDistance), DynamicBentNormalAOTextures, false, false, true);
 		}
 
-		FViewInfo* MainView = Views.Num() > 0 ? &Views[0] : nullptr;
-		const bool bIsMobileMultiView = SceneTextures.Config.bRequireMultiView || (MainView && MainView->Aspects.IsMobileMultiViewEnabled());
-
 		// When renderer is in ERendererOutput::DepthPrepassOnly mode, bRequiresShadowProjections is set to false in InitViews()
 		if (bRequiresShadowProjections)
 		{
@@ -1622,8 +1654,12 @@ void FMobileSceneRenderer::Render(FRDGBuilder& GraphBuilder, const FSceneRenderU
 		{
 			if (bRequiresDBufferDecals)
 			{
-				BasePassTextures.DBufferTextures = CreateDBufferTextures(GraphBuilder, SceneTextures.Config.Extent, ShaderPlatform, bIsMobileMultiView);
-				RenderDBuffer(GraphBuilder, SceneTextures, BasePassTextures.DBufferTextures, InstanceCullingManager);
+				// DBufferTextures may already have been created for Custom Render Passes above
+				if (!BasePassTextures.DBufferTextures.DBufferA && !BasePassTextures.DBufferTextures.DBufferATexArray)
+				{
+					BasePassTextures.DBufferTextures = CreateDBufferTextures(GraphBuilder, SceneTextures.Config.Extent, ShaderPlatform, bIsMobileMultiView);
+				}
+				RenderDBuffer(GraphBuilder, Views, SceneTextures, BasePassTextures.DBufferTextures, InstanceCullingManager);
 			}
 		}
 
@@ -1660,9 +1696,13 @@ void FMobileSceneRenderer::Render(FRDGBuilder& GraphBuilder, const FSceneRenderU
 			RenderForward(GraphBuilder, ViewFamilyTexture, SceneTextures, BasePassTextures, InstanceCullingManager);
 		}
 
-		if (DecalVisibilityTaskData)
+		int32 DecalViewOffset = GetDecalViewOffset(Views);
+		if (DecalViewOffset != INDEX_NONE)
 		{
-			DecalVisibilityTaskData->Finish();
+			for (int32 ViewIndex = 0; ViewIndex < Views.Num(); ViewIndex++)
+			{
+				DecalVisibilityTaskData->Finish(DecalViewOffset, ViewIndex);
+			}
 		}
 
 		if (DoOcclusionQueries() && !bIsFullDepthPrepassEnabled)
@@ -1941,7 +1981,7 @@ void FMobileSceneRenderer::RenderForward(FRDGBuilder& GraphBuilder, FRDGTextureR
 	const FRDGSystemTextures& SystemTextures = FRDGSystemTextures::Get(GraphBuilder);
 
 	FRenderViewContextArray RenderViews;
-	GetRenderViews(Views, RenderViews);
+	GetRenderViews(Views, GetDecalViewOffset(Views), RenderViews);
 
 	bool bDoOcclusionQueries = !bIsFullDepthPrepassEnabled && DoOcclusionQueries();
 	FViewOcclusionQueriesPerView QueriesPerView;
@@ -2053,7 +2093,7 @@ void FMobileSceneRenderer::RenderForwardSinglePass(FRDGBuilder& GraphBuilder, FM
 	const bool bDoOcclusionInMainPass = bDoOcclusionQueries && !bVulkanAdrenoOcclusionMode;
 	PassParameters->RenderTargets.NumOcclusionQueries = bVulkanAdrenoOcclusionMode ? 0u : NumOcclusionQueries;
 	
-	FDecalViewInfo DecalViewInfo = FDecalViewInfo(ViewContext.ViewIndex,
+	FDecalViewInfo DecalViewInfo = FDecalViewInfo(ViewContext.DecalViewOffset, ViewContext.ViewIndex,
 		bRequiresDBufferDecals ? EDecalRenderStage::Emissive : EDecalRenderStage::Mobile,
 		EDecalRenderTargetMode::SceneColor);
 
@@ -2173,7 +2213,8 @@ void FMobileSceneRenderer::RenderForwardMultiPass(FRDGBuilder& GraphBuilder, FMo
 		BuildMeshRenderingCommands(GraphBuilder, EMeshPass::DebugViewMode, View, Scene->GPUScene, InstanceCullingManager, ParameterCollection->DebugViewModeInstanceCullingDrawParams);
 	}
 
-	FDecalViewInfo DecalViewInfo = FDecalViewInfo(ViewContext.ViewIndex, 
+	FDecalViewInfo DecalViewInfo = FDecalViewInfo(ViewContext.DecalViewOffset,
+												  ViewContext.ViewIndex,
 												  bRequiresDBufferDecals ? EDecalRenderStage::Emissive : EDecalRenderStage::Mobile,
 												  EDecalRenderTargetMode::SceneColor);
 
@@ -2430,7 +2471,7 @@ void FMobileSceneRenderer::RenderDeferredSinglePass(FRDGBuilder& GraphBuilder, F
 	const FRDGSystemTextures& SystemTextures = FRDGSystemTextures::Get(GraphBuilder);
 	
 	FRenderViewContextArray RenderViews;
-	GetRenderViews(Views, RenderViews);
+	GetRenderViews(Views, GetDecalViewOffset(Views), RenderViews);
 
 	bool bDoOcclusionQueries = (!bIsFullDepthPrepassEnabled && DoOcclusionQueries());
 	FViewOcclusionQueriesPerView QueriesPerView;
@@ -2504,7 +2545,8 @@ void FMobileSceneRenderer::RenderDeferredSinglePass(FRDGBuilder& GraphBuilder, F
 			BuildMeshRenderingCommands(GraphBuilder, StandardTranslucencyMeshPass, View, Scene->GPUScene, InstanceCullingManager, ParameterCollection->TranslucencyInstanceCullingDrawParams);
 		}
 
-		FDecalViewInfo DecalViewInfo = FDecalViewInfo(ViewContext.ViewIndex,
+		FDecalViewInfo DecalViewInfo = FDecalViewInfo(ViewContext.DecalViewOffset,
+													  ViewContext.ViewIndex,
 													  bRequiresDBufferDecals ? EDecalRenderStage::Emissive : EDecalRenderStage::MobileBeforeLighting,
 													  EDecalRenderTargetMode::SceneColorAndGBuffer);
 
@@ -2557,10 +2599,11 @@ void FMobileSceneRenderer::RenderDeferredSinglePass(FRDGBuilder& GraphBuilder, F
 	}
 }
 
-void FMobileSceneRenderer::RenderDeferredMultiPassDecals(FRDGBuilder& GraphBuilder, FViewInfo& View, uint32 ViewIndex, FSceneTextures& SceneTextures, FInstanceCullingManager& InstanceCullingManager)
+void FMobileSceneRenderer::RenderDeferredMultiPassDecals(FRDGBuilder& GraphBuilder, FViewInfo& View, int32 DecalViewOffset, uint32 ViewIndex, FSceneTextures& SceneTextures, FInstanceCullingManager& InstanceCullingManager)
 {
-	if (Scene->Decals.Num() == 0 && !HasAnyDraw(View.ParallelMeshDrawCommandPasses[EMeshPass::MeshDecal_SceneColorAndGBuffer])
-								 && !HasAnyDraw(View.ParallelMeshDrawCommandPasses[EMeshPass::MeshDecal_SceneColorAndGBufferNoNormal]))
+	if (DecalViewOffset == INDEX_NONE ||
+		(Scene->Decals.Num() == 0 && !HasAnyDraw(View.ParallelMeshDrawCommandPasses[EMeshPass::MeshDecal_SceneColorAndGBuffer])
+								  && !HasAnyDraw(View.ParallelMeshDrawCommandPasses[EMeshPass::MeshDecal_SceneColorAndGBufferNoNormal])))
 	{
 		// Skip decal render-pass if there is nothing to draw
 		return;
@@ -2571,7 +2614,7 @@ void FMobileSceneRenderer::RenderDeferredMultiPassDecals(FRDGBuilder& GraphBuild
 	const EDecalRenderTargetMode DecalRenderTargetModes[] = { EDecalRenderTargetMode::SceneColorAndGBuffer, EDecalRenderTargetMode::SceneColorAndGBufferNoNormal };
 	const EDecalRenderStage DecalRenderStage = bRequiresDBufferDecals ? EDecalRenderStage::Emissive : EDecalRenderStage::MobileBeforeLighting;
 	constexpr const TCHAR* PassNames[2] = { TEXT("SceneColorAndGBuffer"), TEXT("SceneColorAndGBufferNoNormal")};
-	TConstArrayView<const FVisibleDecal*> RelevantDecals = DecalVisibilityTaskData ? DecalVisibilityTaskData->FinishRelevantDecals(ViewIndex, DecalRenderStage) : TConstArrayView<const FVisibleDecal*>();
+	TConstArrayView<const FVisibleDecal*> RelevantDecals = DecalVisibilityTaskData ? DecalVisibilityTaskData->FinishRelevantDecals(DecalViewOffset, ViewIndex, DecalRenderStage) : TConstArrayView<const FVisibleDecal*>();
 
 	for (const EDecalRenderTargetMode DecalRenderTargetMode : DecalRenderTargetModes)
 	{
@@ -2614,15 +2657,13 @@ void FMobileSceneRenderer::RenderDeferredMultiPassDecals(FRDGBuilder& GraphBuild
 			BuildMeshRenderingCommands(GraphBuilder, MeshDecalsPass, View, Scene->GPUScene, InstanceCullingManager, PassParameters->InstanceCullingDrawParams);
 		}
 
-		FDecalViewInfo DecalViewInfo = FDecalViewInfo(ViewIndex,
-			DecalRenderStage,
-			DecalRenderTargetMode);
+		FDecalViewInfo DecalViewInfo = FDecalViewInfo(DecalViewOffset, ViewIndex, DecalRenderStage, DecalRenderTargetMode);
 
 		GraphBuilder.AddPass(
 			RDG_EVENT_NAME("Decals_%s", DecalRendering::GetRenderTargetModeName(DecalRenderTargetMode)),
 			PassParameters,
 			ERDGPassFlags::Raster,
-			[this, DecalViewInfo, PassParameters, &View](FRHICommandList& RHICmdList)
+			[this, DecalViewOffset, DecalViewInfo, PassParameters, &View](FRHICommandList& RHICmdList)
 			{
 				RenderDecals(RHICmdList, View, DecalViewInfo, &PassParameters->InstanceCullingDrawParams);
 			});
@@ -2636,7 +2677,7 @@ void FMobileSceneRenderer::RenderDeferredMultiPass(FRDGBuilder& GraphBuilder, FS
 	const FRDGSystemTextures& SystemTextures = FRDGSystemTextures::Get(GraphBuilder);
 
 	FRenderViewContextArray RenderViews;
-	GetRenderViews(Views, RenderViews);
+	GetRenderViews(Views, GetDecalViewOffset(Views), RenderViews);
 
 	bool bDoOcclusionQueries = (!bIsFullDepthPrepassEnabled && DoOcclusionQueries());
 	FViewOcclusionQueriesPerView QueriesPerView;
@@ -2752,7 +2793,7 @@ void FMobileSceneRenderer::RenderDeferredMultiPass(FRDGBuilder& GraphBuilder, FS
 		SCOPED_GPU_MASK(GraphBuilder.RHICmdList, !View.IsInstancedStereoPass() ? View.GPUMask : (View.GPUMask | View.GetInstancedView()->GPUMask));
 		SCOPED_CONDITIONAL_DRAW_EVENTF(GraphBuilder.RHICmdList, EventView, RenderViews.Num() > 1, TEXT("View%d"), ViewContext.ViewIndex);
 
-		RenderDeferredMultiPassDecals(GraphBuilder, View, ViewContext.ViewIndex, SceneTextures, InstanceCullingManager);
+		RenderDeferredMultiPassDecals(GraphBuilder, View, ViewContext.DecalViewOffset, ViewContext.ViewIndex, SceneTextures, InstanceCullingManager);
 	}
 
 	if (!bIsFullDepthPrepassEnabled)
diff --git a/Engine/Source/Runtime/Renderer/Private/SceneRendering.h b/Engine/Source/Runtime/Renderer/Private/SceneRendering.h
--- a/Engine/Source/Runtime/Renderer/Private/SceneRendering.h
+++ b/Engine/Source/Runtime/Renderer/Private/SceneRendering.h
@@ -2885,7 +2885,7 @@ protected:
 
 	void RenderMobileLocalLightsBuffer(FRDGBuilder& GraphBuilder, FSceneTextures& SceneTextures, const FSortedLightSetSceneInfo& SortedLights);
 
-	void RenderCustomRenderPassBasePass(FRDGBuilder& GraphBuilder, TArrayView<FViewInfo> InViews, FRDGTextureRef ViewFamilyTexture, FSceneTextures& SceneTextures, FInstanceCullingManager& InstanceCullingManager, bool bIncludeTranslucent);
+	void RenderCustomRenderPassBasePass(FRDGBuilder& GraphBuilder, TArrayView<FViewInfo> InViews, FRDGTextureRef ViewFamilyTexture, FSceneTextures& SceneTextures, FMobileBasePassTextures& BasePassTextures, FInstanceCullingManager& InstanceCullingManager, bool bIncludeTranslucent);
 
 	/** Renders the opaque base pass for mobile. */
 	void RenderMobileBasePass(FRHICommandList& RHICmdList, const FViewInfo& View, const FInstanceCullingDrawParams* InstanceCullingDrawParams, const FInstanceCullingDrawParams* SkyPassInstanceCullingDrawParams);
@@ -2922,7 +2922,13 @@ protected:
 	/** Renders decals. */
 	void RenderDecals(FRHICommandList& RHICmdList, FViewInfo& View, const FDecalViewInfo& DecalViewInfo, const FInstanceCullingDrawParams* InstanceCullingDrawParams);
 	void RenderDeferredDecalsMobile(FRHICommandList& RHICmdList, const FViewInfo& View, const FDecalViewInfo& DecalViewInfo);
-	void RenderDBuffer(FRDGBuilder& GraphBuilder, FSceneTextures& SceneTextures, FDBufferTextures& DBufferTextures, FInstanceCullingManager& InstanceCullingManager);
+	void RenderDBuffer(FRDGBuilder& GraphBuilder, TArray<FViewInfo>& InViews, FSceneTextures& SceneTextures, FDBufferTextures& DBufferTextures, FInstanceCullingManager& InstanceCullingManager);
+
+	/** Utility to fetch DecalViewOffset to pass to decal related rendering functions */
+	int32 GetDecalViewOffset(TConstArrayView<FViewInfo> InViews) const
+	{
+		return DecalVisibilityTaskData ? DecalVisibilityTaskData->GetViewOffset(InViews) : INDEX_NONE;
+	}
 
 	/** Renders the atmospheric and height fog */
 	void RenderFog(FRHICommandList& RHICmdList, const FViewInfo& View);
@@ -2944,7 +2950,7 @@ protected:
 	
 	void RenderDeferredSinglePass(FRDGBuilder& GraphBuilder, FSceneTextures& SceneTextures, const FSortedLightSetSceneInfo& SortedLightSet, FMobileBasePassTextures& BasePassTextures, FInstanceCullingManager& InstanceCullingManager);
 	void RenderDeferredMultiPass(FRDGBuilder& GraphBuilder, FSceneTextures& SceneTextures, const FSortedLightSetSceneInfo& SortedLightSet, FMobileBasePassTextures& BasePassTextures, FInstanceCullingManager& InstanceCullingManager);
-	void RenderDeferredMultiPassDecals(FRDGBuilder& GraphBuilder, FViewInfo& View, uint32 ViewIndex, FSceneTextures& SceneTextures, FInstanceCullingManager& InstanceCullingManager);
+	void RenderDeferredMultiPassDecals(FRDGBuilder& GraphBuilder, FViewInfo& View, int32 DecalViewOffset, uint32 ViewIndex, FSceneTextures& SceneTextures, FInstanceCullingManager& InstanceCullingManager);
 
 	void RenderDistortion(FRDGBuilder& GraphBuilder,
 		const FViewInfo& View,
diff --git a/Engine/Source/Runtime/Renderer/Private/SceneVisibility.cpp b/Engine/Source/Runtime/Renderer/Private/SceneVisibility.cpp
--- a/Engine/Source/Runtime/Renderer/Private/SceneVisibility.cpp
+++ b/Engine/Source/Runtime/Renderer/Private/SceneVisibility.cpp
@@ -5894,10 +5894,7 @@ void FDeferredShadingSceneRenderer::BeginInitViews(
 	}
 #endif
 
-	if (!ViewFamily.EngineShowFlags.HitProxies)
-	{
-		TaskDatas.Decals = FDecalVisibilityTaskData::Launch(GraphBuilder, *Scene, Views);
-	}
+	TaskDatas.Decals = FDecalVisibilityTaskData::Launch(GraphBuilder, *Scene, AllViews);
 
 	// Attempt to launch dynamic shadow tasks early before finalizing visibility.
 	if (bRendererOutputFinalSceneColor)
diff --git a/Engine/Source/Runtime/Renderer/Private/Substrate/Substrate.cpp b/Engine/Source/Runtime/Renderer/Private/Substrate/Substrate.cpp
--- a/Engine/Source/Runtime/Renderer/Private/Substrate/Substrate.cpp
+++ b/Engine/Source/Runtime/Renderer/Private/Substrate/Substrate.cpp
@@ -1976,13 +1976,19 @@ void AddSubstrateDBufferBasePass(
 		return;
 	}
 
+	int32 DecalViewOffset = DecalVisibility->GetViewOffset(Views);
+	if (DecalViewOffset == INDEX_NONE)
+	{
+		return;
+	}
+
 	RDG_EVENT_SCOPE_CONDITIONAL(GraphBuilder, IsSubstrateEnabled() && Views.Num() > 0, "Substrate::DBufferBasePass");
 
 	for (int32 ViewIndex = 0; ViewIndex < Views.Num(); ++ViewIndex)
 	{
 		FViewInfo& View = Views[ViewIndex];
 
-		const bool bEnableDecals = IsDBufferPassEnabled(View.GetShaderPlatform()) && DecalVisibility->HasStage(ViewIndex, EDecalRenderStage::BeforeBasePass);
+		const bool bEnableDecals = IsDBufferPassEnabled(View.GetShaderPlatform()) && DecalVisibility->HasStage(DecalViewOffset, ViewIndex, EDecalRenderStage::BeforeBasePass);
 		if (bEnableDecals)
 		{
 			RDG_GPU_MASK_SCOPE(GraphBuilder, View.GPUMask);
@@ -1990,7 +1996,7 @@ void AddSubstrateDBufferBasePass(
 
 			View.BeginRenderView();
 			FDeferredDecalPassTextures DecalPassTextures = GetDeferredDecalPassTextures(GraphBuilder, View, InSceneTextures, &DBufferTextures, EDecalRenderStage::BeforeBasePass);
-			TConstArrayView<const FVisibleDecal*> SortedDecals = DecalVisibility->FinishRelevantDecals(ViewIndex, EDecalRenderStage::BeforeBasePass);
+			TConstArrayView<const FVisibleDecal*> SortedDecals = DecalVisibility->FinishRelevantDecals(DecalViewOffset, ViewIndex, EDecalRenderStage::BeforeBasePass);
 			AddDeferredDecalPass(GraphBuilder, View, SortedDecals, DecalPassTextures, InstanceCullingManager, EDecalRenderStage::BeforeBasePass);
 		}
 	}
