[HWRT] Support static/cached SBT entries without specifying FRHIRayTracingGeometry.
- in which case FAllocationKey relies only on FRayTracingCachedMeshCommandFlags.
- store explicit bDeduplicate flag to identify which allocations are tracked for deduplication instead of assuming it based on whether it has a valid geometry pointer.

Context: For Nanite Ray Tracing using CLAS APIs, we don't store ray tracing geometry pointer in SBT and can share entries as long as the materials / flags are the same.

#Synced-CL 50254349

diff --git a/Engine/Source/Runtime/Renderer/Private/RayTracing/RayTracingShaderBindingTable.cpp b/Engine/Source/Runtime/Renderer/Private/RayTracing/RayTracingShaderBindingTable.cpp
--- a/Engine/Source/Runtime/Renderer/Private/RayTracing/RayTracingShaderBindingTable.cpp
+++ b/Engine/Source/Runtime/Renderer/Private/RayTracing/RayTracingShaderBindingTable.cpp
@@ -450,7 +450,6 @@ FRayTracingSBTAllocation* FRayTracingShaderBindingTable::AllocateStaticRangeInte
 FRayTracingSBTAllocation* FRayTracingShaderBindingTable::AllocateStaticRange(uint32 SegmentCount, const FRHIRayTracingGeometry* Geometry, FRayTracingCachedMeshCommandFlags Flags)
 {
 	ensure(!bStaticAllocationsLocked);
-	check(Geometry != nullptr);
 
 	// No allocation if we are not rendering decals and all segments are decals
 	if (RayTracing::ShouldExcludeDecals() && Flags.bAllSegmentsDecal)
@@ -484,6 +483,7 @@ FRayTracingSBTAllocation* FRayTracingShaderBindingTable::AllocateStaticRange(uin
 	if (Allocation.RefCount == 0)
 	{		
 		Allocation.Allocation = AllocateStaticRangeInternal(AllocatedLayers, SegmentCount, Geometry, Flags);
+		Allocation.Allocation->bDeduplicate = true;
 	}
 	else
 	{
@@ -507,8 +507,7 @@ void FRayTracingShaderBindingTable::FreeStaticRange(FRayTracingSBTAllocation* Al
 
 	TotalStaticAllocationCount--;
 
-	// If geometry is stored then it could have been deduplicated and we can build the allocation key again
-	if (Allocation->Geometry)
+	if (Allocation->bDeduplicate)
 	{
 		FAllocationKey Key;
 		Key.Geometry = Allocation->Geometry;
diff --git a/Engine/Source/Runtime/Renderer/Private/RayTracing/RayTracingShaderBindingTable.h b/Engine/Source/Runtime/Renderer/Private/RayTracing/RayTracingShaderBindingTable.h
--- a/Engine/Source/Runtime/Renderer/Private/RayTracing/RayTracingShaderBindingTable.h
+++ b/Engine/Source/Runtime/Renderer/Private/RayTracing/RayTracingShaderBindingTable.h
@@ -86,6 +86,7 @@ private:
 	// Store the original geometry and flags in the allocation object so it can be used to build the lookup key again used for deduplication
 	const FRHIRayTracingGeometry* Geometry = nullptr;
 	FRayTracingCachedMeshCommandFlags Flags;
+	bool bDeduplicate = false;
 };
 
 using FRayTracingPersistentShaderBindingTableID = uint32;
