Translucency: Basic fix to support AfterDOF translucency holdout composition into scene color.

* Critical for compositing keyed footage use case, to avoid double-DOF on the edge.
* New "r.Translucency.Holdout.Location" console variable: 0=BeforeDOF, 1=AfterDOF.
* Solution could be further extended to automatically support different locations based on the primitive's translucent material location without any cvar. However, this fix currently suffices as we also need to force-disable "r.Translucency.AutoBeforeDOF".

#Synced-CL 50255959

diff --git a/Engine/Source/Runtime/Renderer/Private/PostProcess/PostProcessing.cpp b/Engine/Source/Runtime/Renderer/Private/PostProcess/PostProcessing.cpp
--- a/Engine/Source/Runtime/Renderer/Private/PostProcess/PostProcessing.cpp
+++ b/Engine/Source/Runtime/Renderer/Private/PostProcess/PostProcessing.cpp
@@ -922,6 +922,15 @@ void AddPostProcessingPasses(
 			{
 				RenderHairComposition(GraphBuilder, View, SceneColor.Texture, SceneDepth.Texture, Velocity.Texture);
 			}
+
+			if (GetTranslucentHoldoutComposition() == ETranslucentHoldoutCompositionType::AfterDOF)
+			{
+				FTranslucencyPassResources HoldoutTranslucencyResources = Inputs.TranslucencyViewResourcesMap.Get(ETranslucencyPass::TPT_TranslucencyHoldout);
+				if (HoldoutTranslucencyResources.IsValid())
+				{
+					AddTranslucentHoldoutToSceneColorAlphaPass(GraphBuilder, View, SceneColor.Texture, SceneDepth.Texture, HoldoutTranslucencyResources);
+				}
+			}
 		}
 
 		// Allows for the scene color to be the slice of an array between temporal upscaler and tonemaper.
diff --git a/Engine/Source/Runtime/Renderer/Private/TranslucentRendering.cpp b/Engine/Source/Runtime/Renderer/Private/TranslucentRendering.cpp
--- a/Engine/Source/Runtime/Renderer/Private/TranslucentRendering.cpp
+++ b/Engine/Source/Runtime/Renderer/Private/TranslucentRendering.cpp
@@ -87,6 +87,15 @@ static TAutoConsoleVariable<int32> CVarTranslucencyUpperBoundQuantization(
 	TEXT("Only recommended for use with the transient allocator (on supported platforms) with a large transient texture cache (e.g RHI.TransientAllocator.TextureCacheSize=512)"),
 	ECVF_RenderThreadSafe | ECVF_Default);
 
+static int32 GTranslucencyHoldoutComposition = static_cast<int32>(ETranslucentHoldoutCompositionType::BeforeDOF);
+static FAutoConsoleVariableRef CVarTranslucencyHoldoutLocation(
+	TEXT("r.Translucency.Holdout.Location"),
+	GTranslucencyHoldoutComposition,
+	TEXT("Control the location of translucent holdout composition into scene color alpha.\n")
+	TEXT(" 0: BeforeDOF (default)\n")
+	TEXT(" 1: AfterDOF\n"),
+	ECVF_RenderThreadSafe);
+
 static TAutoConsoleVariable<float> CVarVolumetricCloudSoftBlendingDistanceOnTranslucent(
 	TEXT("r.VolumetricCloud.SoftBlendingDistanceOnTranslucent"), 0.5,
 	TEXT("The soft blending in distance in kilometer used to soft blend in cloud over translucent from the evaluated start depth."),
@@ -390,6 +399,32 @@ FRDGTextureRef AddCopySceneColorPass(FRDGBuilder& GraphBuilder, TArrayView<const
 	return SceneColorCopyTexture;
 }
 
+void AddTranslucentHoldoutToSceneColorAlphaPass(FRDGBuilder& GraphBuilder,
+	const FViewInfo& View,
+	FRDGTextureRef SceneColorTexture,
+	FRDGTextureRef SceneDepthTexture, 
+	FTranslucencyPassResources& TranslucencyPassResources)
+{
+	const ETranslucencyPass::Type PassType = TranslucencyPassResources.Pass;
+	check(PassType == ETranslucencyPass::TPT_TranslucencyHoldout);
+
+	FTranslucencyComposition TranslucencyComposition;
+	TranslucencyComposition.Operation = FTranslucencyComposition::EOperation::ComposeToSceneColorAlpha;
+	TranslucencyComposition.SceneColor = FScreenPassTextureSlice::CreateFromScreenPassTexture(GraphBuilder, FScreenPassTexture(SceneColorTexture, View.ViewRect));
+	TranslucencyComposition.SceneDepth = FScreenPassTexture(SceneDepthTexture, View.ViewRect);
+	TranslucencyComposition.OutputViewport = FScreenPassTextureViewport(SceneDepthTexture, View.ViewRect);
+
+	FScreenPassTexture UpscaledTranslucency = TranslucencyComposition.AddPass(
+		GraphBuilder, View, TranslucencyPassResources);
+
+	ensure(View.ViewRect == UpscaledTranslucency.ViewRect);
+	ensure(UpscaledTranslucency.Texture == SceneColorTexture);
+
+	//Invalidate.
+	TranslucencyPassResources = FTranslucencyPassResources();
+	TranslucencyPassResources.Pass = PassType;
+}
+
 static void AddCopySceneColorAlphaPass(FRDGBuilder& GraphBuilder, TArrayView<const FViewInfo> Views, FRDGTextureRef SourceTexture, FRDGTextureRef TargetTexture, bool bSkipIfUnderwater)
 {
 	ERenderTargetLoadAction LoadAction = ERenderTargetLoadAction::ENoAction;
@@ -936,6 +971,16 @@ FRDGTextureMSAA CreatePostDOFTranslucentTexture(
 		bIsModulate ? GFastVRamConfig.SeparateTranslucencyModulate : GFastVRamConfig.SeparateTranslucency);
 }
 
+ETranslucentHoldoutCompositionType GetTranslucentHoldoutComposition()
+{
+	switch (GTranslucencyHoldoutComposition)
+	{
+	case 0: return ETranslucentHoldoutCompositionType::BeforeDOF;
+	case 1: return ETranslucentHoldoutCompositionType::AfterDOF;
+	default: return ETranslucentHoldoutCompositionType::BeforeDOF;
+	}
+}
+
 void SetupDownsampledTranslucencyViewParameters(
 	const FViewInfo& View,
 	FIntPoint TextureExtent,
@@ -1695,7 +1740,6 @@ void FDeferredShadingSceneRenderer::RenderTranslucencyInner(
 			RDG_EVENT_SCOPE_CONDITIONAL(GraphBuilder, InViews.Num() > 1, "View%d", ViewIndex);
 			
 			const FScreenPassTextureViewport SeparateTranslucencyViewport = SeparateTranslucencyDimensions.GetInstancedStereoViewport(View);
-			const bool bCompositeBackToSceneColor = true;
 			const bool bLumenGIEnabled = Renderer.GetViewPipelineState(View).DiffuseIndirectMethod == EDiffuseIndirectMethod::Lumen;
 
 			/** Separate translucency color is either composited immediately or later during post processing. If done immediately, it's because the view doesn't support
@@ -1721,10 +1765,10 @@ void FDeferredShadingSceneRenderer::RenderTranslucencyInner(
 			}
 
 			// No scale is needed
-			const float ViewportScale = 1.0f;
-
+			constexpr float ViewportScale = 1.0f;
 			// Simply load as we have already initialized the texture.
-			const ERenderTargetLoadAction SeparateTranslucencyColorLoadAction = ERenderTargetLoadAction::ELoad;
+			constexpr ERenderTargetLoadAction SeparateTranslucencyColorLoadAction = ERenderTargetLoadAction::ELoad;
+			constexpr bool bResolveColorTexture = false;
 
 			FOITData OITData = OIT::CreateOITData(GraphBuilder, View, OITPass_SeparateTranslucency);
 
@@ -1739,7 +1783,7 @@ void FDeferredShadingSceneRenderer::RenderTranslucencyInner(
 				SeparateTranslucencyDepthTexture.Target,
 				CreateTranslucentBasePassUniformBuffer(GraphBuilder, Scene, View, ViewIndex, TranslucentLightingVolumeTextures, SceneColorCopyTexture, SceneTextureSetupMode, bLumenGIEnabled, OITData, TranslucencyPass),
 				TranslucencyPass,
-				!bCompositeBackToSceneColor,
+				bResolveColorTexture,
 				bRenderInParallel,
 				InstanceCullingManager);
 
@@ -1755,35 +1799,13 @@ void FDeferredShadingSceneRenderer::RenderTranslucencyInner(
 				OIT::AddOITComposePass(GraphBuilder, View, OITData, SeparateTranslucencyColorTexture.Target);
 			}
 
-			if (bCompositeBackToSceneColor)
-			{
-				FRDGTextureRef SeparateTranslucencyDepthResolve = nullptr;
-				FRDGTextureRef SceneDepthResolve = nullptr;
-				
-				{
-					::AddResolveSceneDepthPass(GraphBuilder, View, SeparateTranslucencyDepthTexture);
-
-					SeparateTranslucencyDepthResolve = SeparateTranslucencyDepthTexture.Resolve;
-					SceneDepthResolve = SceneTextures.Depth.Resolve;
-				}
+			::AddResolveSceneDepthPass(GraphBuilder, View, SeparateTranslucencyDepthTexture);
 
+			if (GetTranslucentHoldoutComposition() == ETranslucentHoldoutCompositionType::BeforeDOF)
+			{
 				FTranslucencyPassResources& TranslucencyPassResources = OutTranslucencyResourceMap->Get(ViewIndex, TranslucencyPass);
 
-				FTranslucencyComposition TranslucencyComposition;
-				TranslucencyComposition.Operation = FTranslucencyComposition::EOperation::ComposeToSceneColorAlpha;
-				TranslucencyComposition.SceneColor = FScreenPassTextureSlice::CreateFromScreenPassTexture(GraphBuilder, FScreenPassTexture(SceneTextures.Color.Target, View.ViewRect));
-				TranslucencyComposition.SceneDepth = FScreenPassTexture(SceneTextures.Depth.Resolve, View.ViewRect);
-				TranslucencyComposition.OutputViewport = FScreenPassTextureViewport(SceneTextures.Depth.Resolve, View.ViewRect);
-
-				FScreenPassTexture UpscaledTranslucency = TranslucencyComposition.AddPass(
-					GraphBuilder, View, TranslucencyPassResources);
-
-				ensure(View.ViewRect == UpscaledTranslucency.ViewRect);
-				ensure(UpscaledTranslucency.Texture == SceneTextures.Color.Target);
-
-				//Invalidate.
-				TranslucencyPassResources = FTranslucencyPassResources();
-				TranslucencyPassResources.Pass = TranslucencyPass;
+				AddTranslucentHoldoutToSceneColorAlphaPass(GraphBuilder, View, SceneTextures.Color.Target, SceneTextures.Depth.Resolve, TranslucencyPassResources);
 			}
 
 			++NumProcessedViews;
diff --git a/Engine/Source/Runtime/Renderer/Private/TranslucentRendering.h b/Engine/Source/Runtime/Renderer/Private/TranslucentRendering.h
--- a/Engine/Source/Runtime/Renderer/Private/TranslucentRendering.h
+++ b/Engine/Source/Runtime/Renderer/Private/TranslucentRendering.h
@@ -38,6 +38,15 @@ DECLARE_GPU_DRAWCALL_STAT_EXTERN(Translucency);
 /** Add Copy SceneColor Pass. */
 FRDGTextureRef AddCopySceneColorPass(FRDGBuilder& GraphBuilder, TArrayView<const FViewInfo> Views, FRDGTextureMSAA SceneColor, bool WithAlpha = false, bool bSkipIfUnderwater = true);
 
+/** Add pass to compose holdout to scene color alpha. */
+void AddTranslucentHoldoutToSceneColorAlphaPass(
+	FRDGBuilder& GraphBuilder,
+	const FViewInfo& View,
+	FRDGTextureRef SceneColorTexture,
+	FRDGTextureRef SceneDepthTexture,
+	FTranslucencyPassResources& TranslucencyPassResources
+);
+
 /** Converts the the translucency pass into the respective mesh pass. */
 EMeshPass::Type TranslucencyPassToMeshPass(ETranslucencyPass::Type TranslucencyPass);
 
@@ -88,3 +97,13 @@ struct FTranslucencyComposition
 		const FViewInfo& View,
 		const FTranslucencyPassResources& TranslucencyTextures) const;
 };
+
+/** Location of the translucent holdout composition into scene color alpha. */
+enum class ETranslucentHoldoutCompositionType : uint8
+{
+	BeforeDOF,
+	AfterDOF,
+};
+
+/** Get the location of the translucent holdout composition into scene color alpha. */
+ETranslucentHoldoutCompositionType GetTranslucentHoldoutComposition();
