Fix for Nanite scene extension crash when using mobile preview.  Need to check current shader platform rather than max shader platform, as Nanite scene extension is not present in mobile preview.

#rnx

#Synced-CL 50270751

diff --git a/Engine/Source/Runtime/Renderer/Private/PrimitiveSceneInfo.cpp b/Engine/Source/Runtime/Renderer/Private/PrimitiveSceneInfo.cpp
--- a/Engine/Source/Runtime/Renderer/Private/PrimitiveSceneInfo.cpp
+++ b/Engine/Source/Runtime/Renderer/Private/PrimitiveSceneInfo.cpp
@@ -879,7 +879,7 @@ void FPrimitiveSceneInfo::CacheNaniteMaterialBins(FScene* Scene, const TArrayVie
 	CSV_SCOPED_TIMING_STAT_EXCLUSIVE(FPrimitiveSceneInfo_CacheNaniteMaterialBins);
 	QUICK_SCOPE_CYCLE_COUNTER(STAT_CacheNaniteMaterialBins);
 
-	const bool bNaniteEnabled = DoesPlatformSupportNanite(GMaxRHIShaderPlatform);
+	const bool bNaniteEnabled = DoesPlatformSupportNanite(Scene->GetShaderPlatform());
 	if (bNaniteEnabled)
 	{
 		const bool bLumenEnabled = DoesPlatformSupportLumenGI(GetFeatureLevelShaderPlatform(Scene->GetFeatureLevel()));
