Fix absolute world position coordinate is queried at wrong coordiantes due to dynamic resolution.
1. Classify PPM passes by before and after UPSCALING and use it to guide the access of ResolutionFractionAndInv
2. For spatialUpscaling and RawOutput. Set ResolutionFractionAndInv to 1.0. As before and after temporal upscaling resolution is the same.

#tests All PPM blendable locations for Mobile and Deferred Spatial/Temporal Upscaler and Rawoutput

#Synced-CL 50214884

diff --git a/Engine/Plugins/Experimental/RemoteSession/Source/RemoteSession/Private/Channels/RemoteSessionARCameraChannel.cpp b/Engine/Plugins/Experimental/RemoteSession/Source/RemoteSession/Private/Channels/RemoteSessionARCameraChannel.cpp
--- a/Engine/Plugins/Experimental/RemoteSession/Source/RemoteSession/Private/Channels/RemoteSessionARCameraChannel.cpp
+++ b/Engine/Plugins/Experimental/RemoteSession/Source/RemoteSession/Private/Channels/RemoteSessionARCameraChannel.cpp
@@ -77,7 +77,11 @@ public:
 		FMaterialShader::ModifyCompilationEnvironment(Parameters, OutEnvironment);
 		OutEnvironment.SetDefine(TEXT("POST_PROCESS_MATERIAL"), 1);
 		OutEnvironment.SetDefine(TEXT("POST_PROCESS_MATERIAL_BEFORE_TONEMAP"), (Parameters.MaterialParameters.BlendableLocation != BL_SceneColorAfterTonemapping) ? 1 : 0);
-		OutEnvironment.SetDefine(TEXT("POST_PROCESS_MATERIAL_SSRINPUT"), (Parameters.MaterialParameters.BlendableLocation == BL_SSRInput) ? 1 : 0);
+
+		EBlendableLocation Location = EBlendableLocation(Parameters.MaterialParameters.BlendableLocation);
+		bool bPostProcessMaterialBeforeUpscaling = (Location == BL_SceneColorBeforeDOF) || (Location == BL_SceneColorAfterDOF) || (Location == BL_TranslucencyAfterDOF);
+		OutEnvironment.SetDefine(TEXT("POST_PROCESS_MATERIAL_BEFORE_UPSCALING"), bPostProcessMaterialBeforeUpscaling ? 1 : 0 );
+
 	}
 };
 
diff --git a/Engine/Plugins/Runtime/AR/AppleAR/AppleARKit/Source/AppleARKit/Private/AppleARKitVideoOverlay.cpp b/Engine/Plugins/Runtime/AR/AppleAR/AppleARKit/Source/AppleARKit/Private/AppleARKitVideoOverlay.cpp
--- a/Engine/Plugins/Runtime/AR/AppleAR/AppleARKit/Source/AppleARKit/Private/AppleARKitVideoOverlay.cpp
+++ b/Engine/Plugins/Runtime/AR/AppleAR/AppleARKit/Source/AppleARKit/Private/AppleARKitVideoOverlay.cpp
@@ -159,7 +159,10 @@ public:
 		FMaterialShader::ModifyCompilationEnvironment(Parameters, OutEnvironment);
 		OutEnvironment.SetDefine(TEXT("POST_PROCESS_MATERIAL"), 1);
 		OutEnvironment.SetDefine(TEXT("POST_PROCESS_MATERIAL_BEFORE_TONEMAP"), (Parameters.MaterialParameters.BlendableLocation != BL_SceneColorAfterTonemapping) ? 1 : 0);
-		OutEnvironment.SetDefine(TEXT("POST_PROCESS_MATERIAL_SSRINPUT"), (Parameters.MaterialParameters.BlendableLocation == BL_SSRInput) ? 1 : 0);
+		
+		EBlendableLocation Location = EBlendableLocation(Parameters.MaterialParameters.BlendableLocation);
+		bool bPostProcessMaterialBeforeUpscaling = (Location == BL_SceneColorBeforeDOF) || (Location == BL_SceneColorAfterDOF) || (Location == BL_TranslucencyAfterDOF);
+		OutEnvironment.SetDefine(TEXT("POST_PROCESS_MATERIAL_BEFORE_UPSCALING"), bPostProcessMaterialBeforeUpscaling ? 1 : 0 );
 		OutEnvironment.SetDefine(TEXT("POST_PROCESS_AR_PASSTHROUGH"), 1);
 
 		const bool bMobileForceDepthRead = MobileUsesFullDepthPrepass(Parameters.Platform);
diff --git a/Engine/Plugins/Runtime/AR/Google/GoogleARCore/Source/GoogleARCoreRendering/Private/GoogleARCorePassthroughCameraRenderer.cpp b/Engine/Plugins/Runtime/AR/Google/GoogleARCore/Source/GoogleARCoreRendering/Private/GoogleARCorePassthroughCameraRenderer.cpp
--- a/Engine/Plugins/Runtime/AR/Google/GoogleARCore/Source/GoogleARCoreRendering/Private/GoogleARCorePassthroughCameraRenderer.cpp
+++ b/Engine/Plugins/Runtime/AR/Google/GoogleARCore/Source/GoogleARCoreRendering/Private/GoogleARCorePassthroughCameraRenderer.cpp
@@ -108,7 +108,11 @@ public:
 		FMaterialShader::ModifyCompilationEnvironment(Parameters, OutEnvironment);
 		OutEnvironment.SetDefine(TEXT("POST_PROCESS_MATERIAL"), 1);
 		OutEnvironment.SetDefine(TEXT("POST_PROCESS_MATERIAL_BEFORE_TONEMAP"), (Parameters.MaterialParameters.BlendableLocation != BL_SceneColorAfterTonemapping) ? 1 : 0);
-		OutEnvironment.SetDefine(TEXT("POST_PROCESS_MATERIAL_SSRINPUT"), (Parameters.MaterialParameters.BlendableLocation == BL_SSRInput) ? 1 : 0);
+		
+		EBlendableLocation Location = EBlendableLocation(Parameters.MaterialParameters.BlendableLocation);
+		bool bPostProcessMaterialBeforeUpscaling = (Location == BL_SceneColorBeforeDOF) || (Location == BL_SceneColorAfterDOF) || (Location == BL_TranslucencyAfterDOF);
+		OutEnvironment.SetDefine(TEXT("POST_PROCESS_MATERIAL_BEFORE_UPSCALING"), bPostProcessMaterialBeforeUpscaling ? 1 : 0 );
+
 	}
 };
 
diff --git a/Engine/Shaders/Private/MaterialTemplate.ush b/Engine/Shaders/Private/MaterialTemplate.ush
--- a/Engine/Shaders/Private/MaterialTemplate.ush
+++ b/Engine/Shaders/Private/MaterialTemplate.ush
@@ -4735,11 +4735,11 @@ void CalcMaterialParametersPost(
 	FIsFrontFace bIsFrontFace)
 {
 	float4 ScreenPosition = SvPositionToScreenPosition(SvPosition);
-#if POST_PROCESS_MATERIAL_BEFORE_TONEMAP && !POST_PROCESS_MATERIAL_SSRINPUT	
+#if POST_PROCESS_MATERIAL_BEFORE_UPSCALING
 	// In this case, SVPositionToTranslatedWorld account for the correct dynamic resolution scaling.
 	float3 TranslatedWorldPosition = SvPositionToTranslatedWorld(SvPosition);
 #else
-	// Otherwise, we must update SVPosition to account for the resolution scaling not happening after at the tonemap stage and after.
+	// Otherwise, we must update SVPosition to account for the resolution scaling not happening after the upscaling stage.
 	float3 TranslatedWorldPosition = SvPositionToTranslatedWorld(SvPosition * float4(View.ResolutionFractionAndInv.xx, 1.0f, 1.0f));
 #endif
 
diff --git a/Engine/Source/Runtime/Renderer/Private/PostProcess/PostProcessMaterial.cpp b/Engine/Source/Runtime/Renderer/Private/PostProcess/PostProcessMaterial.cpp
--- a/Engine/Source/Runtime/Renderer/Private/PostProcess/PostProcessMaterial.cpp
+++ b/Engine/Source/Runtime/Renderer/Private/PostProcess/PostProcessMaterial.cpp
@@ -259,12 +259,11 @@ public:
 
 		EBlendableLocation Location = EBlendableLocation(Parameters.MaterialParameters.BlendableLocation);
 		OutEnvironment.SetDefine(TEXT("POST_PROCESS_MATERIAL_BEFORE_TONEMAP"), (Location == BL_SceneColorAfterTonemapping || Location == BL_ReplacingTonemapper) ? 0 : 1);
-		// Post process SSR is always rendered at native resolution as if it was after tone mapping, so we need to account for the fact that it is independent from DRS.
-		// SSR input should not be affected by exposure so it should be specified separately from POST_PROCESS_MATERIAL_BEFORE_TONEMAP 
-		// in order to be able to make DRS independent CameraVector and WorldPosition nodes.
-		OutEnvironment.SetDefine(TEXT("POST_PROCESS_MATERIAL_SSRINPUT"), (Location == BL_SSRInput) ? 1 : 0);
 		OutEnvironment.SetDefine(TEXT("POST_PROCESS_DISABLE_PRE_EXPOSURE_SCALE"), Parameters.MaterialParameters.bDisablePreExposureScale ? 1 : 0);
 
+		bool bPostProcessMaterialBeforeUpscaling = (Location == BL_SceneColorBeforeDOF) || (Location == BL_SceneColorAfterDOF) || (Location == BL_TranslucencyAfterDOF);
+		OutEnvironment.SetDefine(TEXT("POST_PROCESS_MATERIAL_BEFORE_UPSCALING"), bPostProcessMaterialBeforeUpscaling ? 1 : 0 );
+
 		if (IsMobilePlatform(Parameters.Platform))
 		{
 			OutEnvironment.SetDefine(TEXT("POST_PROCESS_MATERIAL_BEFORE_TONEMAP"), (Parameters.MaterialParameters.BlendableLocation != BL_SceneColorAfterTonemapping) ? 1 : 0);
diff --git a/Engine/Source/Runtime/Renderer/Private/SceneRendering.cpp b/Engine/Source/Runtime/Renderer/Private/SceneRendering.cpp
--- a/Engine/Source/Runtime/Renderer/Private/SceneRendering.cpp
+++ b/Engine/Source/Runtime/Renderer/Private/SceneRendering.cpp
@@ -1764,6 +1764,15 @@ void FViewInfo::SetupUniformBufferParameters(
 
 	{
 		float ResolutionFraction = float(ViewRect.Width()) / float(UnscaledViewRect.Width());
+		
+		// When the primary upscaler is not temporal AA, the resolution is the same between the rendering pass
+		// and the post processing pass after upscaling. No need to modify the ResolutionFractionAndInv.
+		// When ResolutionFraction > 1.0, it is used for other purpose, e.g., in Shadow. 
+		if (PrimaryScreenPercentageMethod != EPrimaryScreenPercentageMethod::TemporalUpscale
+			&& ResolutionFraction < 1.0f)
+		{
+			ResolutionFraction = 1.0f;
+		}
 
 		ViewUniformShaderParameters.ResolutionFractionAndInv.X = ResolutionFraction;
 		ViewUniformShaderParameters.ResolutionFractionAndInv.Y = 1.0f / ResolutionFraction;
