[VSM] Fixed culling view override propagation for non-nanite VSM which caused missmatch between CPU and GPU culling if r.Shadow.LODDistanceFactor was used & also made sure the selected views for local lights matches between CPU/GPU.
- Also fixed r.Shadow.LODDistanceFactor.CascadeScale calculation for local lights and VSM use-cases.

#Synced-CL 50202780

diff --git a/Engine/Source/Runtime/Renderer/Private/ShadowDepthRendering.cpp b/Engine/Source/Runtime/Renderer/Private/ShadowDepthRendering.cpp
--- a/Engine/Source/Runtime/Renderer/Private/ShadowDepthRendering.cpp
+++ b/Engine/Source/Runtime/Renderer/Private/ShadowDepthRendering.cpp
@@ -1257,7 +1257,21 @@ void FProjectedShadowInfo::RenderDepth(
 
 float FProjectedShadowInfo::GetLODDistanceFactor() const
 {
-	return CascadeSettings.bFarShadowCascade ? GFarShadowLODDistanceFactor : GShadowLODDistanceFactor + CascadeSettings.ShadowSplitIndex * GShadowLODDistanceFactorCascadeScale;
+	float Result = GFarShadowLODDistanceFactor;
+	if (CascadeSettings.bFarShadowCascade)
+	{
+		Result = GFarShadowLODDistanceFactor;
+	}
+	else
+	{
+		Result = GShadowLODDistanceFactor;
+		// Make sure values are in a reasonable range (Local lights set SplitIndex to -1 and VSM to 1000)
+		if (CascadeSettings.ShadowSplitIndex > 0 && CascadeSettings.ShadowSplitIndex < 16)
+		{
+			Result += CascadeSettings.ShadowSplitIndex * GShadowLODDistanceFactorCascadeScale;
+		}
+	}
+	return Result;
 }
 
 void FProjectedShadowInfo::ModifyViewForShadow(FViewInfo* FoundView) const
diff --git a/Engine/Source/Runtime/Renderer/Private/Shadows/ShadowSceneRenderer.cpp b/Engine/Source/Runtime/Renderer/Private/Shadows/ShadowSceneRenderer.cpp
--- a/Engine/Source/Runtime/Renderer/Private/Shadows/ShadowSceneRenderer.cpp
+++ b/Engine/Source/Runtime/Renderer/Private/Shadows/ShadowSceneRenderer.cpp
@@ -1190,6 +1190,7 @@ void FShadowSceneRenderer::CreateNaniteViewsForPass(
 				check(Shadow->bShouldRenderVSM);
 				
 				VirtualShadowMapArray.AddRenderViews(
+					FVirtualShadowMapArray::ERenderPassType::Nanite,
 					Shadow,
 					Views,
 					ShadowsLODScaleFactor,
diff --git a/Engine/Source/Runtime/Renderer/Private/VirtualShadowMaps/VirtualShadowMapArray.cpp b/Engine/Source/Runtime/Renderer/Private/VirtualShadowMaps/VirtualShadowMapArray.cpp
--- a/Engine/Source/Runtime/Renderer/Private/VirtualShadowMaps/VirtualShadowMapArray.cpp
+++ b/Engine/Source/Runtime/Renderer/Private/VirtualShadowMaps/VirtualShadowMapArray.cpp
@@ -4395,7 +4395,8 @@ void FVirtualShadowMapArray::RenderVirtualShadowMapsNonNanite(FRDGBuilder& Graph
 			TRACE_CPUPROFILER_EVENT_SCOPE(AddNonNaniteRenderViews);
 			for (FProjectedShadowInfo* ProjectedShadowInfo : ShadowsToAddRenderViews)
 			{
-				AddRenderViews(ProjectedShadowInfo, Views, 1.0f, bHasHZBTexture, false, OutShadowViews);
+				check(ProjectedShadowInfo->ShadowDepthView != nullptr);
+				AddRenderViews(ERenderPassType::NonNanite, ProjectedShadowInfo, Views, 1.0f, bHasHZBTexture, false, OutShadowViews);
 			}
 			CreateMipViews(OutShadowViews);
 		});
@@ -4842,8 +4843,8 @@ static Nanite::FPackedView CreateNanitePackedView(const Nanite::FPackedViewParam
 
 void FVirtualShadowMapArray::AddRenderViewsClipmap(
 	const TSharedPtr<FVirtualShadowMapClipmap>& Clipmap,
-	const FViewInfo* CullingView,
-	float LODScaleFactor,
+	const FViewInfo& CullingView,
+	float NaniteLODScaleFactor,
 	bool bSetHZBParams,
 	bool bUpdateHZBMetaData,
 	TArray<Nanite::FPackedView, SceneRenderingAllocator> &OutVirtualShadowViews) const
@@ -4852,7 +4853,7 @@ void FVirtualShadowMapArray::AddRenderViewsClipmap(
 	BaseParams.ViewRect = FIntRect(0, 0, FVirtualShadowMap::VirtualMaxResolutionXY, FVirtualShadowMap::VirtualMaxResolutionXY);
 	BaseParams.HZBTestViewRect = BaseParams.ViewRect;
 	BaseParams.RasterContextSize = GetPhysicalPoolSize();
-	BaseParams.MaxPixelsPerEdgeMultipler = 1.0f / LODScaleFactor;
+	BaseParams.MaxPixelsPerEdgeMultipler = 1.0f / NaniteLODScaleFactor;
 	BaseParams.PrevTargetLayerIndex = INDEX_NONE;
 	BaseParams.TargetMipLevel = 0;
 	BaseParams.TargetMipCount = 1;	// No mips for clipmaps
@@ -4892,7 +4893,7 @@ void FVirtualShadowMapArray::AddRenderViewsClipmap(
 		BaseParams.MinBoundsRadius = GMinScreenRadiusForShadowCaster;
 	}
 
-	Nanite::SetCullingViewOverrides(CullingView, BaseParams);
+	Nanite::SetCullingViewOverrides(&CullingView, BaseParams);
 
 	const int32 VirtualShadowMapId = Clipmap->GetVirtualShadowMapId();
 	for (int32 ClipmapLevelIndex = 0; ClipmapLevelIndex < Clipmap->GetLevelCount(); ++ClipmapLevelIndex)
@@ -4927,8 +4928,8 @@ void FVirtualShadowMapArray::AddRenderViewsClipmap(
 
 void FVirtualShadowMapArray::AddRenderViewsLocal(
 	const FProjectedShadowInfo* ProjectedShadowInfo,
-	TConstArrayView<FViewInfo> Views,
-	float LODScaleFactor,
+	TConstArrayView<FViewInfo> CullingViews,
+	float NaniteLODScaleFactor,
 	bool bSetHZBParams,
 	bool bUpdateHZBMetaData,
 	TArray<Nanite::FPackedView, SceneRenderingAllocator>& OutVirtualShadowViews) const
@@ -4937,7 +4938,7 @@ void FVirtualShadowMapArray::AddRenderViewsLocal(
 	BaseParams.ViewRect = FIntRect(0, 0, FVirtualShadowMap::VirtualMaxResolutionXY, FVirtualShadowMap::VirtualMaxResolutionXY);
 	BaseParams.HZBTestViewRect = BaseParams.ViewRect;
 	BaseParams.RasterContextSize = GetPhysicalPoolSize();
-	BaseParams.MaxPixelsPerEdgeMultipler = 1.0f / LODScaleFactor;
+	BaseParams.MaxPixelsPerEdgeMultipler = 1.0f / NaniteLODScaleFactor;
 	BaseParams.PrevTargetLayerIndex = INDEX_NONE;
 	BaseParams.TargetMipLevel = 0;
 	BaseParams.TargetMipCount = FVirtualShadowMap::MaxMipLevels;
@@ -4951,11 +4952,12 @@ void FVirtualShadowMapArray::AddRenderViewsLocal(
 
 	// Local lights, select the view closest to the local light to get some kind of reasonable behavior for split screen.
 	int32 ClosestCullingViewIndex = 0;
+	if (CullingViews.Num() > 1)
 	{
-		double MinDistanceSq = (Views[0].GetShadowViewMatrices().GetViewOrigin() + ProjectedShadowInfo->PreShadowTranslation).SquaredLength();
-		for (int Index = 1; Index < Views.Num(); ++Index)
+		double MinDistanceSq = (CullingViews[0].GetShadowViewMatrices().GetViewOrigin() + ProjectedShadowInfo->PreShadowTranslation).SquaredLength();
+		for (int Index = 1; Index < CullingViews.Num(); ++Index)
 		{
-			FVector TestOrigin = Views[Index].GetShadowViewMatrices().GetViewOrigin();
+			FVector TestOrigin = CullingViews[Index].GetShadowViewMatrices().GetViewOrigin();
 			double TestDistanceSq = (TestOrigin + ProjectedShadowInfo->PreShadowTranslation).SquaredLength();
 			if (TestDistanceSq < MinDistanceSq)
 			{
@@ -4976,7 +4978,7 @@ void FVirtualShadowMapArray::AddRenderViewsLocal(
 	BaseParams.Flags |= CacheEntry->IsUncached() ? NANITE_VIEW_FLAG_UNCACHED : 0U;
 	BaseParams.Flags |= CacheEntry->ShouldUseReceiverMask() ? NANITE_VIEW_FLAG_USE_RECEIVER_MASK : 0U;
 
-	Nanite::SetCullingViewOverrides(&Views[ClosestCullingViewIndex], BaseParams);
+	Nanite::SetCullingViewOverrides(&CullingViews[ClosestCullingViewIndex], BaseParams);
 	for (int32 Index = 0; Index < CacheEntry->ShadowMapEntries.Num(); ++Index)
 	{
 		FVirtualShadowMapCacheEntry& LevelEntry = CacheEntry->ShadowMapEntries[Index];
@@ -5002,23 +5004,30 @@ void FVirtualShadowMapArray::AddRenderViewsLocal(
 }
 
 void FVirtualShadowMapArray::AddRenderViews(
+	ERenderPassType RenderPassType,
 	const FProjectedShadowInfo* ProjectedShadowInfo,
-	TConstArrayView<FViewInfo> Views,
-	float LODScaleFactor,
+	TConstArrayView<FViewInfo> SceneRendererViews,
+	float NaniteLODScaleFactor,
 	bool bSetHZBParams,
 	bool bUpdateHZBMetaData,
 	TArray<Nanite::FPackedView, SceneRenderingAllocator>& OutVirtualShadowViews) const
 {
 	check(ProjectedShadowInfo->bWholeSceneShadow);
+	check(NaniteLODScaleFactor > 0.0f);
 
 	if (ProjectedShadowInfo->VirtualShadowMapClipmap)
 	{
-		check(ProjectedShadowInfo->DependentView != nullptr);
-		return AddRenderViewsClipmap(ProjectedShadowInfo->VirtualShadowMapClipmap, ProjectedShadowInfo->DependentView, LODScaleFactor, bSetHZBParams, bUpdateHZBMetaData, OutVirtualShadowViews);
+		// Use the shadow depth view to inherit the same overrides as the CPU culling logic for non-nanite rendering
+		// However, for Nanite we remove the dependency on the ShadowDepthView as that allows skipping the expensive FViewInfo copy for Nanite-only rendering
+		const FViewInfo* CullingView = RenderPassType == ERenderPassType::Nanite ? ProjectedShadowInfo->DependentView : ProjectedShadowInfo->ShadowDepthView;
+		check(CullingView);
+		return AddRenderViewsClipmap(ProjectedShadowInfo->VirtualShadowMapClipmap, *CullingView, NaniteLODScaleFactor, bSetHZBParams, bUpdateHZBMetaData, OutVirtualShadowViews);
 	}
 	else
 	{
-		return AddRenderViewsLocal(ProjectedShadowInfo, Views, LODScaleFactor, bSetHZBParams, bUpdateHZBMetaData, OutVirtualShadowViews);
+		// As above, AddRenderViewsLocal selects the closest view for Nanite shadow rendering.
+		TConstArrayView<FViewInfo> CullingViews = RenderPassType == ERenderPassType::Nanite ? SceneRendererViews : TConstArrayView<FViewInfo>(ProjectedShadowInfo->ShadowDepthView, 1);
+		return AddRenderViewsLocal(ProjectedShadowInfo, CullingViews, NaniteLODScaleFactor, bSetHZBParams, bUpdateHZBMetaData, OutVirtualShadowViews);
 	}
 }
 
diff --git a/Engine/Source/Runtime/Renderer/Private/VirtualShadowMaps/VirtualShadowMapArray.h b/Engine/Source/Runtime/Renderer/Private/VirtualShadowMaps/VirtualShadowMapArray.h
--- a/Engine/Source/Runtime/Renderer/Private/VirtualShadowMaps/VirtualShadowMapArray.h
+++ b/Engine/Source/Runtime/Renderer/Private/VirtualShadowMaps/VirtualShadowMapArray.h
@@ -292,6 +292,13 @@ enum class EVSMVisualizationPostPass
 class FVirtualShadowMapArray
 {
 public:	
+	enum class ERenderPassType
+	{
+		Nanite,
+		NonNanite,
+		NUM
+	};
+
 	FVirtualShadowMapArray(FScene& InScene);
 	~FVirtualShadowMapArray();
 
@@ -444,9 +451,10 @@ public:
 
 	// Add render views, and mark shadow maps as rendered for a given clipmap or set of VSMs
 	void AddRenderViews(
+		ERenderPassType RenderPassType,
 		const FProjectedShadowInfo* ProjectedShadowInfo,
-		TConstArrayView<FViewInfo> Views,
-		float LODScaleFactor,
+		TConstArrayView<FViewInfo> SceneRendererViews,
+		float NaniteLODScaleFactor,
 		bool bSetHzbParams,
 		bool bUpdateHZBMetaData,
 		TArray<Nanite::FPackedView,	SceneRenderingAllocator>& OutVirtualShadowViews) const;
@@ -555,16 +563,16 @@ private:
 
 	void AddRenderViewsLocal(
 		const FProjectedShadowInfo* ProjectedShadowInfo,
-		TConstArrayView<FViewInfo> Views,
-		float LODScaleFactor,
+		TConstArrayView<FViewInfo> CullingViews,
+		float NaniteLODScaleFactor,
 		bool bSetHzbParams,
 		bool bUpdateHZBMetaData,
 		TArray<Nanite::FPackedView,	SceneRenderingAllocator>& OutVirtualShadowViews) const;
 
 	void AddRenderViewsClipmap(
 		const TSharedPtr<FVirtualShadowMapClipmap>& Clipmap,
-		const FViewInfo* CullingView,
-		float LODScaleFactor,
+		const FViewInfo& CullingView,
+		float NaniteLODScaleFactor,
 		bool bSetHzbParams,
 		bool bUpdateHZBMetaData,
 		TArray<Nanite::FPackedView,
