Fix crash on shipping due to missing default UAV on VSMShadowProjectionCS shader with the uber shader permutation.

#Synced-CL 50298075

diff --git a/Engine/Source/Runtime/Renderer/Private/VirtualShadowMaps/VirtualShadowMapProjection.cpp b/Engine/Source/Runtime/Renderer/Private/VirtualShadowMaps/VirtualShadowMapProjection.cpp
--- a/Engine/Source/Runtime/Renderer/Private/VirtualShadowMaps/VirtualShadowMapProjection.cpp
+++ b/Engine/Source/Runtime/Renderer/Private/VirtualShadowMaps/VirtualShadowMapProjection.cpp
@@ -453,9 +453,7 @@ static void RenderVirtualShadowMapProjectionCommon(
 		GSystemTextures.SetDefaultIfNull(GraphBuilder, PassParameters->TileListData);
 		GSystemTextures.SetDefaultIfNull(GraphBuilder, PassParameters->IndirectDispatchArgs);
 		GSystemTextures.SetDefaultIfNull(GraphBuilder, PassParameters->OutStatsBuffer);
-		#if VSM_ENABLE_VISUALIZATION
 		GSystemTextures.SetDefaultIfNull(GraphBuilder, PassParameters->OutVisualize);
-		#endif
 
 		if (!bHasHairStrandsData)
 		{
