Fix incorrect element size for dummy hair voxel node buffer.

#Synced-CL 50298075

diff --git a/Engine/Source/Runtime/Renderer/Private/HairStrands/HairStrandsRendering.cpp b/Engine/Source/Runtime/Renderer/Private/HairStrands/HairStrandsRendering.cpp
--- a/Engine/Source/Runtime/Renderer/Private/HairStrands/HairStrandsRendering.cpp
+++ b/Engine/Source/Runtime/Renderer/Private/HairStrands/HairStrandsRendering.cpp
@@ -314,7 +314,7 @@ TRDGUniformBufferRef<FHairStrandsViewUniformParameters> BindHairStrandsViewUnifo
 TRDGUniformBufferRef<FVirtualVoxelParameters> CreateDefaultHairStrandsVoxelUniformBuffer(FRDGBuilder& GraphBuilder)
 {
 	FRDGBufferSRVRef DummyBuffer = GraphBuilder.CreateSRV(GSystemTextures.GetDefaultBuffer(GraphBuilder, 4u, 0u), PF_R32_UINT);
-	FRDGBufferSRVRef DummyStructuredBuffer = GraphBuilder.CreateSRV(GSystemTextures.GetDefaultStructuredBuffer(GraphBuilder, 4u, 0u), PF_R32_UINT);
+	FRDGBufferSRVRef DummyNodeDescBuffer = GraphBuilder.CreateSRV(GSystemTextures.GetDefaultStructuredBuffer(GraphBuilder, sizeof(FPackedVirtualVoxelNodeDesc), 0u));
 
 	FVirtualVoxelParameters* Out = GraphBuilder.AllocParameters<FVirtualVoxelParameters>();
 	Out->Common.PageCountResolution = FIntVector::ZeroValue;
@@ -349,7 +349,7 @@ TRDGUniformBufferRef<FVirtualVoxelParameters> CreateDefaultHairStrandsVoxelUnifo
 	Out->Common.AllocatedPageCountBuffer = DummyBuffer;
 	Out->Common.PageIndexBuffer = DummyBuffer;
 	Out->Common.PageIndexCoordBuffer = DummyBuffer;
-	Out->Common.NodeDescBuffer = DummyStructuredBuffer;
+	Out->Common.NodeDescBuffer = DummyNodeDescBuffer;
 	Out->Common.CurrGPUMinVoxelSize = DummyBuffer;
 	Out->Common.NextGPUMinVoxelSize = DummyBuffer;
 	Out->PageTexture = GSystemTextures.GetVolumetricBlackDummy(GraphBuilder);
