[NFOR] Default to higher clarity than ghosting. Set r.PathTracing.Denoiser.Prepass.RankedLuminanceVariance back to 1 if the previous config is preferred.

#Synced-CL 50292658

diff --git a/Engine/Source/Runtime/Renderer/Private/PathTracingSpatialTemporalDenoising.cpp b/Engine/Source/Runtime/Renderer/Private/PathTracingSpatialTemporalDenoising.cpp
--- a/Engine/Source/Runtime/Renderer/Private/PathTracingSpatialTemporalDenoising.cpp
+++ b/Engine/Source/Runtime/Renderer/Private/PathTracingSpatialTemporalDenoising.cpp
@@ -160,10 +160,10 @@ namespace {
 
 	TAutoConsoleVariable<int32> CVarPathTracingDenoiserPrepassRankedLuminanceVariance(
 		TEXT("r.PathTracing.Denoiser.Prepass.RankedLuminanceVariance"),
-		1,
+		0,
 		TEXT("Select the luminance type when calculating the variance:\n")
-		TEXT("0: use default luminance to estimate variance.\n")
-		TEXT("1: Use channel ranked luminance when calculating variance.\n"),
+		TEXT("0: use default luminance to estimate variance (Most clarity but may preserve some noise if variance is too high).\n")
+		TEXT("1: Use channel ranked luminance when calculating variance (better in high variance and low motion scene, e.g., dawn, dusk, or low lighting).\n"),
 		ECVF_RenderThreadSafe
 	);
 
