[HWRT] Fix race condition between AddDynamicInstancesToRayTracingScene(...) and GatherRelevantStaticPrimitives(...)
- FPrimitiveSceneInfo::IsCachedRayTracingGeometryValid(...) was reading from Instances arrays in FRayTracingScene while those arrays could be resized by AddDynamicInstancesToRayTracingScene(...).
- Moved some checks later to AddVisibleCachedInstances(...) when access to FRayTracingScene is serialized using AddInstancesPipe.

#Synced-CL 50307718

diff --git a/Engine/Source/Runtime/Renderer/Private/PrimitiveSceneInfo.cpp b/Engine/Source/Runtime/Renderer/Private/PrimitiveSceneInfo.cpp
--- a/Engine/Source/Runtime/Renderer/Private/PrimitiveSceneInfo.cpp
+++ b/Engine/Source/Runtime/Renderer/Private/PrimitiveSceneInfo.cpp
@@ -392,47 +392,6 @@ bool FPrimitiveSceneInfo::IsCachedRayTracingGeometryValid(int8 LODIndex, int32 I
 
 		check(!CachedRayTracingGeometry->GetRequiresBuild() && !CachedRayTracingGeometry->HasPendingBuildRequest());
 
-#if DO_CHECK
-		auto LogErrorDebugInfo = [this, CachedRayTracingGeometry](uint32 RayTracingInstanceIndex)
-			{
-				FRHIRayTracingGeometry* CachedInstanceGeometryRHI = Scene->RayTracingScene.GetCachedInstanceGeometry(RayTracingInstanceIndex);
-				FRHIRayTracingGeometry* CachedRayTracingGeometryRHI = CachedRayTracingGeometry->GetRHI();
-				UE_LOG(LogRenderer, Error, TEXT("IsCachedRayTracingGeometryValid() Error Debug Info \
-					PrimitiveRayTracingFlags: %u \
-					bIsCachedRayTracingInstanceValid: %d \
-					bCachedInstanceHasGeometryRHI: %d, \
-					bCachedRayTracingGeometryHasRHI: %d, \
-					bCachedRayTracingGeometryIsValid: %d, \
-					bCachedRayTracingGeometryIsEvicted: %d, \
-					CachedRayTracingGeometryFlags: %u"),
-					Scene->PrimitiveRayTracingFlags[GetIndex()],
-					bIsCachedRayTracingInstanceValid ? 1 : 0,
-					CachedInstanceGeometryRHI != nullptr ? 1 : 0,
-					CachedRayTracingGeometryRHI != nullptr ? 1 : 0,
-					CachedRayTracingGeometry->IsValid() ? 1 : 0,
-					CachedRayTracingGeometry->IsEvicted() ? 1 : 0,
-					uint32(CachedRayTracingGeometry->GetGeometryState()));
-			};
-
-		if (InstanceData.RayTracingInstanceIndexMain != UINT32_MAX)
-		{
-			if (Scene->RayTracingScene.GetCachedInstanceGeometry(InstanceData.RayTracingInstanceIndexMain) != CachedRayTracingGeometry->GetRHI())
-			{
-				LogErrorDebugInfo(InstanceData.RayTracingInstanceIndexMain);
-				checkf(false, TEXT("Cached ray tracing instance has an unexpected geometry. Check log above."))
-			}
-		}
-
-		if (InstanceData.RayTracingInstanceIndexDecal != UINT32_MAX)
-		{
-			if (Scene->RayTracingScene.GetCachedInstanceGeometry(InstanceData.RayTracingInstanceIndexDecal) != CachedRayTracingGeometry->GetRHI())
-			{
-				LogErrorDebugInfo(InstanceData.RayTracingInstanceIndexDecal);
-				checkf(false, TEXT("Cached ray tracing instance has an unexpected geometry. Check log above."))
-			}
-		}
-#endif // DO_CHECK
-
 		return CachedRayTracingGeometry->IsValid() && !CachedRayTracingGeometry->IsEvicted();
 	}
 
diff --git a/Engine/Source/Runtime/Renderer/Private/RayTracing/RayTracing.cpp b/Engine/Source/Runtime/Renderer/Private/RayTracing/RayTracing.cpp
--- a/Engine/Source/Runtime/Renderer/Private/RayTracing/RayTracing.cpp
+++ b/Engine/Source/Runtime/Renderer/Private/RayTracing/RayTracing.cpp
@@ -249,6 +249,9 @@ namespace RayTracing
 		int32 SBTAllocationUniqueId = INDEX_NONE;
 		uint32 MainRayTracingInstanceIndex = UINT32_MAX;
 		uint32 DecalRayTracingInstanceIndex = UINT32_MAX;
+#if DO_CHECK
+		FRHIRayTracingGeometry* RayTracingGeometry = nullptr;
+#endif
 	};
 
 	struct FRelevantPrimitive
@@ -1491,6 +1494,13 @@ namespace RayTracing
 						RelevantPrimitive->SBTAllocationUniqueId = RTInstanceData.SBTAllocationUniqueId;
 						RelevantPrimitive->MainRayTracingInstanceIndex = RTInstanceData.RayTracingInstanceIndexMain;
 						RelevantPrimitive->DecalRayTracingInstanceIndex = RTInstanceData.RayTracingInstanceIndexDecal;
+#if DO_CHECK
+						if (!bUsingNaniteRayTracing)
+						{
+							// Currently RTInstanceData.RayTracingGeometry is only available for non-nanite geometries
+							RelevantPrimitive->RayTracingGeometry = RTInstanceData.RayTracingGeometry->GetRHI();
+						}
+#endif
 
 						const int32 NumBindings = RTInstanceData.CachedShaderBindingDataBase.Num() + RTInstanceData.CachedShaderBindingDataDecal.Num();
 
@@ -1867,11 +1877,17 @@ namespace RayTracing
 
 			if (RelevantPrimitive.MainRayTracingInstanceIndex != UINT32_MAX)
 			{
+#if DO_CHECK
+				check(RelevantPrimitive.RayTracingGeometry == nullptr || RayTracingScene.GetCachedInstanceGeometry(RelevantPrimitive.MainRayTracingInstanceIndex) == RelevantPrimitive.RayTracingGeometry);
+#endif
 				RayTracingScene.MarkInstanceVisible(RelevantPrimitive.MainRayTracingInstanceIndex, View.GetRayTracingSceneViewHandle());
 			}
 
 			if (RelevantPrimitive.DecalRayTracingInstanceIndex != UINT32_MAX && !bExcludeDecals)
 			{
+#if DO_CHECK
+				check(RelevantPrimitive.RayTracingGeometry == nullptr || RayTracingScene.GetCachedInstanceGeometry(RelevantPrimitive.DecalRayTracingInstanceIndex) == RelevantPrimitive.RayTracingGeometry);
+#endif
 				RayTracingScene.MarkInstanceVisible(RelevantPrimitive.DecalRayTracingInstanceIndex, View.GetRayTracingSceneViewHandle());
 			}
 
