Eliminate IsInverted and associated defines given that the assumption is so ingrained in the engine that changing it is not something that would ever happen.

#Synced-CL 50235042

diff --git a/Engine/Plugins/Experimental/Water/Source/Runtime/Private/WaterInfoRendering.cpp b/Engine/Plugins/Experimental/Water/Source/Runtime/Private/WaterInfoRendering.cpp
--- a/Engine/Plugins/Experimental/Water/Source/Runtime/Private/WaterInfoRendering.cpp
+++ b/Engine/Plugins/Experimental/Water/Source/Runtime/Private/WaterInfoRendering.cpp
@@ -286,8 +286,6 @@ static void FinalizeWaterInfo(
 
 static FMatrix BuildOrthoMatrix(float InOrthoWidth, float InOrthoHeight)
 {
-	check((int32)ERHIZBuffer::IsInverted);
-
 	const FMatrix::FReal OrthoWidth = InOrthoWidth / 2.0f;
 	const FMatrix::FReal OrthoHeight = InOrthoHeight / 2.0f;
 
diff --git a/Engine/Plugins/Experimental/Water/Source/Runtime/Private/WaterMeshSceneProxy.cpp b/Engine/Plugins/Experimental/Water/Source/Runtime/Private/WaterMeshSceneProxy.cpp
--- a/Engine/Plugins/Experimental/Water/Source/Runtime/Private/WaterMeshSceneProxy.cpp
+++ b/Engine/Plugins/Experimental/Water/Source/Runtime/Private/WaterMeshSceneProxy.cpp
@@ -1359,8 +1359,6 @@ void FViewWaterQuadTree::BuildGPUQuadTree(FRDGBuilder& GraphBuilder)
 {
 	auto BuildOrthoMatrix = [](float InOrthoWidth, float InOrthoHeight, float DepthRange)
 	{
-		check((int32)ERHIZBuffer::IsInverted);
-
 		const FMatrix::FReal OrthoWidth = InOrthoWidth / 2.0f;
 		const FMatrix::FReal OrthoHeight = InOrthoHeight / 2.0f;
 
diff --git a/Engine/Plugins/Experimental/Water/Source/Runtime/Private/WaterSubsystem.cpp b/Engine/Plugins/Experimental/Water/Source/Runtime/Private/WaterSubsystem.cpp
--- a/Engine/Plugins/Experimental/Water/Source/Runtime/Private/WaterSubsystem.cpp
+++ b/Engine/Plugins/Experimental/Water/Source/Runtime/Private/WaterSubsystem.cpp
@@ -830,7 +830,7 @@ void UWaterSubsystem::ComputeUnderwaterPostProcess(FVector ViewLocation, FSceneV
 	}
 
 	// Compute distance from view origin to the corner of the near plane. This distance needs to be taken into account when computing whether the view intersects the water surface.
-	const FVector4f NearPlaneCornerViewSpace = FVector4f(SceneView->ViewMatrices.GetInvProjectionMatrix().TransformFVector4(FVector4(1.0f, 1.0f, (bool)ERHIZBuffer::IsInverted ? 1.0f : 0.0f, 1.0f)));
+	const FVector4f NearPlaneCornerViewSpace = FVector4f(SceneView->ViewMatrices.GetInvProjectionMatrix().TransformFVector4(FVector4(1.0f, 1.0f, 1.0f, 1.0f)));
 	const float ViewToNearPlaneCornerDistance = FVector2f(NearPlaneCornerViewSpace / NearPlaneCornerViewSpace.W).Length();
 	bool bAnyDefinitelyUnderwater = false;
 	bool bAnyPossiblyUnderwater = false;
diff --git a/Engine/Plugins/MovieScene/MoviePipelineMaskRenderPass/Source/MoviePipelineMaskRenderPass/Private/MoviePipelinePanoramicPass.cpp b/Engine/Plugins/MovieScene/MoviePipelineMaskRenderPass/Source/MoviePipelineMaskRenderPass/Private/MoviePipelinePanoramicPass.cpp
--- a/Engine/Plugins/MovieScene/MoviePipelineMaskRenderPass/Source/MoviePipelineMaskRenderPass/Private/MoviePipelinePanoramicPass.cpp
+++ b/Engine/Plugins/MovieScene/MoviePipelineMaskRenderPass/Source/MoviePipelineMaskRenderPass/Private/MoviePipelinePanoramicPass.cpp
@@ -357,8 +357,6 @@ FSceneView* UMoviePipelinePanoramicPass::GetSceneViewForSampleState(FSceneViewFa
 		const float MatrixFOV = FMath::Max(0.001f, ViewFOV) * (float)PI / 360.0f;
 		// ToDo: I think this is a FMath::DegreesToRadians, easier to read that way than PI/360
 
-		static_assert((int32)ERHIZBuffer::IsInverted != 0, "ZBuffer should be inverted");
-
 		float XAxisMultiplier = 1.f;
 		float YAxisMultiplier = 1.f;
 		if (PaneSizeX > PaneSizeY)
diff --git a/Engine/Plugins/MovieScene/MovieRenderPipeline/Source/MovieRenderPipelineRenderPasses/Private/Graph/Renderers/MovieGraphDeferredPanoramicPass.cpp b/Engine/Plugins/MovieScene/MovieRenderPipeline/Source/MovieRenderPipelineRenderPasses/Private/Graph/Renderers/MovieGraphDeferredPanoramicPass.cpp
--- a/Engine/Plugins/MovieScene/MovieRenderPipeline/Source/MovieRenderPipelineRenderPasses/Private/Graph/Renderers/MovieGraphDeferredPanoramicPass.cpp
+++ b/Engine/Plugins/MovieScene/MovieRenderPipeline/Source/MovieRenderPipelineRenderPasses/Private/Graph/Renderers/MovieGraphDeferredPanoramicPass.cpp
@@ -625,8 +625,6 @@ void FMovieGraphDeferredPanoramicPass::Render(const FMovieGraphTraversalContext&
 				// ToDo: I think this is a FMath::DegreesToRadians, easier to read that way than PI/360
 				Data.NearClippingPlane = MinZ;
 
-				static_assert((int32)ERHIZBuffer::IsInverted != 0, "ZBuffer should be inverted");
-
 				float XAxisMultiplier = 1.f;
 				float YAxisMultiplier = 1.f;
 				SceneViewInitOptions.ProjectionMatrix = FReversedZPerspectiveMatrix(
diff --git a/Engine/Plugins/Runtime/MeshModelingToolset/Source/ModelingComponents/Private/Scene/WorldRenderCapture.cpp b/Engine/Plugins/Runtime/MeshModelingToolset/Source/ModelingComponents/Private/Scene/WorldRenderCapture.cpp
--- a/Engine/Plugins/Runtime/MeshModelingToolset/Source/ModelingComponents/Private/Scene/WorldRenderCapture.cpp
+++ b/Engine/Plugins/Runtime/MeshModelingToolset/Source/ModelingComponents/Private/Scene/WorldRenderCapture.cpp
@@ -458,7 +458,6 @@ FSceneViewInitOptions MakeSceneViewInitOptions(
 	FMatrix ViewRotationMatrix = FRotationMatrix::Make((FQuat)RenderFrame.Rotation.Inverse());
 
 	const float HalfFOVRadians = FMath::DegreesToRadians<float>(HorzFOVDegrees) * 0.5f;
-	static_assert((int32)ERHIZBuffer::IsInverted != 0, "Check NearPlane and Projection Matrix");
 	FMatrix ProjectionMatrix = FReversedZPerspectiveMatrix(HalfFOVRadians, 1.0f, 1.0f, (float)NearPlaneDist);
 
 	FSceneViewInitOptions ViewInitOptions;
diff --git a/Engine/Plugins/Runtime/nDisplay/Source/DisplayCluster/Private/Game/EngineClasses/Scene/DisplayClusterStageGeometryComponent.cpp b/Engine/Plugins/Runtime/nDisplay/Source/DisplayCluster/Private/Game/EngineClasses/Scene/DisplayClusterStageGeometryComponent.cpp
--- a/Engine/Plugins/Runtime/nDisplay/Source/DisplayCluster/Private/Game/EngineClasses/Scene/DisplayClusterStageGeometryComponent.cpp
+++ b/Engine/Plugins/Runtime/nDisplay/Source/DisplayCluster/Private/Game/EngineClasses/Scene/DisplayClusterStageGeometryComponent.cpp
@@ -325,26 +325,13 @@ void UDisplayClusterStageGeometryComponent::GetSceneViewInitOptions(const FVecto
 		const float XAxisMultiplier = 1.0f;
 		const float YAxisMultiplier = 1.0f;
 
-		if ((bool)ERHIZBuffer::IsInverted)
-		{
-			OutViewInitOptions.ProjectionMatrix = FReversedZPerspectiveMatrix(
-				MatrixFOV,
-				MatrixFOV,
-				XAxisMultiplier,
-				YAxisMultiplier,
-				MinZ,
-				MaxZ);
-		}
-		else
-		{
-			OutViewInitOptions.ProjectionMatrix = FPerspectiveMatrix(
-				MatrixFOV,
-				MatrixFOV,
-				XAxisMultiplier,
-				YAxisMultiplier,
-				MinZ,
-				MaxZ);
-		}
+		OutViewInitOptions.ProjectionMatrix = FReversedZPerspectiveMatrix(
+			MatrixFOV,
+			MatrixFOV,
+			XAxisMultiplier,
+			YAxisMultiplier,
+			MinZ,
+			MaxZ);
 
 		OutViewInitOptions.FOV = GeometryMapFOV;
 	}
diff --git a/Engine/Plugins/Runtime/nDisplay/Source/DisplayClusterLightCardEditor/Private/Viewport/DisplayClusterLightCardEditorViewportClient.cpp b/Engine/Plugins/Runtime/nDisplay/Source/DisplayClusterLightCardEditor/Private/Viewport/DisplayClusterLightCardEditorViewportClient.cpp
--- a/Engine/Plugins/Runtime/nDisplay/Source/DisplayClusterLightCardEditor/Private/Viewport/DisplayClusterLightCardEditorViewportClient.cpp
+++ b/Engine/Plugins/Runtime/nDisplay/Source/DisplayClusterLightCardEditor/Private/Viewport/DisplayClusterLightCardEditorViewportClient.cpp
@@ -2264,26 +2264,13 @@ void FDisplayClusterLightCardEditorViewportClient::GetNormalMapSceneViewInitOpti
 	const float XAxisMultiplier = 1.0f;
 	const float YAxisMultiplier = 1.0f;
 
-	if ((bool)ERHIZBuffer::IsInverted)
-	{
-		OutViewInitOptions.ProjectionMatrix = FReversedZPerspectiveMatrix(
-			MatrixFOV,
-			MatrixFOV,
-			XAxisMultiplier,
-			YAxisMultiplier,
-			MinZ,
-			MaxZ);
-	}
-	else
-	{
-		OutViewInitOptions.ProjectionMatrix = FPerspectiveMatrix(
-			MatrixFOV,
-			MatrixFOV,
-			XAxisMultiplier,
-			YAxisMultiplier,
-			MinZ,
-			MaxZ);
-	}
+	OutViewInitOptions.ProjectionMatrix = FReversedZPerspectiveMatrix(
+		MatrixFOV,
+		MatrixFOV,
+		XAxisMultiplier,
+		YAxisMultiplier,
+		MinZ,
+		MaxZ);
 
 	OutViewInitOptions.SceneViewStateInterface = ViewState.GetReference();
 	OutViewInitOptions.ViewElementDrawer = this;
diff --git a/Engine/Plugins/Runtime/nDisplay/Source/DisplayClusterScenePreview/Private/DisplayClusterLightCardEditorHelper.cpp b/Engine/Plugins/Runtime/nDisplay/Source/DisplayClusterScenePreview/Private/DisplayClusterLightCardEditorHelper.cpp
--- a/Engine/Plugins/Runtime/nDisplay/Source/DisplayClusterScenePreview/Private/DisplayClusterLightCardEditorHelper.cpp
+++ b/Engine/Plugins/Runtime/nDisplay/Source/DisplayClusterScenePreview/Private/DisplayClusterLightCardEditorHelper.cpp
@@ -686,49 +686,23 @@ void FDisplayClusterLightCardEditorHelper::GetSceneViewInitOptions(
 		const float OrthoWidth = 0.5f * FOVScale * ViewportSize.X;
 		const float OrthoHeight = 0.5f * FOVScale * ViewportSize.Y;
 
-		if ((bool)ERHIZBuffer::IsInverted)
-		{
-			OutViewInitOptions.ProjectionMatrix = FReversedZOrthoMatrix(
-				OrthoWidth,
-				OrthoHeight,
-				ZScale,
-				ZOffset
-			);
-		}
-		else
-		{
-			OutViewInitOptions.ProjectionMatrix = FOrthoMatrix(
-				OrthoWidth,
-				OrthoHeight,
-				ZScale,
-				ZOffset
-			);
-		}
+		OutViewInitOptions.ProjectionMatrix = FReversedZOrthoMatrix(
+			OrthoWidth,
+			OrthoHeight,
+			ZScale,
+			ZOffset
+		);
 	}
 	else
 	{
-		if ((bool)ERHIZBuffer::IsInverted)
-		{
-			OutViewInitOptions.ProjectionMatrix = FReversedZPerspectiveMatrix(
-				HalfFOV,
-				HalfFOV,
-				XAxisMultiplier,
-				YAxisMultiplier,
-				InNearClipPlane,
-				InNearClipPlane
-			);
-		}
-		else
-		{
-			OutViewInitOptions.ProjectionMatrix = FPerspectiveMatrix(
-				HalfFOV,
-				HalfFOV,
-				XAxisMultiplier,
-				YAxisMultiplier,
-				InNearClipPlane,
-				InNearClipPlane
-			);
-		}
+		OutViewInitOptions.ProjectionMatrix = FReversedZPerspectiveMatrix(
+			HalfFOV,
+			HalfFOV,
+			XAxisMultiplier,
+			YAxisMultiplier,
+			InNearClipPlane,
+			InNearClipPlane
+		);
 	}
 
 	if (!OutViewInitOptions.IsValidViewRectangle())
diff --git a/Engine/Plugins/Runtime/nDisplay/Source/DisplayClusterShaders/Private/Shaders/DisplayClusterShadersPreprocess_UVLightCards.cpp b/Engine/Plugins/Runtime/nDisplay/Source/DisplayClusterShaders/Private/Shaders/DisplayClusterShadersPreprocess_UVLightCards.cpp
--- a/Engine/Plugins/Runtime/nDisplay/Source/DisplayClusterShaders/Private/Shaders/DisplayClusterShadersPreprocess_UVLightCards.cpp
+++ b/Engine/Plugins/Runtime/nDisplay/Source/DisplayClusterShaders/Private/Shaders/DisplayClusterShadersPreprocess_UVLightCards.cpp
@@ -275,24 +275,12 @@ bool FDisplayClusterShadersPreprocess_UVLightCards::RenderPreprocess_UVLightCard
 	const float ZOffset = UE_OLD_HALF_WORLD_MAX;
 	const float OrthoSize = 0.5f * InParameters.ProjectionPlaneSize;
 
-	if ((bool)ERHIZBuffer::IsInverted)
-	{
-		ViewInitOptions.ProjectionMatrix = FReversedZOrthoMatrix(
-			OrthoSize,
-			OrthoSize,
-			ZScale,
-			ZOffset
-		);
-	}
-	else
-	{
-		ViewInitOptions.ProjectionMatrix = FOrthoMatrix(
-			OrthoSize,
-			OrthoSize,
-			ZScale,
-			ZOffset
-		);
-	}
+	ViewInitOptions.ProjectionMatrix = FReversedZOrthoMatrix(
+		OrthoSize,
+		OrthoSize,
+		ZScale,
+		ZOffset
+	);
 
 	ViewInitOptions.BackgroundColor = FLinearColor::Black;
 
diff --git a/Engine/Shaders/Private/Common.ush b/Engine/Shaders/Private/Common.ush
--- a/Engine/Shaders/Private/Common.ush
+++ b/Engine/Shaders/Private/Common.ush
@@ -111,19 +111,15 @@ FloatDeriv4 ConstructFloatDeriv4(float4 InValue, float4 InDdx, float4 InDdy)
 #define KILOMETER_TO_CENTIMETER	(KILOMETER_TO_METER * METER_TO_CENTIMETER)
 #define CENTIMETER_TO_KILOMETER	(1.0f / KILOMETER_TO_CENTIMETER)
 
-#define NearDepthValue (HAS_INVERTED_Z_BUFFER ? 1.0f : 0.0f)
-#define FarDepthValue  (HAS_INVERTED_Z_BUFFER ? 0.0f : 1.0f)
+#define NearDepthValue 1.0f
+#define FarDepthValue  0.0f
 
 float NearestDeviceDepth(float DepthA, 
                          float DepthB,
                          float DepthC = FarDepthValue
 )
 {
-#if HAS_INVERTED_Z_BUFFER
 	return max3(DepthA, DepthB, DepthC);
-#else
-	return min3(DepthA, DepthB, DepthC);
-#endif
 }
 
 float FarthestDeviceDepth(float DepthA, 
@@ -131,11 +127,7 @@ float FarthestDeviceDepth(float DepthA,
                           float DepthC = NearDepthValue
 )
 {
-#if HAS_INVERTED_Z_BUFFER
-	return min3(DepthA, DepthB, DepthC);
-#else
 	return max3(DepthA, DepthB, DepthC);
-#endif
 }
 
 const static MaterialFloat PI = 3.1415926535897932f;
@@ -1392,16 +1384,11 @@ float2 AdjustClipToPrevClipForProjectionType(float2 Velocity, float DeviceZ)
 		* Motion vectors on far distance planes are greatly influenced by slight rotations in camera angle
 		* And with Orthographic cameras, this magnifies the velocities over large distances.
 		* As a result, we need to counter this by reducing the speed with the DeviceZ
-		*/
-#if HAS_INVERTED_Z_BUFFER
-		/**
+		*
 		* Because this is ClipToPrevClip rather than velocity, we use the non - linear device z which works better than linear depth,
 		* however if adjusting velocity directly, linear depth should be used instead. 
 		*/
 		Velocity *= DeviceZ;
-#else
-		Velocity *= 1.0f - DeviceZ;
-#endif
 	}
 	return Velocity;
 }
diff --git a/Engine/Shaders/Private/DistortAccumulatePS.usf b/Engine/Shaders/Private/DistortAccumulatePS.usf
--- a/Engine/Shaders/Private/DistortAccumulatePS.usf
+++ b/Engine/Shaders/Private/DistortAccumulatePS.usf
@@ -43,11 +43,7 @@
 
 void ClipOccludedPixel(float PixelZ, float SceneZ)
 {
-#if	HAS_INVERTED_Z_BUFFER
 	float ClipDepth = PixelZ - SceneZ;
-#else
-	float ClipDepth = SceneZ - PixelZ;
-#endif
 	clip(ClipDepth);
 }
 
diff --git a/Engine/Shaders/Private/DistortFiltering.usf b/Engine/Shaders/Private/DistortFiltering.usf
--- a/Engine/Shaders/Private/DistortFiltering.usf
+++ b/Engine/Shaders/Private/DistortFiltering.usf
@@ -29,9 +29,7 @@ void CopySceneColorDepthTexturePS(FScreenVertexOutput Input, out float4 OutColor
 	LinearColor.rgb = min(LinearColor.rgb, MaxExposedLuminance.xxx);
 
 	float DeviceZ = Texture2DSample(SceneDepthTexture, SceneDepthSampler, TextureUV).r;
-#if HAS_INVERTED_Z_BUFFER
 	DeviceZ = max(0.000000000001, DeviceZ);	// TODO: avoid bad values when DeviceZ is far=0 when using inverted z
-#endif
 	float4 SVPosition = float4(SVPos, DeviceZ, 1);
 	float3 TranslatedWorldPosition = SvPositionToTranslatedWorld(SVPosition);
 
diff --git a/Engine/Shaders/Private/DownsampleDepthPixelShader.usf b/Engine/Shaders/Private/DownsampleDepthPixelShader.usf
--- a/Engine/Shaders/Private/DownsampleDepthPixelShader.usf
+++ b/Engine/Shaders/Private/DownsampleDepthPixelShader.usf
@@ -81,11 +81,7 @@ void Main(
 
 	if (DownsampleDepthFilter < DOWNSAMPLE_DEPTH_FILTER_CBMINMAX)
 	{
-#if HAS_INVERTED_Z_BUFFER
 		float FarDepth = min(min(DeviceZ0, DeviceZ1), min(DeviceZ2, DeviceZ3));
-#else
-		float FarDepth = max(max(DeviceZ0, DeviceZ1), max(DeviceZ2, DeviceZ3));
-#endif
 
 		// Max depth shrinks the silhouettes around foreground objects and is conservative for depth testing
 		// Sample 0 has consistent error, use whichever one is requested for this downsample
diff --git a/Engine/Shaders/Private/HairStrands/HairStrandsDeepShadowCommon.ush b/Engine/Shaders/Private/HairStrands/HairStrandsDeepShadowCommon.ush
--- a/Engine/Shaders/Private/HairStrands/HairStrandsDeepShadowCommon.ush
+++ b/Engine/Shaders/Private/HairStrands/HairStrandsDeepShadowCommon.ush
@@ -61,20 +61,12 @@ float ComputeHairCount(float4 DomValue, float DistanceToFrontDepth, float4 InLay
 
 float GetDomDistanceToFrontDepth(float FrontDepth, float LightSpaceZ)
 {
-#if HAS_INVERTED_Z_BUFFER
 	return max(0.0f, FrontDepth - LightSpaceZ);
-#else
-	return max(0.0f, LightSpaceZ - FrontDepth);
-#endif
 }
 
 float GetDomDistanceToFrontDepthWithBias(float FrontDepth, float LightSpaceZ, float DepthBias)
 {
-#if HAS_INVERTED_Z_BUFFER
 	return max(0.0f, FrontDepth - LightSpaceZ - DepthBias);
-#else
-	return max(0.0f, LightSpaceZ - FrontDepth - DepthBias);
-#endif
 }
 
 float SampleDOM_PCF2x2(
diff --git a/Engine/Shaders/Private/HeterogeneousVolumes/HeterogeneousVolumesComposite.usf b/Engine/Shaders/Private/HeterogeneousVolumes/HeterogeneousVolumesComposite.usf
--- a/Engine/Shaders/Private/HeterogeneousVolumes/HeterogeneousVolumesComposite.usf
+++ b/Engine/Shaders/Private/HeterogeneousVolumes/HeterogeneousVolumesComposite.usf
@@ -28,9 +28,7 @@ int bUseAVSM;
 float SampleTransmittance(uint2 ViewPixelCoord)
 {
 	float DeviceZ = SceneDepthTexture.Load(int3(ViewPixelCoord, 0)).r;
-#if HAS_INVERTED_Z_BUFFER
 	DeviceZ = max(0.000000000001, DeviceZ);
-#endif // HAS_INVERTED_Z_BUFFER
 
 	float3 TranslatedWorldPosition = SvPositionToResolvedTranslatedWorld(float4(ViewPixelCoord + View.TemporalAAJitter.xy, DeviceZ, 1));
 	return AVSM_SampleTransmittance(TranslatedWorldPosition, ResolvedView.TranslatedWorldCameraOrigin);
diff --git a/Engine/Shaders/Private/HeterogeneousVolumes/HeterogeneousVolumesHardwareRayTracing.usf b/Engine/Shaders/Private/HeterogeneousVolumes/HeterogeneousVolumesHardwareRayTracing.usf
--- a/Engine/Shaders/Private/HeterogeneousVolumes/HeterogeneousVolumesHardwareRayTracing.usf
+++ b/Engine/Shaders/Private/HeterogeneousVolumes/HeterogeneousVolumesHardwareRayTracing.usf
@@ -330,9 +330,7 @@ RAY_TRACING_ENTRY_RAYGEN(SparseVoxelsRayGenerationShader1)
 
 	// Extract depth
 	float DeviceZ = SceneDepthTexture.Load(int3(PixelCoord, 0)).r;
-#if HAS_INVERTED_Z_BUFFER
 	DeviceZ = max(0.000000000001, DeviceZ);
-#endif // HAS_INVERTED_Z_BUFFER
 
 	// Clip trace distance
 	float SceneDepth = min(ConvertFromDeviceZ(DeviceZ), GetMaxTraceDistance());
@@ -534,9 +532,7 @@ RAY_TRACING_ENTRY_RAYGEN(RenderSingleScatteringWithPreshadingRGS)
 
 	// Extract depth
 	float DeviceZ = SceneDepthTexture.Load(int3(PixelCoord, 0)).r;
-#if HAS_INVERTED_Z_BUFFER
 	DeviceZ = max(0.000000000001, DeviceZ);
-#endif // HAS_INVERTED_Z_BUFFER
 
 	// Clip trace distance
 	float SceneDepth = min(ConvertFromDeviceZ(DeviceZ), GetMaxTraceDistance());
diff --git a/Engine/Shaders/Private/HeterogeneousVolumes/HeterogeneousVolumesLiveShadingGlobalPipeline.usf b/Engine/Shaders/Private/HeterogeneousVolumes/HeterogeneousVolumesLiveShadingGlobalPipeline.usf
--- a/Engine/Shaders/Private/HeterogeneousVolumes/HeterogeneousVolumesLiveShadingGlobalPipeline.usf
+++ b/Engine/Shaders/Private/HeterogeneousVolumes/HeterogeneousVolumesLiveShadingGlobalPipeline.usf
@@ -71,9 +71,7 @@ void ScreenTileClassificationCS(
 #if 0
 	// Extract depth
 	float DeviceZ = SceneDepthTexture.Load(int3(ViewPixelCoord, 0)).r;
-#if HAS_INVERTED_Z_BUFFER
 	DeviceZ = max(0.000000000001, DeviceZ);
-#endif // HAS_INVERTED_Z_BUFFER
 
 	// Clip trace distance
 	float SceneDepth = min(ConvertFromDeviceZ(DeviceZ), MaxTraceDistance);
diff --git a/Engine/Shaders/Private/HeterogeneousVolumes/HeterogeneousVolumesLiveShadingPipeline.usf b/Engine/Shaders/Private/HeterogeneousVolumes/HeterogeneousVolumesLiveShadingPipeline.usf
--- a/Engine/Shaders/Private/HeterogeneousVolumes/HeterogeneousVolumesLiveShadingPipeline.usf
+++ b/Engine/Shaders/Private/HeterogeneousVolumes/HeterogeneousVolumesLiveShadingPipeline.usf
@@ -668,9 +668,7 @@ void RenderSingleScatteringWithLiveShading(
 
 	// Extract depth
 	float DeviceZ = SceneDepthTexture.Load(int3(ViewPixelCoord, 0)).r;
-#if HAS_INVERTED_Z_BUFFER
 	DeviceZ = max(0.000000000001, DeviceZ);
-#endif // HAS_INVERTED_Z_BUFFER
 
 	// Clip trace distance
 	float SceneDepth = min(ConvertFromDeviceZ(DeviceZ), GetMaxTraceDistance());
diff --git a/Engine/Shaders/Private/HeterogeneousVolumes/HeterogeneousVolumesLiveShadingShadows.usf b/Engine/Shaders/Private/HeterogeneousVolumes/HeterogeneousVolumesLiveShadingShadows.usf
--- a/Engine/Shaders/Private/HeterogeneousVolumes/HeterogeneousVolumesLiveShadingShadows.usf
+++ b/Engine/Shaders/Private/HeterogeneousVolumes/HeterogeneousVolumesLiveShadingShadows.usf
@@ -557,10 +557,7 @@ void RenderVolumetricShadowMapForLightWithLiveShadingCS(
 #if USE_CAMERA_SCENE_DEPTH
 	float2 ViewPixelCoord = (PixelCoord + View.TemporalAAJitter.xy) * CameraDownsampleFactor + View.ViewRectMin.xy / CameraDownsampleFactor;
 	float DeviceZ = SceneDepthTexture.Load(int3(ViewPixelCoord, 0)).r;
-
-	#if HAS_INVERTED_Z_BUFFER
-		DeviceZ = max(0.000000000001, DeviceZ);
-	#endif // HAS_INVERTED_Z_BUFFER
+	DeviceZ = max(0.000000000001, DeviceZ);
 
 	// Project ray through view transform to trim ray-length by opaque depth
 	FDFVector3 WorldViewRayOrigin = DFFastSubtract(GetTranslatedWorldCameraPosFromView(ViewPixelCoord), PrimaryView.PreViewTranslation);
diff --git a/Engine/Shaders/Private/HeterogeneousVolumes/HeterogeneousVolumesPreshadingPipeline.usf b/Engine/Shaders/Private/HeterogeneousVolumes/HeterogeneousVolumesPreshadingPipeline.usf
--- a/Engine/Shaders/Private/HeterogeneousVolumes/HeterogeneousVolumesPreshadingPipeline.usf
+++ b/Engine/Shaders/Private/HeterogeneousVolumes/HeterogeneousVolumesPreshadingPipeline.usf
@@ -210,9 +210,7 @@ void RenderSingleScatteringWithPreshadingCS(
 
 	// Extract depth
 	float DeviceZ = SceneDepthTexture.Load(int3(ViewPixelCoord, 0)).r;
-#if HAS_INVERTED_Z_BUFFER
 	DeviceZ = max(0.000000000001, DeviceZ);
-#endif // HAS_INVERTED_Z_BUFFER
 
 	// Clip trace distance
 	float SceneDepth = min(ConvertFromDeviceZ(DeviceZ), GetMaxTraceDistance());
diff --git a/Engine/Shaders/Private/HeterogeneousVolumes/HeterogeneousVolumesShadowMask.usf b/Engine/Shaders/Private/HeterogeneousVolumes/HeterogeneousVolumesShadowMask.usf
--- a/Engine/Shaders/Private/HeterogeneousVolumes/HeterogeneousVolumesShadowMask.usf
+++ b/Engine/Shaders/Private/HeterogeneousVolumes/HeterogeneousVolumesShadowMask.usf
@@ -77,9 +77,7 @@ void RenderVolumetricShadowMaskCS(
 
 	// Extract depth
 	float DeviceZ = SceneDepthTexture.Load(int3(PixelCoord, 0)).r;
-#if HAS_INVERTED_Z_BUFFER
 	DeviceZ = max(0.000000000001, DeviceZ);
-#endif // HAS_INVERTED_Z_BUFFER
 
 	// Initialize as input
 	RWVolumetricShadowMaskTexture[PixelCoord] = ShadowMaskTexture[PixelCoord];
diff --git a/Engine/Shaders/Private/HeterogeneousVolumes/HeterogeneousVolumesVisualizeCascades.usf b/Engine/Shaders/Private/HeterogeneousVolumes/HeterogeneousVolumesVisualizeCascades.usf
--- a/Engine/Shaders/Private/HeterogeneousVolumes/HeterogeneousVolumesVisualizeCascades.usf
+++ b/Engine/Shaders/Private/HeterogeneousVolumes/HeterogeneousVolumesVisualizeCascades.usf
@@ -105,9 +105,7 @@ void VisualizeCascadesCS(
 
 	// Extract depth
 	float DeviceZ = SceneDepthTexture.Load(int3(ViewPixelCoord, 0)).r;
-#if HAS_INVERTED_Z_BUFFER
 	DeviceZ = max(0.000000000001, DeviceZ);
-#endif // HAS_INVERTED_Z_BUFFER
 
 	// Clip trace distance
 	float Depth = min(ConvertFromDeviceZ(DeviceZ), GetMaxTraceDistance());
diff --git a/Engine/Shaders/Private/HeterogeneousVolumes/HeterogeneousVolumesVoxelGridPipeline.usf b/Engine/Shaders/Private/HeterogeneousVolumes/HeterogeneousVolumesVoxelGridPipeline.usf
--- a/Engine/Shaders/Private/HeterogeneousVolumes/HeterogeneousVolumesVoxelGridPipeline.usf
+++ b/Engine/Shaders/Private/HeterogeneousVolumes/HeterogeneousVolumesVoxelGridPipeline.usf
@@ -116,13 +116,8 @@ bool BoxIntersectsFrustum(float3 BoxOrigin, float3 BoxExtent, float4 Planes[5])
 bool BoxIntersectsViewFrustum(float3 BoxOrigin, float3 BoxExtent)
 {
 	float Epsilon = 1.0e-6;
-#if HAS_INVERTED_Z_BUFFER
 	float NearDepth = 1.0;
 	float FarDepth = Epsilon;
-#else
-	float NearDepth = 0.0;
-	float FarDepth = 1.0 - Epsilon;
-#endif // HAS_INVERTED_Z_BUFFER
 
 	FDFVector3 BoxWorldOrigin = DFPromote(BoxOrigin); // LWC_TODO
 	float3 BoxTranslatedWorldOrigin = DFFastToTranslatedWorld(BoxWorldOrigin, PrimaryView.PreViewTranslation);
diff --git a/Engine/Shaders/Private/HeterogeneousVolumes/HeterogeneousVolumesVoxelGridRendering.usf b/Engine/Shaders/Private/HeterogeneousVolumes/HeterogeneousVolumesVoxelGridRendering.usf
--- a/Engine/Shaders/Private/HeterogeneousVolumes/HeterogeneousVolumesVoxelGridRendering.usf
+++ b/Engine/Shaders/Private/HeterogeneousVolumes/HeterogeneousVolumesVoxelGridRendering.usf
@@ -575,9 +575,7 @@ void RenderTransmittanceWithVoxelGridCS(
 
 	// Extract depth
 	float DeviceZ = SceneDepthTexture.Load(int3(PixelCoord, 0)).r;
-#if HAS_INVERTED_Z_BUFFER
 	DeviceZ = max(0.000000000001, DeviceZ);
-#endif // HAS_INVERTED_Z_BUFFER
 
 	// Clip trace distance
 	float SceneDepth = min(ConvertFromDeviceZ(DeviceZ), MaxTraceDistance);
diff --git a/Engine/Shaders/Private/LocalFogVolumes/LocalFogVolumeSplat.usf b/Engine/Shaders/Private/LocalFogVolumes/LocalFogVolumeSplat.usf
--- a/Engine/Shaders/Private/LocalFogVolumes/LocalFogVolumeSplat.usf
+++ b/Engine/Shaders/Private/LocalFogVolumes/LocalFogVolumeSplat.usf
@@ -94,9 +94,7 @@ void LocalFogVolumeTiledPS(
 	float DeviceZ = LookupDeviceZ(uint2(SVPos.xy));
 #endif
 
-#if HAS_INVERTED_Z_BUFFER
 	DeviceZ = max(0.000000000001, DeviceZ);	// Make sure SvPositionToWorld does not return bad values when DeviceZ is far=0 when using inverted z
-#endif
 	float3 DepthBufferTranslatedWorldPos = SvPositionToResolvedTranslatedWorld(float4(ViewRectSVPos, DeviceZ, 1.0));
 
 #if LFV_HALFRES_PS
diff --git a/Engine/Shaders/Private/MeshDecals.usf b/Engine/Shaders/Private/MeshDecals.usf
--- a/Engine/Shaders/Private/MeshDecals.usf
+++ b/Engine/Shaders/Private/MeshDecals.usf
@@ -107,11 +107,9 @@ void MainVS(
 	OutGlobalClipPlaneDistance = dot(ResolvedView.GlobalClippingPlane, float4(WorldPosition.xyz, 1));
 #endif
 
-#if HAS_INVERTED_Z_BUFFER
 	// Move all geometry a small amount towards the camera to avoid z-fighting.
 	// Assuming a reverse z buffer this pushes the near plane a fixed small distance.
 	Output.Position.z += ResolvedView.DecalDepthBias;
-#endif
 }
 
 #endif // VERTEXSHADER
diff --git a/Engine/Shaders/Private/MeshEdges.usf b/Engine/Shaders/Private/MeshEdges.usf
--- a/Engine/Shaders/Private/MeshEdges.usf
+++ b/Engine/Shaders/Private/MeshEdges.usf
@@ -33,9 +33,7 @@ float Gamma;
 // Return true if A is closer than B
 bool IsCloser(float A, float B)
 {
-	return HAS_INVERTED_Z_BUFFER
-		? (A > B)
-		: (A < B);
+	return (A > B);
 }
 
 // Compose wireframe texture over shaded rendertarget, predicated on depth test with nearest shaded surface
diff --git a/Engine/Shaders/Private/PostProcessCompositePrimitives.usf b/Engine/Shaders/Private/PostProcessCompositePrimitives.usf
--- a/Engine/Shaders/Private/PostProcessCompositePrimitives.usf
+++ b/Engine/Shaders/Private/PostProcessCompositePrimitives.usf
@@ -374,7 +374,7 @@ void MainCompositePostProcessPrimitivesPS(
 
 		DepthMask = ComputeDepthMask(SampleSceneDeviceZ, SampleEditorDeviceZ);
 		#if WRITE_DEPTH
-		OutDepth = HAS_INVERTED_Z_BUFFER ? max(SampleSceneDeviceZ, SampleEditorDeviceZ) : min(SampleSceneDeviceZ, SampleEditorDeviceZ);
+		OutDepth = max(SampleSceneDeviceZ, SampleEditorDeviceZ);
         #endif
 	}
 	#elif MSAA_SAMPLE_COUNT > 1
@@ -422,9 +422,7 @@ void MainCompositePostProcessPrimitivesPS(
 			}
 
 			#if WRITE_DEPTH
-			OutDepth = HAS_INVERTED_Z_BUFFER
-				? max3(OutDepth, SampleSceneDeviceZ, SampleEditorDeviceZ)
-				: min3(OutDepth, SampleSceneDeviceZ, SampleEditorDeviceZ);
+			OutDepth = max3(OutDepth, SampleSceneDeviceZ, SampleEditorDeviceZ);
             #endif
 		}
 
diff --git a/Engine/Shaders/Private/RayTracing/RayTracingDebug.usf b/Engine/Shaders/Private/RayTracing/RayTracingDebug.usf
--- a/Engine/Shaders/Private/RayTracing/RayTracingDebug.usf
+++ b/Engine/Shaders/Private/RayTracing/RayTracingDebug.usf
@@ -462,11 +462,7 @@ RAY_TRACING_ENTRY_RAYGEN_OR_INLINE(RayTracingDebugMain)
 
 	float4 Result = (float4)0;
 	
-#if	HAS_INVERTED_Z_BUFFER
 	float OutDepth = 0.0f;
-#else
-	float OutDepth = 1.0f;
-#endif
 
 	bool bIsFarField = false;
 
diff --git a/Engine/Shaders/Private/ReflectionEnvironmentShaders.usf b/Engine/Shaders/Private/ReflectionEnvironmentShaders.usf
--- a/Engine/Shaders/Private/ReflectionEnvironmentShaders.usf
+++ b/Engine/Shaders/Private/ReflectionEnvironmentShaders.usf
@@ -926,9 +926,7 @@ void RenderRealTimeReflectionHeightFogPS(
 	float DeviceZ = FarDepthValue;
 #endif
 
-#if HAS_INVERTED_Z_BUFFER
 	DeviceZ = max(1.0e-10, DeviceZ); // For ConvertFromDeviceZ to return a valid value (non inf).
-#endif
 
 	float SceneDepth = ConvertFromDeviceZ(DeviceZ);
 	float3 WorldPositionRelativeToCamera = ScreenVector.xyz * SceneDepth;
diff --git a/Engine/Shaders/Private/SingleLayerWaterShading.ush b/Engine/Shaders/Private/SingleLayerWaterShading.ush
--- a/Engine/Shaders/Private/SingleLayerWaterShading.ush
+++ b/Engine/Shaders/Private/SingleLayerWaterShading.ush
@@ -174,10 +174,8 @@ WaterVolumeLightingOutput EvaluateWaterVolumeLighting(
 
 #endif // !SINGLE_LAYER_WATER_MOBILE_SHADING
 
-#if HAS_INVERTED_Z_BUFFER
 		// SvPositionToWorld will result in NaN if SceneDeviceZ=0
 		SceneDeviceZ = max(0.000000000001f, SceneDeviceZ);
-#endif
 
 		const float BehindWaterDeltaDepth = CameraIsUnderWater ? WaterDepth : max(0.0f, SceneDepth - WaterDepth);
 
diff --git a/Engine/Shaders/Private/SkyAtmosphere.usf b/Engine/Shaders/Private/SkyAtmosphere.usf
--- a/Engine/Shaders/Private/SkyAtmosphere.usf
+++ b/Engine/Shaders/Private/SkyAtmosphere.usf
@@ -162,9 +162,7 @@ float RaySphereIntersectNearest(float3 RayOrigin, float3 RayDir, float3 SphereCe
 
 float4 GetScreenTranslatedWorldPos(float4 SVPos, float DeviceZ)
 {
-#if HAS_INVERTED_Z_BUFFER
 	DeviceZ = max(0.000000000001, DeviceZ);	// TODO: investigate why SvPositionToWorld returns bad values when DeviceZ is far=0 when using inverted z
-#endif
 	return float4(SvPositionToTranslatedWorld(float4(SVPos.xy, DeviceZ, 1.0)), 1.0);
 }
 
diff --git a/Engine/Shaders/Private/TemporalAA.usf b/Engine/Shaders/Private/TemporalAA.usf
--- a/Engine/Shaders/Private/TemporalAA.usf
+++ b/Engine/Shaders/Private/TemporalAA.usf
@@ -2072,37 +2072,33 @@ FTAAHistoryPayload TemporalAASample(uint2 GroupId, uint2 GroupThreadId, uint Gro
 
 		float2 DepthOffset = float2(AA_CROSS, AA_CROSS);
 		float DepthOffsetXx = float(AA_CROSS);
-		#if HAS_INVERTED_Z_BUFFER
-			// Nearest depth is the largest depth (depth surface 0=far, 1=near).
-			if(Depths.x > Depths.y) 
-			{
-				DepthOffsetXx = -AA_CROSS;
-			}
-			if(Depths.z > Depths.w) 
-			{
-				DepthOffset.x = -AA_CROSS;
-			}
-			float DepthsXY = max(Depths.x, Depths.y);
-			float DepthsZW = max(Depths.z, Depths.w);
-			if(DepthsXY > DepthsZW) 
-			{
-				DepthOffset.y = -AA_CROSS;
-				DepthOffset.x = DepthOffsetXx; 
-			}
-			float DepthsXYZW = max(DepthsXY, DepthsZW);
-			if(DepthsXYZW > PosN.z) 
-			{
-				// This is offset for reading from velocity texture.
-				// This supports half or fractional resolution velocity textures.
-				// With the assumption that UV position scales between velocity and color.
-				VelocityOffset = DepthOffset * InputSceneColorSize.zw;
-				// This is [0 to 1] flipped in Y.
-				//PosN.xy = ScreenPos + DepthOffset * OutputViewportSize.zw * 2.0;
-				PosN.z = DepthsXYZW;
-			}
-		#else // !HAS_INVERTED_Z_BUFFER
-			#error Fix me!
-		#endif // !HAS_INVERTED_Z_BUFFER
+		// Nearest depth is the largest depth (depth surface 0=far, 1=near).
+		if(Depths.x > Depths.y) 
+		{
+			DepthOffsetXx = -AA_CROSS;
+		}
+		if(Depths.z > Depths.w) 
+		{
+			DepthOffset.x = -AA_CROSS;
+		}
+		float DepthsXY = max(Depths.x, Depths.y);
+		float DepthsZW = max(Depths.z, Depths.w);
+		if(DepthsXY > DepthsZW) 
+		{
+			DepthOffset.y = -AA_CROSS;
+			DepthOffset.x = DepthOffsetXx; 
+		}
+		float DepthsXYZW = max(DepthsXY, DepthsZW);
+		if(DepthsXYZW > PosN.z) 
+		{
+			// This is offset for reading from velocity texture.
+			// This supports half or fractional resolution velocity textures.
+			// With the assumption that UV position scales between velocity and color.
+			VelocityOffset = DepthOffset * InputSceneColorSize.zw;
+			// This is [0 to 1] flipped in Y.
+			//PosN.xy = ScreenPos + DepthOffset * OutputViewportSize.zw * 2.0;
+			PosN.z = DepthsXYZW;
+		}
 	}
 	#endif	// AA_CROSS
 
diff --git a/Engine/Shaders/Private/Tools/VisualizeTexture.usf b/Engine/Shaders/Private/Tools/VisualizeTexture.usf
--- a/Engine/Shaders/Private/Tools/VisualizeTexture.usf
+++ b/Engine/Shaders/Private/Tools/VisualizeTexture.usf
@@ -150,9 +150,7 @@ void VisualizeTexturePS(
 		if (InputValueMapping == INPUT_MAPPING_DEPTH)
 		{
 			// current z buffer is negated (1 is near, 0 is far)
-#if HAS_INVERTED_Z_BUFFER
 			Depth = 1 - Depth;
-#endif
 
 #if FEATURE_LEVEL >= FEATURE_LEVEL_SM4
 			// map depth to nice colors
diff --git a/Engine/Shaders/Private/VolumetricCloud.usf b/Engine/Shaders/Private/VolumetricCloud.usf
--- a/Engine/Shaders/Private/VolumetricCloud.usf
+++ b/Engine/Shaders/Private/VolumetricCloud.usf
@@ -592,9 +592,7 @@ void MainCommon(in FMaterialPixelParameters MaterialParameters, in float4 SvPosi
 #if CLOUD_MIN_AND_MAX_DEPTH
 
 		float2 MinAndMaxDepth = RenderVolumetricCloudParameters.SceneDepthMinAndMaxTexture.Load(uint3(SceneDepthTextureCoord / 2, 0)).rg;
-	#if HAS_INVERTED_Z_BUFFER
 		MinAndMaxDepth = max(0.000000000001, MinAndMaxDepth.yx); // x is clostest depth, y is furthest depth
-	#endif
 
 		// Trace up to furthest depth.
 		float DeviceZ = MinAndMaxDepth.y;
@@ -612,9 +610,7 @@ void MainCommon(in FMaterialPixelParameters MaterialParameters, in float4 SvPosi
 #else // CLOUD_MIN_AND_MAX_DEPTH
 
 		float DeviceZ = RenderVolumetricCloudParameters.SceneDepthTexture.Load(uint3(SceneDepthTextureCoord, 0)).r;
-	#if HAS_INVERTED_Z_BUFFER
 		DeviceZ = max(0.000000000001, DeviceZ);
-	#endif
 		float3 DepthBufferTranslatedWorldPos = SvPositionToTranslatedWorld(float4(SvPosition.xy, DeviceZ, 1.0));
 		const float TDepthBufferKm = min(length(DepthBufferTranslatedWorldPos * CENTIMETER_TO_KILOMETER - RayTranslatedWorldOriginKm), MaxHalfFloat);
 		OutDepth.xyzw = float4(TDepthBufferKm, TDepthBufferKm, TDepthBufferKm, TDepthBufferKm);
diff --git a/Engine/Shaders/Private/VolumetricRenderTarget.usf b/Engine/Shaders/Private/VolumetricRenderTarget.usf
--- a/Engine/Shaders/Private/VolumetricRenderTarget.usf
+++ b/Engine/Shaders/Private/VolumetricRenderTarget.usf
@@ -833,9 +833,7 @@ void ComposeVolumetricRTOverScenePS(
 	float4 ClipPosition = mul(float4(TranslatedWorldPosition, 1.0), PrimaryView.TranslatedWorldToClip);
 	ClipPosition /= ClipPosition.w;
 	float PixelDeviceZ = ClipPosition.z;
-#if HAS_INVERTED_Z_BUFFER
 	PixelDeviceZ = max(0.000000000001, PixelDeviceZ);
-#endif
 
 	float3 ScreenTranslatedWorldPosition = SvPositionToTranslatedWorld(float4(SVPos.xy, PixelDeviceZ, 1.0));
 
@@ -846,9 +844,7 @@ void ComposeVolumetricRTOverScenePS(
 #else
 	float PixelDeviceZ = SceneTexturesStruct.SceneDepthTexture.Load(uint3(SVPos.xy, 0)).r;
 #endif
-#if HAS_INVERTED_Z_BUFFER
 	PixelDeviceZ = max(0.000000000001, PixelDeviceZ);
-#endif
 	float3 ScreenTranslatedWorldPosition = SvPositionToTranslatedWorld(float4(SVPos.xy, PixelDeviceZ, 1.0));
 #endif 
 
diff --git a/Engine/Source/Editor/UnrealEd/Private/EditorViewportClient.cpp b/Engine/Source/Editor/UnrealEd/Private/EditorViewportClient.cpp
--- a/Engine/Source/Editor/UnrealEd/Private/EditorViewportClient.cpp
+++ b/Engine/Source/Editor/UnrealEd/Private/EditorViewportClient.cpp
@@ -1285,28 +1285,14 @@ FSceneView* FEditorViewportClient::CalcSceneView(FSceneViewFamily* ViewFamily, c
 					// Avoid zero ViewFOV's which cause divide by zero's in projection matrix
 					const float MatrixFOV = FMath::Max(0.001f, ModifiedViewFOV) * (float)PI / 360.0f;
 
-				    if ((bool)ERHIZBuffer::IsInverted)
-				    {
-					    ViewInitOptions.ProjectionMatrix = FReversedZPerspectiveMatrix(
-						    MatrixFOV,
-						    MatrixFOV,
-						    1.0f,
-						    AspectRatio,
-						    MinZ,
-						    MaxZ
-						    );
-				    }
-				    else
-				    {
-					    ViewInitOptions.ProjectionMatrix = FPerspectiveMatrix(
-						    MatrixFOV,
-						    MatrixFOV,
-						    1.0f,
-						    AspectRatio,
-						    MinZ,
-						    MaxZ
-						    );
-				    }
+					ViewInitOptions.ProjectionMatrix = FReversedZPerspectiveMatrix(
+						MatrixFOV,
+						MatrixFOV,
+						1.0f,
+						AspectRatio,
+						MinZ,
+						MaxZ
+						);
 			    }
 			    else
 			    {
@@ -1342,35 +1328,19 @@ FSceneView* FEditorViewportClient::CalcSceneView(FSceneViewFamily* ViewFamily, c
 						MatrixFOV = FMath::Max(0.001f, ModifiedViewFOV) * (float)UE_PI / 360.0f;
 					}
 
-				    if ((bool)ERHIZBuffer::IsInverted)
-				    {
-					    ViewInitOptions.ProjectionMatrix = FReversedZPerspectiveMatrix(
-						    MatrixFOV,
-						    MatrixFOV,
-						    XAxisMultiplier,
-						    YAxisMultiplier,
-						    MinZ,
-						    MaxZ
-						    );
-				    }
-				    else
-				    {
-					    ViewInitOptions.ProjectionMatrix = FPerspectiveMatrix(
-						    MatrixFOV,
-						    MatrixFOV,
-						    XAxisMultiplier,
-						    YAxisMultiplier,
-						    MinZ,
-						    MaxZ
-						    );
-				    }
+					ViewInitOptions.ProjectionMatrix = FReversedZPerspectiveMatrix(
+						MatrixFOV,
+						MatrixFOV,
+						XAxisMultiplier,
+						YAxisMultiplier,
+						MinZ,
+						MaxZ
+						);
 			    }
 			}
 		}
 		else
 		{
-			static_assert((bool)ERHIZBuffer::IsInverted, "Check all the Rotation Matrix transformations!");
-
 			//The divisor for the matrix needs to match the translation code.
 			const float Zoom = GetOrthoUnitsPerPixel(Viewport);
 			float OrthoWidth = FMath::Clamp(Zoom * ViewportSize.X/2.0f, 0.0f, UE_LARGE_HALF_WORLD_MAX);
diff --git a/Engine/Source/Editor/UnrealEd/Private/ThumbnailHelpers.cpp b/Engine/Source/Editor/UnrealEd/Private/ThumbnailHelpers.cpp
--- a/Engine/Source/Editor/UnrealEd/Private/ThumbnailHelpers.cpp
+++ b/Engine/Source/Editor/UnrealEd/Private/ThumbnailHelpers.cpp
@@ -111,7 +111,6 @@ FSceneView* FThumbnailPreviewScene::CreateView(FSceneViewFamily * ViewFamily, in
 
 	const float FOVDegrees = GetFOV();
 	const float HalfFOVRadians = FMath::DegreesToRadians<float>(FOVDegrees) * 0.5f;
-	static_assert((int32)ERHIZBuffer::IsInverted != 0, "Check NearPlane and Projection Matrix");
 	const float NearPlane = 1.0f;
 	FMatrix ProjectionMatrix = FReversedZPerspectiveMatrix(
 		HalfFOVRadians,
diff --git a/Engine/Source/Runtime/Engine/Private/SceneView.cpp b/Engine/Source/Runtime/Engine/Private/SceneView.cpp
--- a/Engine/Source/Runtime/Engine/Private/SceneView.cpp
+++ b/Engine/Source/Runtime/Engine/Private/SceneView.cpp
@@ -828,7 +828,6 @@ void FSceneView::SetupViewFrustum()
 	}
 
 	// Derive the view's near clipping distance and plane.
-	static_assert((int32)ERHIZBuffer::IsInverted != 0, "Fix Near Clip distance!");
 	bHasNearClippingPlane = ViewMatrices.GetViewProjectionMatrix().GetFrustumNearPlane(NearClippingPlane);
 	NearClippingDistance = ViewMatrices.ComputeNearPlane();
 }
diff --git a/Engine/Source/Runtime/Engine/Private/ShaderCompiler/ShaderCompiler.cpp b/Engine/Source/Runtime/Engine/Private/ShaderCompiler/ShaderCompiler.cpp
--- a/Engine/Source/Runtime/Engine/Private/ShaderCompiler/ShaderCompiler.cpp
+++ b/Engine/Source/Runtime/Engine/Private/ShaderCompiler/ShaderCompiler.cpp
@@ -3807,8 +3807,6 @@ void GlobalBeginCompileShader(
 		}
 	}
 
-	SET_SHADER_DEFINE(Input.Environment, HAS_INVERTED_Z_BUFFER, (bool)ERHIZBuffer::IsInverted);
-
 	if (Input.Environment.CompilerFlags.Contains(CFLAG_HLSL2021))
 	{
 		SET_SHADER_DEFINE(Input.Environment, COMPILER_SUPPORTS_HLSL2021, 1);
diff --git a/Engine/Source/Runtime/Engine/Private/UserInterface/Canvas.cpp b/Engine/Source/Runtime/Engine/Private/UserInterface/Canvas.cpp
--- a/Engine/Source/Runtime/Engine/Private/UserInterface/Canvas.cpp
+++ b/Engine/Source/Runtime/Engine/Private/UserInterface/Canvas.cpp
@@ -441,28 +441,14 @@ FMatrix FCanvas::CalcProjectionMatrix(uint32 ViewSizeX, uint32 ViewSizeY, float
 	// convert FOV to randians
 	float FOVRad = fFOV * (float)UE_PI / 360.0f;
 	// project based on the FOV and near plane given
-	if ((bool)ERHIZBuffer::IsInverted)
-	{
-		return AdjustProjectionMatrixForRHI(
-			FReversedZPerspectiveMatrix(
-				FOVRad,
-				ViewSizeX,
-				ViewSizeY,
-				NearPlane
-				)
-			);
-	}
-	else
-	{
-		return AdjustProjectionMatrixForRHI(
-			FPerspectiveMatrix(
-				FOVRad,
-				ViewSizeX,
-				ViewSizeY,
-				NearPlane
-				)
-			);
-	}
+	return AdjustProjectionMatrixForRHI(
+		FReversedZPerspectiveMatrix(
+			FOVRad,
+			ViewSizeX,
+			ViewSizeY,
+			NearPlane
+			)
+		);
 }
 
 bool FCanvasBatchedElementRenderItem::Render_RenderThread(FCanvasRenderContext& RenderContext, FMeshPassProcessorRenderState& DrawRenderState, const FCanvas* Canvas)
diff --git a/Engine/Source/Runtime/RHI/Private/RHI.cpp b/Engine/Source/Runtime/RHI/Private/RHI.cpp
--- a/Engine/Source/Runtime/RHI/Private/RHI.cpp
+++ b/Engine/Source/Runtime/RHI/Private/RHI.cpp
@@ -121,8 +121,8 @@ const FClearValueBinding FClearValueBinding::White(FLinearColor(1.0f, 1.0f, 1.0f
 const FClearValueBinding FClearValueBinding::Transparent(FLinearColor(0.0f, 0.0f, 0.0f, 0.0f));
 const FClearValueBinding FClearValueBinding::DepthOne(1.0f, 0);
 const FClearValueBinding FClearValueBinding::DepthZero(0.0f, 0);
-const FClearValueBinding FClearValueBinding::DepthNear((float)ERHIZBuffer::NearPlane, 0);
-const FClearValueBinding FClearValueBinding::DepthFar((float)ERHIZBuffer::FarPlane, 0);
+const FClearValueBinding FClearValueBinding::DepthNear(1.0f, 0);
+const FClearValueBinding FClearValueBinding::DepthFar(0.0f, 0);
 const FClearValueBinding FClearValueBinding::Green(FLinearColor(0.0f, 1.0f, 0.0f, 1.0f));
 // Note: this is used as the default normal for DBuffer decals.  It must decode to a value of 0 in DecodeDBufferData.
 const FClearValueBinding FClearValueBinding::DefaultNormal8Bit(FLinearColor(128.0f / 255.0f, 128.0f / 255.0f, 128.0f / 255.0f, 1.0f));
@@ -1111,11 +1111,6 @@ void RHIDumpResourceMemoryToCSV()
 
 #endif // RHI_ENABLE_RESOURCE_INFO
 
-static_assert(ERHIZBuffer::FarPlane != ERHIZBuffer::NearPlane, "Near and Far planes must be different!");
-static_assert((int32)ERHIZBuffer::NearPlane == 0 || (int32)ERHIZBuffer::NearPlane == 1, "Invalid Values for Near Plane, can only be 0 or 1!");
-static_assert((int32)ERHIZBuffer::FarPlane == 0 || (int32)ERHIZBuffer::FarPlane == 1, "Invalid Values for Far Plane, can only be 0 or 1");
-
-
 /**
  * RHI configuration settings.
  */
diff --git a/Engine/Source/Runtime/RHI/Public/RHIDefinitions.h b/Engine/Source/Runtime/RHI/Public/RHIDefinitions.h
--- a/Engine/Source/Runtime/RHI/Public/RHIDefinitions.h
+++ b/Engine/Source/Runtime/RHI/Public/RHIDefinitions.h
@@ -269,16 +269,11 @@ enum
 };
 static_assert(MaxSimultaneousRenderTargets <= (1 << MaxSimultaneousRenderTargets_NumBits), "MaxSimultaneousRenderTargets will not fit on MaxSimultaneousRenderTargets_NumBits");
 
-enum class ERHIZBuffer
+enum class UE_DEPRECATED(5.8, "ERHIZBuffer is being removed as the assumption of using an Inverted Z buffer is so ingrained that it could not reasonably be changed.") ERHIZBuffer
 {
-	// Before changing this, make sure all math & shader assumptions are correct! Also wrap your C++ assumptions with
-	//		static_assert(ERHIZBuffer::IsInvertedZBuffer(), ...);
-	// Shader-wise, make sure to update Definitions.usf, HAS_INVERTED_Z_BUFFER
 	FarPlane = 0,
 	NearPlane = 1,
-
-	// 'bool' for knowing if the API is using Inverted Z buffer
-	IsInverted = (int32)((int32)ERHIZBuffer::FarPlane < (int32)ERHIZBuffer::NearPlane),
+	IsInverted = 1,
 };
 
 
@@ -391,10 +386,10 @@ enum ECompareFunction
 	ECompareFunction_NumBits = 3,
 
 	// Utility enumerations
-	CF_DepthNearOrEqual		= (((int32)ERHIZBuffer::IsInverted != 0) ? CF_GreaterEqual : CF_LessEqual),
-	CF_DepthNear			= (((int32)ERHIZBuffer::IsInverted != 0) ? CF_Greater : CF_Less),
-	CF_DepthFartherOrEqual	= (((int32)ERHIZBuffer::IsInverted != 0) ? CF_LessEqual : CF_GreaterEqual),
-	CF_DepthFarther			= (((int32)ERHIZBuffer::IsInverted != 0) ? CF_Less : CF_Greater),
+	CF_DepthNearOrEqual		= CF_GreaterEqual,
+	CF_DepthNear			= CF_Greater,
+	CF_DepthFartherOrEqual	= CF_LessEqual,
+	CF_DepthFarther			= CF_Less,
 };
 static_assert(ECompareFunction_Num <= (1 << ECompareFunction_NumBits), "ECompareFunction_Num will not fit on ECompareFunction_NumBits");
 
diff --git a/Engine/Source/Runtime/Renderer/Private/BasePassRendering.cpp b/Engine/Source/Runtime/Renderer/Private/BasePassRendering.cpp
--- a/Engine/Source/Runtime/Renderer/Private/BasePassRendering.cpp
+++ b/Engine/Source/Runtime/Renderer/Private/BasePassRendering.cpp
@@ -912,7 +912,7 @@ static void ClearGBufferAtMaxZ(
 
 			SetGraphicsPipelineState(RHICmdList, GraphicsPSOInit, 0);
 
-			SetShaderParametersLegacyVS(RHICmdList, VertexShader, float(ERHIZBuffer::FarPlane));
+			SetShaderParametersLegacyVS(RHICmdList, VertexShader, 0.f);
 
 			RHICmdList.SetViewport(View.ViewRect.Min.X, View.ViewRect.Min.Y, 0, View.ViewRect.Max.X, View.ViewRect.Max.Y, 1);
 
diff --git a/Engine/Source/Runtime/Renderer/Private/CompositionLighting/PostProcessAmbientOcclusion.cpp b/Engine/Source/Runtime/Renderer/Private/CompositionLighting/PostProcessAmbientOcclusion.cpp
--- a/Engine/Source/Runtime/Renderer/Private/CompositionLighting/PostProcessAmbientOcclusion.cpp
+++ b/Engine/Source/Runtime/Renderer/Private/CompositionLighting/PostProcessAmbientOcclusion.cpp
@@ -902,8 +902,6 @@ void AddAmbientOcclusionPass(
 			const FVector4f Far = (FVector4f)ProjectionMatrix.TransformFVector4(FVector4(0, 0, Settings.AmbientOcclusionFadeDistance));
 			DepthFar = FMath::Min(1.0f, Far.Z / Far.W);
 
-			static_assert(bool(ERHIZBuffer::IsInverted), "Inverted depth buffer is assumed when setting depth bounds test for AO.");
-
 			FRenderTargetParameters* ClearParameters = GraphBuilder.AllocParameters<FRenderTargetParameters>();
 			ClearParameters->RenderTargets[0] = Output.GetRenderTargetBinding();
 			ClearParameters->RenderTargets.DepthStencil = DepthStencilBinding;
diff --git a/Engine/Source/Runtime/Renderer/Private/FogRendering.cpp b/Engine/Source/Runtime/Renderer/Private/FogRendering.cpp
--- a/Engine/Source/Runtime/Renderer/Private/FogRendering.cpp
+++ b/Engine/Source/Runtime/Renderer/Private/FogRendering.cpp
@@ -473,15 +473,7 @@ static void RenderViewFog(
 		FVector4f ClipSpaceMaxDistance = (FVector4f)View.ViewMatrices.GetProjectionMatrix().TransformPosition(ViewSpaceStartFogPoint); // LWC_TODO: precision loss
 		float FogClipSpaceZ = ClipSpaceMaxDistance.Z / ClipSpaceMaxDistance.W;
 		FogClipSpaceZ = FMath::Clamp(FogClipSpaceZ, 0.f, 1.f);
-
-		if (bool(ERHIZBuffer::IsInverted))
-		{
-			RHICmdList.SetDepthBounds(0.0f, FogClipSpaceZ);
-		}
-		else
-		{
-			RHICmdList.SetDepthBounds(FogClipSpaceZ, 1.0f);
-		}
+		RHICmdList.SetDepthBounds(0.0f, FogClipSpaceZ);
 	}
 
 	SetGraphicsPipelineState(RHICmdList, GraphicsPSOInit, 0);
diff --git a/Engine/Source/Runtime/Renderer/Private/IndirectLightRendering.cpp b/Engine/Source/Runtime/Renderer/Private/IndirectLightRendering.cpp
--- a/Engine/Source/Runtime/Renderer/Private/IndirectLightRendering.cpp
+++ b/Engine/Source/Runtime/Renderer/Private/IndirectLightRendering.cpp
@@ -1542,7 +1542,6 @@ void FDeferredShadingSceneRenderer::RenderDiffuseIndirectAndAmbientOcclusion(
 
 				// only use depth bound optimization when diffuse indirect is disable
 				bool bUseDepthBounds = CVarDiffuseIndirectOffUseDepthBoundsAO.GetValueOnRenderThread() &&
-										bool(ERHIZBuffer::IsInverted) && // Inverted depth buffer is assumed when setting depth bounds test for AO.
 										!PermutationVector.Get<FDiffuseIndirectCompositePS::FUpscaleSSGIDim>() &&
 										AmbientOcclusionMask &&
 										ViewPipelineState.DiffuseIndirectMethod == EDiffuseIndirectMethod::Disabled;
diff --git a/Engine/Source/Runtime/Renderer/Private/LocalFogVolumeRendering.cpp b/Engine/Source/Runtime/Renderer/Private/LocalFogVolumeRendering.cpp
--- a/Engine/Source/Runtime/Renderer/Private/LocalFogVolumeRendering.cpp
+++ b/Engine/Source/Runtime/Renderer/Private/LocalFogVolumeRendering.cpp
@@ -820,16 +820,8 @@ static FDepthBoundSetup GetDepthBoundSetup(float FogStartDistance, FMatrix ViewP
 	DepthBoundSetup.FogClipDeviceZ = FMath::Clamp(FogClipSpaceZ, 0.f, 1.f);
 
 	DepthBoundSetup.bEnabled = GSupportsDepthBoundsTest;
-	if (bool(ERHIZBuffer::IsInverted))
-	{
-		DepthBoundSetup.MinDeviceZ = 0.0f;
-		DepthBoundSetup.MaxDeviceZ = DepthBoundSetup.FogClipDeviceZ;
-	}
-	else
-	{
-		DepthBoundSetup.MinDeviceZ = DepthBoundSetup.FogClipDeviceZ;
-		DepthBoundSetup.MaxDeviceZ = 1.0f;
-	}
+	DepthBoundSetup.MinDeviceZ = 0.0f;
+	DepthBoundSetup.MaxDeviceZ = DepthBoundSetup.FogClipDeviceZ;
 
 	return DepthBoundSetup;
 }
diff --git a/Engine/Source/Runtime/Renderer/Private/PostProcess/PostProcessAmbientOcclusionMobile.cpp b/Engine/Source/Runtime/Renderer/Private/PostProcess/PostProcessAmbientOcclusionMobile.cpp
--- a/Engine/Source/Runtime/Renderer/Private/PostProcess/PostProcessAmbientOcclusionMobile.cpp
+++ b/Engine/Source/Runtime/Renderer/Private/PostProcess/PostProcessAmbientOcclusionMobile.cpp
@@ -923,8 +923,6 @@ static void AddMobileAmbientOcclusionPass(
 		const FVector4f Far = (FVector4f)ProjectionMatrix.TransformFVector4(FVector4(0, 0, Settings.AmbientOcclusionFadeDistance));
 		DepthFar = FMath::Min(1.0f, Far.Z / Far.W);
 
-		static_assert(bool(ERHIZBuffer::IsInverted), "Inverted depth buffer is assumed when setting depth bounds test for AO.");
-
 		FRenderTargetParameters* ClearParameters = GraphBuilder.AllocParameters<FRenderTargetParameters>();
 		ClearParameters->RenderTargets[0] = Output.GetRenderTargetBinding();
 		ClearParameters->RenderTargets.DepthStencil = DepthStencilBinding;
@@ -1100,8 +1098,6 @@ static void AddMobileAmbientOcclusionUpsamplePass(
 		const FVector4f Far = (FVector4f)ProjectionMatrix.TransformFVector4(FVector4(0, 0, Settings.AmbientOcclusionFadeDistance));
 		DepthFar = FMath::Min(1.0f, Far.Z / Far.W);
 
-		static_assert(bool(ERHIZBuffer::IsInverted), "Inverted depth buffer is assumed when setting depth bounds test for AO.");
-
 		FRenderTargetParameters* ClearParameters = GraphBuilder.AllocParameters<FRenderTargetParameters>();
 		ClearParameters->RenderTargets[0] = Output.GetRenderTargetBinding();
 		ClearParameters->RenderTargets.DepthStencil = DepthStencilBinding;
diff --git a/Engine/Source/Runtime/Renderer/Private/PostProcess/PostProcessSelectionOutline.cpp b/Engine/Source/Runtime/Renderer/Private/PostProcess/PostProcessSelectionOutline.cpp
--- a/Engine/Source/Runtime/Renderer/Private/PostProcess/PostProcessSelectionOutline.cpp
+++ b/Engine/Source/Runtime/Renderer/Private/PostProcess/PostProcessSelectionOutline.cpp
@@ -156,7 +156,7 @@ FScreenPassTexture AddSelectionOutlinePass(
 				DepthStencilDesc.Reset();
 				DepthStencilDesc.Format = PF_DepthStencil;
 				// This is a reversed Z depth surface, so 0.0f is the far plane.
-				DepthStencilDesc.ClearValue = FClearValueBinding((float)ERHIZBuffer::FarPlane, 0);
+				DepthStencilDesc.ClearValue = FClearValueBinding(0.f, 0);
 				DepthStencilDesc.Flags = TexCreate_DepthStencilTargetable | TexCreate_ShaderResource;
 				DepthStencilDesc.NumSamples = NumSamples;
 				DepthStencilDesc.Extent = DepthStencilViewport.Extent;
@@ -259,16 +259,16 @@ FScreenPassTexture AddSelectionOutlinePass(
 
 					// top
 					RHICmdList.SetScissorRect(true, DepthStencilViewport.Rect.Min.X, DepthStencilViewport.Rect.Min.Y, DepthStencilViewport.Rect.Max.X, InnerRect.Min.Y);
-					DrawClearQuad(RHICmdList, false, FLinearColor(), true, (float)ERHIZBuffer::FarPlane, true, 0, DepthStencilViewport.Extent, FIntRect());
+					DrawClearQuad(RHICmdList, false, FLinearColor(), true, 0.f, true, 0, DepthStencilViewport.Extent, FIntRect());
 					// bottom
 					RHICmdList.SetScissorRect(true, DepthStencilViewport.Rect.Min.X, InnerRect.Max.Y, DepthStencilViewport.Rect.Max.X, DepthStencilViewport.Rect.Max.Y);
-					DrawClearQuad(RHICmdList, false, FLinearColor(), true, (float)ERHIZBuffer::FarPlane, true, 0, DepthStencilViewport.Extent, FIntRect());
+					DrawClearQuad(RHICmdList, false, FLinearColor(), true, 0.f, true, 0, DepthStencilViewport.Extent, FIntRect());
 					// left
 					RHICmdList.SetScissorRect(true, DepthStencilViewport.Rect.Min.X, DepthStencilViewport.Rect.Min.Y, InnerRect.Min.X, DepthStencilViewport.Rect.Max.Y);
-					DrawClearQuad(RHICmdList, false, FLinearColor(), true, (float)ERHIZBuffer::FarPlane, true, 0, DepthStencilViewport.Extent, FIntRect());
+					DrawClearQuad(RHICmdList, false, FLinearColor(), true, 0.f, true, 0, DepthStencilViewport.Extent, FIntRect());
 					// right
 					RHICmdList.SetScissorRect(true, InnerRect.Max.X, DepthStencilViewport.Rect.Min.Y, DepthStencilViewport.Rect.Max.X, DepthStencilViewport.Rect.Max.Y);
-					DrawClearQuad(RHICmdList, false, FLinearColor(), true, (float)ERHIZBuffer::FarPlane, true, 0, DepthStencilViewport.Extent, FIntRect());
+					DrawClearQuad(RHICmdList, false, FLinearColor(), true, 0.f, true, 0, DepthStencilViewport.Extent, FIntRect());
 
 					RHICmdList.SetScissorRect(false, 0, 0, 0, 0);
 				}
diff --git a/Engine/Source/Runtime/Renderer/Private/PostProcess/PostProcessVisualizeLevelInstance.cpp b/Engine/Source/Runtime/Renderer/Private/PostProcess/PostProcessVisualizeLevelInstance.cpp
--- a/Engine/Source/Runtime/Renderer/Private/PostProcess/PostProcessVisualizeLevelInstance.cpp
+++ b/Engine/Source/Runtime/Renderer/Private/PostProcess/PostProcessVisualizeLevelInstance.cpp
@@ -78,7 +78,7 @@ FScreenPassTexture AddVisualizeLevelInstancePass(
 			DepthStencilDesc.Reset();
 			DepthStencilDesc.Format = PF_DepthStencil;
 			// This is a reversed Z depth surface, so 0.0f is the far plane.
-			DepthStencilDesc.ClearValue = FClearValueBinding((float)ERHIZBuffer::FarPlane, 0);
+			DepthStencilDesc.ClearValue = FClearValueBinding(0.f, 0);
 			DepthStencilDesc.Flags = TexCreate_DepthStencilTargetable | TexCreate_ShaderResource;
 			DepthStencilDesc.NumSamples = NumSamples;
 			DepthStencilDesc.AliasableFormats.Add(PF_X24_G8);
diff --git a/Engine/Source/Runtime/Renderer/Private/ReflectionEnvironmentCapture.cpp b/Engine/Source/Runtime/Renderer/Private/ReflectionEnvironmentCapture.cpp
--- a/Engine/Source/Runtime/Renderer/Private/ReflectionEnvironmentCapture.cpp
+++ b/Engine/Source/Runtime/Renderer/Private/ReflectionEnvironmentCapture.cpp
@@ -1259,11 +1259,7 @@ FMatrix CalcCubeFaceViewRotationMatrix(ECubeFace Face)
 
 FMatrix GetCubeProjectionMatrix(float HalfFovDeg, float CubeMapSize, float NearPlane)
 {
-	if ((bool)ERHIZBuffer::IsInverted)
-	{
-		return FReversedZPerspectiveMatrix(HalfFovDeg * float(PI) / 180.0f, CubeMapSize, CubeMapSize, NearPlane);
-	}
-	return FPerspectiveMatrix(HalfFovDeg, CubeMapSize, CubeMapSize, NearPlane);
+	return FReversedZPerspectiveMatrix(HalfFovDeg * float(PI) / 180.0f, CubeMapSize, CubeMapSize, NearPlane);
 }
 
 void CaptureSceneIntoScratchCubemap(
diff --git a/Engine/Source/Runtime/Renderer/Private/RendererUtils.cpp b/Engine/Source/Runtime/Renderer/Private/RendererUtils.cpp
--- a/Engine/Source/Runtime/Renderer/Private/RendererUtils.cpp
+++ b/Engine/Source/Runtime/Renderer/Private/RendererUtils.cpp
@@ -199,39 +199,20 @@ FDepthBounds::FDepthBoundsValues FDepthBounds::CalculateNearFarDepthExcludingSky
 {
 	FDepthBounds::FDepthBoundsValues Values;
 
-	if (bool(ERHIZBuffer::IsInverted))
-	{
-		float SmallestFloatAbove0;
+	float SmallestFloatAbove0;
 
-		if (GPixelFormats[PF_DepthStencil].bIs24BitUnormDepthStencil)
-		{
-			SmallestFloatAbove0 = 1.0f / 16777215.0f;			// 24bit norm depth
-		}
-		else
-		{
-			SmallestFloatAbove0 = 1.1754943508e-38;				// 32bit float depth
-		}
-
-		Values.MinDepth = SmallestFloatAbove0;
-		Values.MaxDepth = float(ERHIZBuffer::NearPlane);
+	if (GPixelFormats[PF_DepthStencil].bIs24BitUnormDepthStencil)
+	{
+		SmallestFloatAbove0 = 1.0f / 16777215.0f;			// 24bit norm depth
 	}
 	else
 	{
-		float SmallestFloatBelow1;
-
-		if (GPixelFormats[PF_DepthStencil].bIs24BitUnormDepthStencil)
-		{
-			SmallestFloatBelow1 = 16777214.0f / 16777215.0f;	// 24bit norm depth
-		}
-		else
-		{
-			SmallestFloatBelow1 = 0.9999999404;					// 32bit float depth
-		}
-
-		Values.MinDepth = float(ERHIZBuffer::NearPlane);
-		Values.MaxDepth = SmallestFloatBelow1;
+		SmallestFloatAbove0 = 1.1754943508e-38;				// 32bit float depth
 	}
 
+	Values.MinDepth = SmallestFloatAbove0;
+	Values.MaxDepth = 1.f;
+
 	return Values;
 }
 
diff --git a/Engine/Source/Runtime/Renderer/Private/SceneCaptureRendering.cpp b/Engine/Source/Runtime/Renderer/Private/SceneCaptureRendering.cpp
--- a/Engine/Source/Runtime/Renderer/Private/SceneCaptureRendering.cpp
+++ b/Engine/Source/Runtime/Renderer/Private/SceneCaptureRendering.cpp
@@ -513,7 +513,6 @@ static void UpdateSceneCaptureContent_RenderThread(
 
 static void BuildOrthoMatrix(FIntPoint InRenderTargetSize, float InOrthoWidth, int32 InTileID, int32 InNumXTiles, int32 InNumYTiles, FMatrix& OutProjectionMatrix)
 {
-	check((int32)ERHIZBuffer::IsInverted);
 	float const XAxisMultiplier = 1.0f;
 	float const YAxisMultiplier = InRenderTargetSize.X / float(InRenderTargetSize.Y);
 
@@ -571,28 +570,14 @@ void BuildProjectionMatrix(FIntPoint InRenderTargetSize, float InFOV, float InNe
 	float const XAxisMultiplier = 1.0f;
 	float const YAxisMultiplier = InRenderTargetSize.X / float(InRenderTargetSize.Y);
 
-	if ((int32)ERHIZBuffer::IsInverted)
-	{
-		OutProjectionMatrix = FReversedZPerspectiveMatrix(
-			InFOV,
-			InFOV,
-			XAxisMultiplier,
-			YAxisMultiplier,
-			InNearClippingPlane,
-			InNearClippingPlane
-			);
-	}
-	else
-	{
-		OutProjectionMatrix = FPerspectiveMatrix(
-			InFOV,
-			InFOV,
-			XAxisMultiplier,
-			YAxisMultiplier,
-			InNearClippingPlane,
-			InNearClippingPlane
-			);
-	}
+	OutProjectionMatrix = FReversedZPerspectiveMatrix(
+		InFOV,
+		InFOV,
+		XAxisMultiplier,
+		YAxisMultiplier,
+		InNearClippingPlane,
+		InNearClippingPlane
+		);
 }
 
 void GetShowOnlyAndHiddenComponents(USceneCaptureComponent* SceneCaptureComponent, TSet<FPrimitiveComponentId>& HiddenPrimitives, TOptional<TSet<FPrimitiveComponentId>>& ShowOnlyPrimitives)
diff --git a/Engine/Source/Runtime/Renderer/Private/SceneHitProxyRendering.cpp b/Engine/Source/Runtime/Renderer/Private/SceneHitProxyRendering.cpp
--- a/Engine/Source/Runtime/Renderer/Private/SceneHitProxyRendering.cpp
+++ b/Engine/Source/Runtime/Renderer/Private/SceneHitProxyRendering.cpp
@@ -306,7 +306,7 @@ static void DoRenderHitProxies(
 				RHICmdList.SetViewport(View.ViewRect.Min.X, View.ViewRect.Min.Y, 0.0f, View.ViewRect.Max.X, View.ViewRect.Max.Y, 1.0f);
 				DrawClearQuad(RHICmdList, true, FLinearColor::White, false, 0, false, 0, HitProxyTextureExtent, FIntRect());
 				// Clear the depth buffer for each DPG.
-				DrawClearQuad(RHICmdList, false, FLinearColor(), true, (float)ERHIZBuffer::FarPlane, true, 0, HitProxyTextureExtent, FIntRect());
+				DrawClearQuad(RHICmdList, false, FLinearColor(), true, 0.f, true, 0, HitProxyTextureExtent, FIntRect());
 			}
 		});
 	}
diff --git a/Engine/Source/Runtime/Renderer/Private/SceneOcclusion.cpp b/Engine/Source/Runtime/Renderer/Private/SceneOcclusion.cpp
--- a/Engine/Source/Runtime/Renderer/Private/SceneOcclusion.cpp
+++ b/Engine/Source/Runtime/Renderer/Private/SceneOcclusion.cpp
@@ -729,7 +729,6 @@ static bool AllocatePlanarReflectionOcclusionQuery(const FViewInfo& View, const
 		else if (!View.IsPerspectiveProjection())
 		{
 			// Transform parallel near plane
-			static_assert((int32)ERHIZBuffer::IsInverted != 0, "Check equation for culling!");
 			bAllowBoundsTest = View.WorldToScreen(OcclusionBounds.Origin).Z - View.ViewMatrices.GetProjectionMatrix().M[2][2] * OcclusionBounds.SphereRadius < 1;
 		}
 		else
diff --git a/Engine/Source/Runtime/Renderer/Private/SceneTextures.cpp b/Engine/Source/Runtime/Renderer/Private/SceneTextures.cpp
--- a/Engine/Source/Runtime/Renderer/Private/SceneTextures.cpp
+++ b/Engine/Source/Runtime/Renderer/Private/SceneTextures.cpp
@@ -769,7 +769,7 @@ void FSceneTextures::InitializeViewFamily(FRDGBuilder& GraphBuilder, FViewFamily
 
 	if (Config.bRequiresDepthAux)
 	{
-		const float FarDepth = (float)ERHIZBuffer::FarPlane;
+		const float FarDepth = 0.f;
 		const FLinearColor FarDepthColor(FarDepth, FarDepth, FarDepth, FarDepth);
 		ETextureCreateFlags MemorylessFlag = TexCreate_None;
 		if (IsMobileDeferredShadingEnabled(Config.ShaderPlatform) || (Config.NumSamples > 1 && Config.bMemorylessMSAA))
diff --git a/Engine/Source/Runtime/Renderer/Private/SceneVisibility.cpp b/Engine/Source/Runtime/Renderer/Private/SceneVisibility.cpp
--- a/Engine/Source/Runtime/Renderer/Private/SceneVisibility.cpp
+++ b/Engine/Source/Runtime/Renderer/Private/SceneVisibility.cpp
@@ -2928,7 +2928,6 @@ bool FGPUOcclusionPacket::OcclusionCullPrimitive(VisitorType& Visitor, FOcclusio
 			else if (!View.IsPerspectiveProjection())
 			{
 				// Transform parallel near plane
-				static_assert((int32)ERHIZBuffer::IsInverted != 0, "Check equation for culling!");
 				bAllowBoundsTest = View.WorldToScreen(OcclusionBounds.Origin).Z - View.ViewMatrices.GetProjectionMatrix().M[2][2] * OcclusionBounds.SphereRadius < 1;
 			}
 			else
diff --git a/Engine/Source/Runtime/Renderer/Private/Shadows/FirstPersonSelfShadow.cpp b/Engine/Source/Runtime/Renderer/Private/Shadows/FirstPersonSelfShadow.cpp
--- a/Engine/Source/Runtime/Renderer/Private/Shadows/FirstPersonSelfShadow.cpp
+++ b/Engine/Source/Runtime/Renderer/Private/Shadows/FirstPersonSelfShadow.cpp
@@ -272,8 +272,8 @@ FFirstPersonSelfShadowInputs CreateFirstPersonSelfShadowInputs(FRDGBuilder& Grap
 		uint32 CheckerboardMode = 0;
 		switch (FMath::Clamp(CVarFirstPersonSelfShadowCheckerboardMode.GetValueOnRenderThread(), 0, 2))
 		{
-		case 0: CheckerboardMode = (bool)ERHIZBuffer::IsInverted ? 1 : 0; break;
-		case 1: CheckerboardMode = (bool)ERHIZBuffer::IsInverted ? 0 : 1; break;
+		case 0: CheckerboardMode = 1; break;
+		case 1: CheckerboardMode = 0; break;
 		case 2: CheckerboardMode = 2; break;
 		}
 
@@ -305,7 +305,6 @@ FFirstPersonSelfShadowInputs CreateFirstPersonSelfShadowInputs(FRDGBuilder& Grap
 
 void RenderFirstPersonSelfShadow(FRDGBuilder& GraphBuilder, const FSceneRendererBase& SceneRenderer, const TArray<FViewInfo>& Views, FRDGTextureRef ScreenShadowMaskTexture, const FFirstPersonSelfShadowInputs& Inputs, const FLightSceneInfo& LightSceneInfo)
 {
-	static_assert((bool)ERHIZBuffer::IsInverted, "Inverted depth buffer is assumed because FPixelShaderUtils::AddFullscreenPass is drawing at depth 0!");
 	check(ScreenShadowMaskTexture);
 	check(LightCastsFirstPersonSelfShadow(LightSceneInfo));
 
diff --git a/Engine/Source/Runtime/Renderer/Private/SystemTextures.cpp b/Engine/Source/Runtime/Renderer/Private/SystemTextures.cpp
--- a/Engine/Source/Runtime/Renderer/Private/SystemTextures.cpp
+++ b/Engine/Source/Runtime/Renderer/Private/SystemTextures.cpp
@@ -371,7 +371,7 @@ void FSystemTextures::InitializeCommonTextures(FRHICommandListImmediate& RHICmdL
 	}
 
 	{
-		DepthDummy = static_cast<bool>(ERHIZBuffer::IsInverted) ? BlackDummy : WhiteDummy;
+		DepthDummy = BlackDummy;
 		BlackDepthCube = CreateRenderTarget(GBlackTextureDepthCube->TextureRHI, TEXT("BlackDepthCube"));
 	}
 
diff --git a/Engine/Source/Runtime/Renderer/Private/VolumetricCloudRendering.cpp b/Engine/Source/Runtime/Renderer/Private/VolumetricCloudRendering.cpp
--- a/Engine/Source/Runtime/Renderer/Private/VolumetricCloudRendering.cpp
+++ b/Engine/Source/Runtime/Renderer/Private/VolumetricCloudRendering.cpp
@@ -708,7 +708,7 @@ void SetupDefaultRenderVolumetricCloudGlobalParameters(FRDGBuilder& GraphBuilder
 	FRDGTextureRef BlackDummy = GSystemTextures.GetBlackDummy(GraphBuilder);
 	VolumetricCloudParams.VolumetricCloud = CloudInfo.GetVolumetricCloudCommonShaderParameters();
 	VolumetricCloudParams.SceneDepthTexture = BlackDummy;
-	VolumetricCloudParams.SceneDepthMinAndMaxTexture = ((bool)ERHIZBuffer::IsInverted ? GSystemTextures.GetBlackDummy(GraphBuilder) : GSystemTextures.GetWhiteDummy(GraphBuilder));
+	VolumetricCloudParams.SceneDepthMinAndMaxTexture = GSystemTextures.GetBlackDummy(GraphBuilder);
 	VolumetricCloudParams.CloudShadowTexture0 = BlackDummy;
 	VolumetricCloudParams.CloudShadowTexture1 = BlackDummy;
 	VolumetricCloudParams.CloudBilinearTextureSampler = TStaticSamplerState<SF_Bilinear>::GetRHI();
@@ -2938,7 +2938,7 @@ bool FSceneRenderer::RenderVolumetricCloud(
 
 				if (CloudRC.bCloudDebugViewModeEnabled)
 				{
-					CloudRC.SceneDepthZ = (bool)ERHIZBuffer::IsInverted ? GSystemTextures.GetBlackDummy(GraphBuilder) : GSystemTextures.GetWhiteDummy(GraphBuilder);
+					CloudRC.SceneDepthZ = GSystemTextures.GetBlackDummy(GraphBuilder);
 					CloudRC.SceneDepthMinAndMax = CloudRC.SceneDepthZ;
 				}
 
diff --git a/Engine/Source/Runtime/Renderer/Private/VolumetricRenderTarget.cpp b/Engine/Source/Runtime/Renderer/Private/VolumetricRenderTarget.cpp
--- a/Engine/Source/Runtime/Renderer/Private/VolumetricRenderTarget.cpp
+++ b/Engine/Source/Runtime/Renderer/Private/VolumetricRenderTarget.cpp
@@ -845,7 +845,7 @@ void ReconstructVolumetricRenderTarget(
 
 		const bool bVisualizeConservativeDensity = ShouldViewVisualizeVolumetricCloudConservativeDensity(ViewInfo, ViewInfo.Family->EngineShowFlags);
 		PassParameters->HalfResDepthTexture = bVisualizeConservativeDensity ?
-			((bool)ERHIZBuffer::IsInverted ? SystemTextures.Black : SystemTextures.White) :
+			SystemTextures.Black :
 			((VolumetricCloudRT.GetMode() == 0 || VolumetricCloudRT.GetMode() == 3) ?
 				HalfResolutionDepthCheckerboardMinMaxTexture :
 				SceneDepthTexture);
