[HWRT] Move FRayTracingGometryInstance and ERayTracingInstanceFlags out of RHIResources.h.
- this class and enum are no longer used at RHI level, instead it goes through FRayTracingInstanceBufferBuilder in Renderer module.

#Synced-CL 50328147

diff --git a/Engine/Source/Runtime/D3D12RHI/Private/D3D12RayTracing.cpp b/Engine/Source/Runtime/D3D12RHI/Private/D3D12RayTracing.cpp
--- a/Engine/Source/Runtime/D3D12RHI/Private/D3D12RayTracing.cpp
+++ b/Engine/Source/Runtime/D3D12RHI/Private/D3D12RayTracing.cpp
@@ -3201,34 +3201,6 @@ void FD3D12Device::CleanupRayTracing()
 	}
 }
 
-static D3D12_RAYTRACING_INSTANCE_FLAGS TranslateRayTracingInstanceFlags(ERayTracingInstanceFlags InFlags)
-{
-	D3D12_RAYTRACING_INSTANCE_FLAGS Result = D3D12_RAYTRACING_INSTANCE_FLAG_NONE;
-
-	if (EnumHasAnyFlags(InFlags, ERayTracingInstanceFlags::TriangleCullDisable))
-	{
-		Result |= D3D12_RAYTRACING_INSTANCE_FLAG_TRIANGLE_CULL_DISABLE;
-	}
-
-	if (!EnumHasAnyFlags(InFlags, ERayTracingInstanceFlags::TriangleCullReverse))
-	{
-		// Counterclockwise is the default for UE. Reversing culling is achieved by *not* setting this flag.
-		Result |= D3D12_RAYTRACING_INSTANCE_FLAG_TRIANGLE_FRONT_COUNTERCLOCKWISE;
-	}
-
-	if (EnumHasAnyFlags(InFlags, ERayTracingInstanceFlags::ForceOpaque))
-	{
-		Result |= D3D12_RAYTRACING_INSTANCE_FLAG_FORCE_OPAQUE;
-	}
-
-	if (EnumHasAnyFlags(InFlags, ERayTracingInstanceFlags::ForceNonOpaque))
-	{
-		Result |= D3D12_RAYTRACING_INSTANCE_FLAG_FORCE_NON_OPAQUE;
-	}
-
-	return Result;
-}
-
 FRayTracingAccelerationStructureSize FD3D12DynamicRHI::RHICalcRayTracingSceneSize(const FRayTracingSceneInitializer& Initializer)
 {
 	D3D12_BUILD_RAYTRACING_ACCELERATION_STRUCTURE_INPUTS BuildInputs = {};
diff --git a/Engine/Source/Runtime/RHI/Private/DynamicRHI.cpp b/Engine/Source/Runtime/RHI/Private/DynamicRHI.cpp
--- a/Engine/Source/Runtime/RHI/Private/DynamicRHI.cpp
+++ b/Engine/Source/Runtime/RHI/Private/DynamicRHI.cpp
@@ -23,9 +23,6 @@
 #include "RHITextureReference.h"
 #include "RHIAllocators.h"
 
-static_assert(sizeof(FRayTracingGeometryInstance) <= 104,
-	"Ray tracing instance descriptor is expected to be no more than 104 bytes, as there may be a very large number of them.");
-
 #ifndef PLATFORM_ALLOW_NULL_RHI
 	#define PLATFORM_ALLOW_NULL_RHI		0
 #endif
diff --git a/Engine/Source/Runtime/RHI/Public/RHIContext.h b/Engine/Source/Runtime/RHI/Public/RHIContext.h
--- a/Engine/Source/Runtime/RHI/Public/RHIContext.h
+++ b/Engine/Source/Runtime/RHI/Public/RHIContext.h
@@ -26,7 +26,6 @@ class FRHIRenderTargetView;
 class FRHISetRenderTargetsInfo;
 class FRHIShaderBindingLayout;
 struct FViewportBounds;
-struct FRayTracingGeometryInstance;
 struct FRayTracingShaderBindings;
 struct FRayTracingGeometrySegment;
 struct FRayTracingGeometryBuildParams;
diff --git a/Engine/Source/Runtime/RHI/Public/RHIResources.h b/Engine/Source/Runtime/RHI/Public/RHIResources.h
--- a/Engine/Source/Runtime/RHI/Public/RHIResources.h
+++ b/Engine/Source/Runtime/RHI/Public/RHIResources.h
@@ -3326,62 +3326,6 @@ inline ERHIDescriptorType RHIDescriptorTypeFromViewDesc(const FRHIViewDesc& InVi
 // Ray tracing resources
 //
 
-enum class ERayTracingInstanceFlags : uint8
-{
-	None = 0,
-	TriangleCullDisable = 1 << 1, // No back face culling. Triangle is visible from both sides.
-	TriangleCullReverse = 1 << 2, // Makes triangle front-facing if its vertices are counterclockwise from ray origin.
-	ForceOpaque = 1 << 3, // Disable any-hit shader invocation for this instance.
-	ForceNonOpaque = 1 << 4, // Force any-hit shader invocation even if geometries inside the instance were marked opaque.
-};
-ENUM_CLASS_FLAGS(ERayTracingInstanceFlags);
-
-class FRHIRayTracingGeometry;
-/**
-* High level descriptor of one or more instances of a mesh in a ray tracing scene.
-* All instances covered by this descriptor will share shader bindings, but may have different transforms and user data.
-*/
-struct FRayTracingGeometryInstance
-{
-	FRHIRayTracingGeometry* GeometryRHI = nullptr;
-
-	int32 InstanceContributionToHitGroupIndex = INDEX_NONE;
-
-	// A single physical mesh may be duplicated many times in the scene with different transforms and user data.
-	// All copies share the same shader binding table entries and therefore will have the same material and shader resources.
-	TArrayView<const FMatrix> Transforms;
-
-	// Offsets into the scene's instance scene data buffer used to get instance transforms from GPUScene
-	// If BaseInstanceSceneDataOffset != -1, instances are assumed to be continuous.
-	int32 BaseInstanceSceneDataOffset = -1;
-	TArrayView<const uint32> InstanceSceneDataOffsets;
-
-	// Conservative number of instances. Some of the actual instances may be made inactive if GPU transforms are used.
-	// Must be less or equal to number of entries in Transforms view if CPU transform data is used.
-	uint32 NumTransforms = 0;
-
-	// Each geometry copy can receive a user-provided integer, which can be used to retrieve extra shader parameters or customize appearance.
-	// This data can be retrieved using GetInstanceUserData() in closest/any hit shaders.
-	// If UserData view is empty, then DefaultUserData value will be used for all instances.
-	// If UserData view is used, then it must have the same number of entries as NumInstances.
-	uint32 DefaultUserData = 0;
-	TArrayView<const uint32> UserData;
-
-	// Whether local bounds scale and center translation should be applied to the instance transform.
-	bool bApplyLocalBoundsTransform : 1 = false;
-	// Whether to increment UserData for each instance of this geometry (only applied when using DefaultUserData)
-	bool bIncrementUserDataPerInstance : 1 = false;
-
-	bool bUsesLightingChannels : 1 = false;
-
-	// Mask that will be tested against one provided to TraceRay() in shader code.
-	// If binary AND of instance mask with ray mask is zero, then the instance is considered not intersected / invisible.
-	uint8 Mask = 0xFF;
-
-	// Flags to control triangle back face culling, whether to allow any-hit shaders, etc.
-	ERayTracingInstanceFlags Flags = ERayTracingInstanceFlags::None;
-};
-
 struct FRayTracingGeometryOfflineDataHeader
 {
 	uint32 Reserved[6] = { 0xFFFFFFFF, 0xFFFFFFFF, 0xFFFFFFFF, 0xFFFFFFFF, 0xFFFFFFFF, 0xFFFFFFFF };
diff --git a/Engine/Source/Runtime/Renderer/Private/Tests/RayTracingTestbed.cpp b/Engine/Source/Runtime/Renderer/Private/Tests/RayTracingTestbed.cpp
--- a/Engine/Source/Runtime/Renderer/Private/Tests/RayTracingTestbed.cpp
+++ b/Engine/Source/Runtime/Renderer/Private/Tests/RayTracingTestbed.cpp
@@ -41,6 +41,7 @@ IMPLEMENT_CUSTOM_SIMPLE_AUTOMATION_TEST(FRayTracingTestbed, FRayTracingTestbedBa
 
 #if RHI_RAYTRACING
 
+#include "RayTracingGeometryInstance.h"
 #include "RayTracing/RayTracingBasicShaders.h"
 #include "RayTracingInstanceBufferUtil.h"
 #include "RayTracingDefinitions.h"
diff --git a/Engine/Source/Runtime/Renderer/Public/PrimitiveSceneInfo.h b/Engine/Source/Runtime/Renderer/Public/PrimitiveSceneInfo.h
--- a/Engine/Source/Runtime/Renderer/Public/PrimitiveSceneInfo.h
+++ b/Engine/Source/Runtime/Renderer/Public/PrimitiveSceneInfo.h
@@ -16,6 +16,7 @@
 #include "RendererInterface.h"
 #include "ShaderParameterMacros.h"
 #include "RayTracingMeshDrawCommands.h"
+#include "RayTracingGeometryInstance.h"
 
 enum class ERayTracingPrimitiveFlags : uint8;
 
diff --git a/Engine/Source/Runtime/Renderer/Public/RayTracingGeometryInstance.h b/Engine/Source/Runtime/Renderer/Public/RayTracingGeometryInstance.h
new file mode 100644
--- /dev/null
+++ b/Engine/Source/Runtime/Renderer/Public/RayTracingGeometryInstance.h
@@ -0,0 +1,65 @@
+// Copyright Epic Games, Inc. All Rights Reserved.
+
+#pragma once
+
+#include "CoreMinimal.h"
+
+class FRHIRayTracingGeometry;
+
+enum class ERayTracingInstanceFlags : uint8
+{
+	None = 0,
+	TriangleCullDisable = 1 << 1, // No back face culling. Triangle is visible from both sides.
+	TriangleCullReverse = 1 << 2, // Makes triangle front-facing if its vertices are counterclockwise from ray origin.
+	ForceOpaque = 1 << 3, // Disable any-hit shader invocation for this instance.
+	ForceNonOpaque = 1 << 4, // Force any-hit shader invocation even if geometries inside the instance were marked opaque.
+};
+ENUM_CLASS_FLAGS(ERayTracingInstanceFlags);
+
+/**
+* High level descriptor of one or more instances of a mesh in a ray tracing scene.
+* All instances covered by this descriptor will share shader bindings, but may have different transforms and user data.
+*/
+struct FRayTracingGeometryInstance
+{
+	FRHIRayTracingGeometry* GeometryRHI = nullptr;
+
+	int32 InstanceContributionToHitGroupIndex = INDEX_NONE;
+
+	// A single physical mesh may be duplicated many times in the scene with different transforms and user data.
+	// All copies share the same shader binding table entries and therefore will have the same material and shader resources.
+	TArrayView<const FMatrix> Transforms;
+
+	// Offsets into the scene's instance scene data buffer used to get instance transforms from GPUScene
+	// If BaseInstanceSceneDataOffset != -1, instances are assumed to be continuous.
+	int32 BaseInstanceSceneDataOffset = -1;
+	TArrayView<const uint32> InstanceSceneDataOffsets;
+
+	// Conservative number of instances. Some of the actual instances may be made inactive if GPU transforms are used.
+	// Must be less or equal to number of entries in Transforms view if CPU transform data is used.
+	uint32 NumTransforms = 0;
+
+	// Each geometry copy can receive a user-provided integer, which can be used to retrieve extra shader parameters or customize appearance.
+	// This data can be retrieved using GetInstanceUserData() in closest/any hit shaders.
+	// If UserData view is empty, then DefaultUserData value will be used for all instances.
+	// If UserData view is used, then it must have the same number of entries as NumInstances.
+	uint32 DefaultUserData = 0;
+	TArrayView<const uint32> UserData;
+
+	// Whether local bounds scale and center translation should be applied to the instance transform.
+	bool bApplyLocalBoundsTransform : 1 = false;
+	// Whether to increment UserData for each instance of this geometry (only applied when using DefaultUserData)
+	bool bIncrementUserDataPerInstance : 1 = false;
+
+	bool bUsesLightingChannels : 1 = false;
+
+	// Mask that will be tested against one provided to TraceRay() in shader code.
+	// If binary AND of instance mask with ray mask is zero, then the instance is considered not intersected / invisible.
+	uint8 Mask = 0xFF;
+
+	// Flags to control triangle back face culling, whether to allow any-hit shaders, etc.
+	ERayTracingInstanceFlags Flags = ERayTracingInstanceFlags::None;
+};
+
+static_assert(sizeof(FRayTracingGeometryInstance) <= 104,
+	"Ray tracing instance descriptor is expected to be no more than 104 bytes, as there may be a very large number of them.");
diff --git a/Engine/Source/Runtime/Renderer/Public/RayTracingInstanceBufferUtil.h b/Engine/Source/Runtime/Renderer/Public/RayTracingInstanceBufferUtil.h
--- a/Engine/Source/Runtime/Renderer/Public/RayTracingInstanceBufferUtil.h
+++ b/Engine/Source/Runtime/Renderer/Public/RayTracingInstanceBufferUtil.h
@@ -15,6 +15,7 @@ struct FDFVector3;
 
 struct FRayTracingInstanceGroup;
 struct FRayTracingInstanceDescriptor;
+struct FRayTracingGeometryInstance;
 
 struct FRayTracingInstanceBufferBuilderInitializer
 {
diff --git a/Engine/Source/Runtime/VulkanRHI/Private/VulkanRayTracing.cpp b/Engine/Source/Runtime/VulkanRHI/Private/VulkanRayTracing.cpp
--- a/Engine/Source/Runtime/VulkanRHI/Private/VulkanRayTracing.cpp
+++ b/Engine/Source/Runtime/VulkanRHI/Private/VulkanRayTracing.cpp
@@ -665,34 +665,6 @@ static void GetTLASBuildData(
 		&BuildData.SizesInfo);
 }
 
-static VkGeometryInstanceFlagsKHR TranslateRayTracingInstanceFlags(ERayTracingInstanceFlags InFlags)
-{
-	VkGeometryInstanceFlagsKHR Result = 0;
-
-	if (EnumHasAnyFlags(InFlags, ERayTracingInstanceFlags::TriangleCullDisable))
-	{
-		Result |= VK_GEOMETRY_INSTANCE_TRIANGLE_FACING_CULL_DISABLE_BIT_KHR;
-	}
-
-	if (!EnumHasAnyFlags(InFlags, ERayTracingInstanceFlags::TriangleCullReverse))
-	{
-		// Counterclockwise is the default for UE.
-		Result |= VK_GEOMETRY_INSTANCE_TRIANGLE_FRONT_COUNTERCLOCKWISE_BIT_KHR;
-	}
-
-	if (EnumHasAnyFlags(InFlags, ERayTracingInstanceFlags::ForceOpaque))
-	{
-		Result |= VK_GEOMETRY_INSTANCE_FORCE_OPAQUE_BIT_KHR;
-	}
-
-	if (EnumHasAnyFlags(InFlags, ERayTracingInstanceFlags::ForceNonOpaque))
-	{
-		Result |= VK_GEOMETRY_INSTANCE_FORCE_NO_OPAQUE_BIT_KHR;
-	}
-
-	return Result;
-}
-
 FVulkanRayTracingScene::FVulkanRayTracingScene(FRayTracingSceneInitializer InInitializer, FVulkanDevice& InDevice)
 	: Device(InDevice)
 	, Initializer(MoveTemp(InInitializer))
