Fix issue in virtual texture upload cache where tiles in PendingRelease were not being released quickly due to incorrect early out from the iteration loop.
Also added some more virtual texture stats to show reasons why updates are throttled.
Updated default value for r.VT.MaxUploadRequests to be more in line with r.VT.MaxUploadMemory.
#Synced-CL 50344094

diff --git a/Engine/Source/Runtime/Engine/Private/VT/VirtualTextureUploadCache.cpp b/Engine/Source/Runtime/Engine/Private/VT/VirtualTextureUploadCache.cpp
--- a/Engine/Source/Runtime/Engine/Private/VT/VirtualTextureUploadCache.cpp
+++ b/Engine/Source/Runtime/Engine/Private/VT/VirtualTextureUploadCache.cpp
@@ -34,7 +34,7 @@ static TAutoConsoleVariable<int32> CVarVTMaxUploadMemory(
 
 static TAutoConsoleVariable<int32> CVarMaxUploadRequests(
 	TEXT("r.VT.MaxUploadRequests"),
-	2000,
+	6000,
 	TEXT("Maximum number of virtual texture tile upload requests that can be in flight."),
 	ECVF_RenderThreadSafe);
 
@@ -599,7 +599,7 @@ void FVirtualTextureUploadCache::UpdateFreeList(FRHICommandListBase& RHICmdList,
 	{
 		if (CurrentFrame < It->FrameSubmitted + ReleaseFrameDelay)
 		{
-			break;
+			continue;
 		}
 
 		TileAllocator.Free(RHICmdList, It->TileHandle);
diff --git a/Engine/Source/Runtime/Renderer/Private/VT/VirtualTextureSystem.cpp b/Engine/Source/Runtime/Renderer/Private/VT/VirtualTextureSystem.cpp
--- a/Engine/Source/Runtime/Renderer/Private/VT/VirtualTextureSystem.cpp
+++ b/Engine/Source/Runtime/Renderer/Private/VT/VirtualTextureSystem.cpp
@@ -62,8 +62,12 @@ DECLARE_DWORD_COUNTER_STAT(TEXT("Num page prefetch"), STAT_NumPagePrefetch, STAT
 DECLARE_DWORD_COUNTER_STAT(TEXT("Num page update"), STAT_NumPageUpdate, STATGROUP_VirtualTexturing);
 DECLARE_DWORD_COUNTER_STAT(TEXT("Num mapped page update"), STAT_NumMappedPageUpdate, STATGROUP_VirtualTexturing);
 DECLARE_DWORD_COUNTER_STAT(TEXT("Num continuous page update"), STAT_NumContinuousPageUpdate, STATGROUP_VirtualTexturing);
-DECLARE_DWORD_COUNTER_STAT(TEXT("Num page allocation fails"), STAT_NumPageAllocateFails, STATGROUP_VirtualTexturing);
-
+DECLARE_DWORD_COUNTER_STAT(TEXT("Num page produced"), STAT_NumPagesProduced, STATGROUP_VirtualTexturing);
+DECLARE_DWORD_COUNTER_STAT(TEXT("Num page pending"), STAT_NumPagesPending, STATGROUP_VirtualTexturing);
+DECLARE_DWORD_COUNTER_STAT(TEXT("Num page throttled"), STAT_NumPagesThrottled, STATGROUP_VirtualTexturing);
+DECLARE_DWORD_COUNTER_STAT(TEXT("Num page saturated"), STAT_NumPagesSaturated, STATGROUP_VirtualTexturing);
+DECLARE_DWORD_COUNTER_STAT(TEXT("Num page invalid"), STAT_NumPagesInvalid, STATGROUP_VirtualTexturing);
+DECLARE_DWORD_COUNTER_STAT(TEXT("Num page fail allocate"), STAT_NumPagesFailAllocate, STATGROUP_VirtualTexturing);
 DECLARE_DWORD_COUNTER_STAT(TEXT("Num stacks requested"), STAT_NumStacksRequested, STATGROUP_VirtualTexturing);
 DECLARE_DWORD_COUNTER_STAT(TEXT("Num stacks produced"), STAT_NumStacksProduced, STATGROUP_VirtualTexturing);
 
@@ -2420,9 +2424,13 @@ void FVirtualTextureSystem::SubmitRequests(FRHICommandList& RHICmdList, ERHIFeat
 		const bool bSyncProduceLockedTiles = ShouldSyncProduceLockedTiles();
 		const uint32 MaxPagesProduced = Settings.MaxPagesProduced;
 		const uint32 PageFreeThreshold = VirtualTextureScalability::GetPageFreeThreshold();
-		uint32 NumStacksProduced = 0u;
+		uint32 NumPagesPending = 0u;
+		uint32 NumPagesThrottled = 0u;
+		uint32 NumPagesInvalid = 0u;
+		uint32 NumPagesSaturated = 0u;
 		uint32 NumPagesProduced = 0u;
-		uint32 NumPageAllocateFails = 0u;
+		uint32 NumPagesFailAllocate = 0u;
+		uint32 NumStacksProduced = 0u;
 
 		for (uint32 RequestIndex = 0u; RequestIndex < RequestList->GetNumLoadRequests(); ++RequestIndex)
 		{
@@ -2444,28 +2452,39 @@ void FVirtualTextureSystem::SubmitRequests(FRHICommandList& RHICmdList, ERHIFeat
 
 			const EVTRequestPagePriority Priority = bLockTile || Settings.bForceSyncPageUpdate ? EVTRequestPagePriority::High : EVTRequestPagePriority::Normal;
 			FVTRequestPageResult RequestPageResult = Producer.GetVirtualTexture()->RequestPageData(RHICmdList, ProducerHandle, ProducerTextureLayerMask, TileToLoad.Local_vLevel, TileToLoad.Local_vAddress, Priority);
-			if (RequestPageResult.Status == EVTRequestPageStatus::Pending && bForceProduceTile)
+			if (RequestPageResult.Status == EVTRequestPageStatus::Pending)
 			{
-				// If we're forcing production of this tile, we're OK producing data now (and possibly waiting) as long as data is pending
-				RequestPageResult.Status = EVTRequestPageStatus::Available;
-				bWaitForProducers = true;
+				if (bForceProduceTile)
+				{
+					// If we're forcing production of this tile, we're OK producing data now (and possibly waiting) as long as data is pending
+					RequestPageResult.Status = EVTRequestPageStatus::Available;
+					bWaitForProducers = true;
+				}
+				else
+				{
+					NumPagesPending++;
+				}
 			}
 
 			if (RequestPageResult.Status == EVTRequestPageStatus::Available && !bForceProduceTile && NumPagesProduced >= MaxPagesProduced)
 			{
 				// Don't produce non-locked pages yet, if we're over our limit
 				RequestPageResult.Status = EVTRequestPageStatus::Pending;
+				NumPagesThrottled++;
 			}
 
 			bool bTileLoaded = false;
 			bool bTileInvalid = false;
 			if (RequestPageResult.Status == EVTRequestPageStatus::Invalid)
 			{
-				//checkf(!bLockTile, TEXT("Tried to lock an invalid VT tile"));
-
 				bTileInvalid = true;
+				NumPagesInvalid++;
 				UE_LOG(LogVirtualTexturing, Verbose, TEXT("vAddr %i@%i is not a valid request for AllocatedVT but is still requested."), TileToLoad.Local_vAddress, TileToLoad.Local_vLevel);
 			}
+			else if (RequestPageResult.Status == EVTRequestPageStatus::Saturated)
+			{
+				NumPagesSaturated++;
+			}
 			else if (RequestPageResult.Status == EVTRequestPageStatus::Available)
 			{
 				FVTProduceTargetLayer ProduceTarget[VIRTUALTEXTURE_SPACE_MAXLAYERS];
@@ -2504,7 +2523,7 @@ void FVirtualTextureSystem::SubmitRequests(FRHICommandList& RHICmdList, ERHIFeat
 						else
 						{
 							bProduceTargetValid = false;
-							NumPageAllocateFails++;
+							NumPagesFailAllocate++;
 							break;
 						}
 					}
@@ -2624,7 +2643,12 @@ void FVirtualTextureSystem::SubmitRequests(FRHICommandList& RHICmdList, ERHIFeat
 
 		INC_DWORD_STAT_BY(STAT_NumStacksRequested, RequestList->GetNumLoadRequests());
 		INC_DWORD_STAT_BY(STAT_NumStacksProduced, NumStacksProduced);
-		INC_DWORD_STAT_BY(STAT_NumPageAllocateFails, NumPageAllocateFails);
+		INC_DWORD_STAT_BY(STAT_NumPagesProduced, NumPagesProduced);
+		INC_DWORD_STAT_BY(STAT_NumPagesPending, NumPagesPending);
+		INC_DWORD_STAT_BY(STAT_NumPagesThrottled, NumPagesThrottled);
+		INC_DWORD_STAT_BY(STAT_NumPagesSaturated, NumPagesSaturated);
+		INC_DWORD_STAT_BY(STAT_NumPagesInvalid, NumPagesInvalid);
+		INC_DWORD_STAT_BY(STAT_NumPagesFailAllocate, NumPagesFailAllocate);
 	}
 
 	{
