Fix rhi validation error when using the mobile renderer.

#rnx

#Synced-CL 50154827

diff --git a/Engine/Source/Runtime/Renderer/Private/MobileBasePassRendering.cpp b/Engine/Source/Runtime/Renderer/Private/MobileBasePassRendering.cpp
--- a/Engine/Source/Runtime/Renderer/Private/MobileBasePassRendering.cpp
+++ b/Engine/Source/Runtime/Renderer/Private/MobileBasePassRendering.cpp
@@ -554,8 +554,16 @@ void FMobileSceneRenderer::RenderMobileBasePass(FRHICommandList& RHICmdList, con
 	// editor primitives
 	FMeshPassProcessorRenderState DrawRenderState;
 	DrawRenderState.SetBlendState(TStaticBlendStateWriteMask<CW_RGBA>::GetRHI());
-	DrawRenderState.SetDepthStencilAccess(Scene->DefaultBasePassDepthStencilAccess);
-	DrawRenderState.SetDepthStencilState(TStaticDepthStencilState<true, CF_DepthNearOrEqual>::GetRHI());
+	if (bIsFullDepthPrepassEnabled)
+	{
+		DrawRenderState.SetDepthStencilAccess(FExclusiveDepthStencil::DepthRead_StencilWrite);
+		DrawRenderState.SetDepthStencilState(TStaticDepthStencilState<false>::GetRHI());
+	}
+	else
+	{
+		DrawRenderState.SetDepthStencilAccess(FExclusiveDepthStencil::DepthWrite_StencilWrite);
+		DrawRenderState.SetDepthStencilState(TStaticDepthStencilState<true>::GetRHI());
+	}
 	RenderMobileEditorPrimitives(RHICmdList, View, DrawRenderState, InstanceCullingDrawParams);
 }
 
