Path Tracing: Add support for sun color scale parameter on directional lights

This parameter can be used to tint the sun as visible to the camera. Use the new parameter to also implement the skydome color compensation without a special case branch.

#Synced-CL 50355781

diff --git a/Engine/Shaders/Private/PathTracing/PathTracingCore.ush b/Engine/Shaders/Private/PathTracing/PathTracingCore.ush
--- a/Engine/Shaders/Private/PathTracing/PathTracingCore.ush
+++ b/Engine/Shaders/Private/PathTracing/PathTracingCore.ush
@@ -666,12 +666,7 @@ FPackedPathTracingPayload TraceTransparentRay(inout FPathState PathState, int Bo
 				//       but is not consistent with how reflections behave and the sorting order will be arbitrary
 				FPathTracingPayload LightHitPayload = (FPathTracingPayload)0; // clear all fields
 				LightHitPayload.HitT = LightHit.HitT;
-				LightHitPayload.Radiance = LightHit.Radiance;
-				if ((SceneLights[LightId].Flags & PATHTRACER_FLAG_TYPE_MASK) == PATHTRACING_LIGHT_SKY)
-				{
-					// UE-222576: cancel out scale factor that is baked into the texture for camera rays only
-					LightHitPayload.Radiance /= GetColor(LightId);
-				}
+				LightHitPayload.Radiance = LightHit.Radiance * SceneLights[LightId].CameraColorScale;
 				LightHitPayload.ShadingModelID = SHADINGMODELID_UNLIT;
 				// light hits are fully transparent to match how the light loop works (the light surface does not cast a shadow)
 				LightHitPayload.BSDFOpacity = 1.0;
diff --git a/Engine/Shaders/Private/PathTracing/PathTracingDebug.usf b/Engine/Shaders/Private/PathTracing/PathTracingDebug.usf
--- a/Engine/Shaders/Private/PathTracing/PathTracingDebug.usf
+++ b/Engine/Shaders/Private/PathTracing/PathTracingDebug.usf
@@ -80,10 +80,11 @@ float3 TraceShadowRay(FRayDesc Ray, float PathRoughness, float PreviousPathRough
 		Ray.TMin = POSITIVE_INFINITY;
 		Ray.TMax = POSITIVE_INFINITY;
 	}
+	const uint InstanceInclusionMask = PATHTRACER_MASK_SHADOW | (length2(Ray.Origin) < Pow2(PATH_TRACING_SKIP_HAIR_DISTANCE) ? PATHTRACER_MASK_HAIR_SHADOW : 0);
 	TraceRay(
 		TLAS,
 		RAY_FLAG_ACCEPT_FIRST_HIT_AND_END_SEARCH | RAY_FLAG_SKIP_CLOSEST_HIT_SHADER,
-		PATHTRACER_MASK_SHADOW,
+		InstanceInclusionMask,
 		RAY_TRACING_SHADER_SLOT_SHADOW,
 		RAY_TRACING_NUM_SHADER_SLOTS,
 		MissShaderIndex,
@@ -521,12 +522,7 @@ FPackedPathTracingPayload TraceTransparentRay(inout FPathState PathState, int Bo
 				//       but is not consistent with how reflections behave and the sorting order will be arbitrary
 				FPathTracingPayload LightHitPayload = (FPathTracingPayload)0; // clear all fields
 				LightHitPayload.HitT = LightHit.HitT;
-				LightHitPayload.Radiance = LightHit.Radiance;
-				if ((SceneLights[LightId].Flags & PATHTRACER_FLAG_TYPE_MASK) == PATHTRACING_LIGHT_SKY)
-				{
-					// UE-222576: cancel out scale factor that is baked into the texture for camera rays only
-					LightHitPayload.Radiance /= GetColor(LightId);
-				}
+				LightHitPayload.Radiance = LightHit.Radiance * SceneLights[LightId].CameraColorScale;
 				LightHitPayload.ShadingModelID = SHADINGMODELID_UNLIT;
 				// light hits are fully transparent to match how the light loop works (the light surface does not cast a shadow)
 				LightHitPayload.BSDFOpacity = 1.0;
diff --git a/Engine/Shaders/Shared/RayTracingTypes.h b/Engine/Shaders/Shared/RayTracingTypes.h
--- a/Engine/Shaders/Shared/RayTracingTypes.h
+++ b/Engine/Shaders/Shared/RayTracingTypes.h
@@ -61,8 +61,10 @@ struct FPathTracingLight
 	uint    MissShaderIndex;  // used to implement light functions
 	float2  RectLightAtlasUVScale;  // Rect. light atlas UV transformation
 	float2  RectLightAtlasUVOffset; // Rect. light atlas UV transformation
+	float3  CameraColorScale; // Tinting factor for camera rays only
+	uint    Padding;
 };
-HLSL_STATIC_ASSERT(sizeof(FPathTracingLight) == 112, "Path tracing light structure should be kept as small as possible");
+HLSL_STATIC_ASSERT(sizeof(FPathTracingLight) == 128, "Path tracing light structure should be kept as small as possible");
 
 struct FPathTracingPackedPathState
 {
diff --git a/Engine/Source/Runtime/Renderer/Private/PathTracing.cpp b/Engine/Source/Runtime/Renderer/Private/PathTracing.cpp
--- a/Engine/Source/Runtime/Renderer/Private/PathTracing.cpp
+++ b/Engine/Source/Runtime/Renderer/Private/PathTracing.cpp
@@ -3004,6 +3004,7 @@ void SetLightParameters(
 		DestLight.VolumetricScatteringIntensity = Scene->SkyLight->VolumetricScatteringIntensity;
 		DestLight.IESAtlasIndex = INDEX_NONE;
 		DestLight.MissShaderIndex = 0;
+		DestLight.CameraColorScale = FVector3f(1, 1, 1);
 		if ((Scene->SkyLight->bRealTimeCaptureEnabled && (View.SkyAtmosphereUniformShaderParameters == nullptr ||  !IsSkyAtmosphereHoldout(View.CachedViewUniformShaderParameters->EnvironmentComponentsFlags))) || CVarPathTracingVisibleLights.GetValueOnRenderThread() == 2)
 		{
 			// When using the realtime capture system, always make the skylight visible
@@ -3018,7 +3019,7 @@ void SetLightParameters(
 				// NOTE: this color is already baked into the skylight texture so that importance sampling takes it into account, we pass it in here so that camera rays can factor it out
 				// This is only for the realtime capture case, because otherwise (specified cube map case) we want the displayed texture and lighting to match
 				// NOTE: We want to use the effective light color, not the exposed skylight color SkyLightColor from the View UniformBuffer so that the overall intensity is correct (see jira UE-280734)
-				DestLight.Color = FVector3f(Scene->SkyLight->GetEffectiveLightColor());
+				DestLight.CameraColorScale /= FVector3f(Scene->SkyLight->GetEffectiveLightColor());
 			}
 		}
 	}
@@ -3056,6 +3057,7 @@ void SetLightParameters(
 			DestLight.Flags |= Light.LightSceneInfo->Proxy->GetCastCloudShadows() ? PATHTRACER_FLAG_CAST_CLOUD_SHADOW_MASK : 0;
 			DestLight.IESAtlasIndex = INDEX_NONE;
 			DestLight.MissShaderIndex = 0;
+			DestLight.CameraColorScale = FVector3f(Light.LightSceneInfo->Proxy->GetAtmosphereSunDiskColorScale());
 
 			if (RayTracingLightFunctionMap)
 			{
@@ -3156,6 +3158,8 @@ void SetLightParameters(
 		DestLight.RectLightAtlasUVOffset = FVector2f(0.0f, 0.0f);
 		DestLight.RectLightAtlasUVScale = FVector2f(0.0f, 0.0f);
 
+		DestLight.CameraColorScale = FVector3f(1, 1, 1);
+
 		if (RayTracingLightFunctionMap)
 		{
 			const int32* LightFunctionIndex = RayTracingLightFunctionMap->Find(Light.LightSceneInfo);
