[Front Layer Translucency] Batch all view rect draw calls into a single ClearTranslucencyGBuffer pass.
- it was incorrect to use ERenderTargetLoadAction::NoAction here when rendering multiple views and batching into a single pass avoids having to use ERenderTargetLoadAction::Load.

#Synced-CL 50355052

diff --git a/Engine/Source/Runtime/Renderer/Private/FrontLayerTranslucency.cpp b/Engine/Source/Runtime/Renderer/Private/FrontLayerTranslucency.cpp
--- a/Engine/Source/Runtime/Renderer/Private/FrontLayerTranslucency.cpp
+++ b/Engine/Source/Runtime/Renderer/Private/FrontLayerTranslucency.cpp
@@ -493,29 +493,55 @@ FFrontLayerTranslucencyData FDeferredShadingSceneRenderer::RenderFrontLayerTrans
 			FRDGTextureDesc NormalDesc(FRDGTextureDesc::Create2D(SceneTextures.Config.Extent, NormalFormat, FClearValueBinding::Transparent, TexCreate_ShaderResource | TexCreate_RenderTargetable));
 			Out.Normal = GraphBuilder.CreateTexture(NormalDesc, TEXT("FrontLayerTranslucency.Normal"));
 			Out.SceneDepth = GraphBuilder.CreateTexture(SceneTextures.Depth.Target->Desc, TEXT("FrontLayerTranslucency.SceneDepth"));
+		}
+
+		// Clear ViewRect for each view
+		{
+			TArray<FIntRect, SceneRenderingAllocator>& ViewRectsToClear = GraphBuilder.AllocArray<FIntRect>();
+			ViewRectsToClear.Reserve(InViews.Num());
 
 			for (const FViewInfo& View : InViews)
 			{
 				if (bNeedFrontLayerData > 1 || IsLumenFrontLayerTranslucencyEnabled(View))
 				{
-					FFrontLayerTranslucencyClearGBufferPS::FParameters* PassParameters = GraphBuilder.AllocParameters<FFrontLayerTranslucencyClearGBufferPS::FParameters>();
-					PassParameters->RenderTargets[0] = FRenderTargetBinding(Out.Normal, ERenderTargetLoadAction::ENoAction, 0);
-					PassParameters->RenderTargets.DepthStencil = FDepthStencilBinding(Out.SceneDepth, ERenderTargetLoadAction::ENoAction, FExclusiveDepthStencil::DepthWrite_StencilNop);
-					PassParameters->SceneTexturesStruct = SceneTextures.UniformBuffer;
-	
-					TShaderMapRef<FFrontLayerTranslucencyClearGBufferPS> PixelShader(View.ShaderMap);
-					FPixelShaderUtils::AddFullscreenPass<FFrontLayerTranslucencyClearGBufferPS>(
-						GraphBuilder,
-						View.ShaderMap,
-						RDG_EVENT_NAME("ClearTranslucencyGBuffer"),
-						PixelShader,
-						PassParameters,
-						View.ViewRect,
-						TStaticBlendState<>::GetRHI(),
-						TStaticRasterizerState<FM_Solid, CM_None>::GetRHI(),
-						TStaticDepthStencilState<true, CF_Always>::GetRHI());
+					ViewRectsToClear.Add(View.ViewRect);
 				}
 			}
+
+			if (!ViewRectsToClear.IsEmpty())
+			{
+				const FGlobalShaderMap* GlobalShaderMap = InViews[0].ShaderMap;
+
+				FFrontLayerTranslucencyClearGBufferPS::FParameters* PassParameters = GraphBuilder.AllocParameters<FFrontLayerTranslucencyClearGBufferPS::FParameters>();
+				PassParameters->RenderTargets[0] = FRenderTargetBinding(Out.Normal, ERenderTargetLoadAction::ENoAction, 0);
+				PassParameters->RenderTargets.DepthStencil = FDepthStencilBinding(Out.SceneDepth, ERenderTargetLoadAction::ENoAction, FExclusiveDepthStencil::DepthWrite_StencilNop);
+				PassParameters->SceneTexturesStruct = SceneTextures.UniformBuffer;
+
+				auto PixelShader = GlobalShaderMap->GetShader<FFrontLayerTranslucencyClearGBufferPS>();
+
+				ClearUnusedGraphResources(PixelShader, PassParameters);
+
+				// Not using FPixelShaderUtils::AddFullscreenPass because want to clear multiple view rects in a single pass
+				GraphBuilder.AddPass(
+					RDG_EVENT_NAME("ClearTranslucencyGBuffer"),
+					PassParameters,
+					ERDGPassFlags::Raster,
+					[PassParameters, GlobalShaderMap, PixelShader, &ViewRectsToClear](FRDGAsyncTask, FRHICommandList& RHICmdList)
+					{
+						for (const FIntRect& ViewRect : ViewRectsToClear)
+						{
+							FPixelShaderUtils::DrawFullscreenPixelShader(
+								RHICmdList,
+								GlobalShaderMap,
+								PixelShader,
+								*PassParameters,
+								ViewRect,
+								TStaticBlendState<>::GetRHI(),
+								TStaticRasterizerState<FM_Solid, CM_None>::GetRHI(),
+								TStaticDepthStencilState<true, CF_Always>::GetRHI());
+						}
+					});
+			}
 		}
 
 		// Render front layer data for each view
