Refactored BoneMap data into its own separate buffer from BoneHierarchy for clarity. Previously the data was implicitly being put in the BoneHierarchy buffer which is confusing.

#Synced-CL 50310951

diff --git a/Engine/Source/Runtime/Engine/Private/SkinningSceneExtensionProxy.cpp b/Engine/Source/Runtime/Engine/Private/SkinningSceneExtensionProxy.cpp
--- a/Engine/Source/Runtime/Engine/Private/SkinningSceneExtensionProxy.cpp
+++ b/Engine/Source/Runtime/Engine/Private/SkinningSceneExtensionProxy.cpp
@@ -105,27 +105,21 @@ void FSkinningSceneExtensionProxy::InitForBoneMapStorage()
 	check(BoneTransformStorageMode == EBoneTransformStorageMode::BoneMap);
 	TConstArrayView<FSkelMeshRenderSection> Sections = MeshObject->GetRenderSections(LODIndex);
 
-	MaxBoneTransformIndex = 0;
 	MaxBoneTransformCount = 0;
 	for (const FSkelMeshRenderSection& Section : Sections)
 	{
 		if (Section.IsValid())
 		{
-			for (uint32 Index : Section.BoneMap)
-			{
-				MaxBoneTransformIndex = FMath::Max<uint32>(MaxBoneTransformIndex, Index);
-			}
-
 			MaxBoneTransformCount += Section.BoneMap.Num();
 		}
 	}
-	BoneHierarchy.Empty(MaxBoneTransformCount);
+	BoneMap.Empty(MaxBoneTransformCount);
 
 	for (const FSkelMeshRenderSection& Section : Sections)
 	{
 		if (Section.IsValid())
 		{
-			BoneHierarchy.Append(Section.BoneMap);
+			BoneMap.Append(Section.BoneMap);
 		}
 	}
 
diff --git a/Engine/Source/Runtime/Engine/Public/SkeletalRenderPublic.h b/Engine/Source/Runtime/Engine/Public/SkeletalRenderPublic.h
--- a/Engine/Source/Runtime/Engine/Public/SkeletalRenderPublic.h
+++ b/Engine/Source/Runtime/Engine/Public/SkeletalRenderPublic.h
@@ -98,13 +98,10 @@ enum class EBoneTransformStorageMode : uint8
 	// Transforms are stored directly without modification on the GPU. This may require a section bone map buffer indirection on the GPU.
 	Direct,
 
-	// Transforms are stored with a built-in bone map indirection for the specific LOD. This representation is only valid for the current LOD.
+	// Transforms are stored with a built-in bone map indirection for the specific LOD.
 	BoneMap
 };
 
-// Returns the globally configured bone transform storage mode.
-EBoneTransformStorageMode GetBoneTransformStorageMode();
-
 // If true, the skinning scene extension supports GPU skin skeletal mesh objects.
 ENGINE_API extern bool IsGPUSkinSceneExtensionEnabled();
 
diff --git a/Engine/Source/Runtime/Engine/Public/SkinningSceneExtensionProxy.h b/Engine/Source/Runtime/Engine/Public/SkinningSceneExtensionProxy.h
--- a/Engine/Source/Runtime/Engine/Public/SkinningSceneExtensionProxy.h
+++ b/Engine/Source/Runtime/Engine/Public/SkinningSceneExtensionProxy.h
@@ -31,6 +31,11 @@ public:
 		return MeshObject;
 	}
 
+	TConstArrayView<uint16> GetBoneMap() const
+	{
+		return BoneMap;
+	}
+
 	TConstArrayView<uint32> GetBoneHierarchy() const
 	{
 		return BoneHierarchy;
@@ -46,19 +51,19 @@ public:
 		return MaxBoneTransformCount;
 	}
 
-	uint32 GetMaxBoneTransformIndex() const
+	uint32 GetMaxBoneMapCount() const
 	{
-		return MaxBoneTransformIndex;
+		return BoneMap.Num();
 	}
 
 	uint32 GetMaxBoneHierarchyCount() const
 	{
-		return MaxBoneTransformCount;
+		return BoneHierarchy.Num();
 	}
 
 	uint32 GetMaxBoneObjectSpaceCount() const
 	{
-		return GetBoneTransformStorageMode() == EBoneTransformStorageMode::BoneMap ? 0 : MaxBoneTransformCount;
+		return BoneHierarchy.Num();
 	}
 
 	uint32 GetMaxBoneInfluenceCount() const
@@ -130,12 +135,12 @@ protected:
 	const USkinnedAsset* SkinnedAsset = nullptr;
 	const FSkeletalMeshObject* MeshObject = nullptr;
 
+	TArray<uint16> BoneMap;
 	TArray<uint32> BoneHierarchy;
 	TArray<float> BoneObjectSpace;
 
 	uint32 UniqueAnimationCount         = 1u;
 	uint16 MaxBoneTransformCount        = 0u;
-	uint16 MaxBoneTransformIndex        = 0u;
 	uint16 MaxBoneInfluenceCount        = 0u;
 	uint16 LODIndex                     = InvalidLODIndex;
 
diff --git a/Engine/Source/Runtime/Renderer/Private/Skinning/AnimBankTransformProvider.cpp b/Engine/Source/Runtime/Renderer/Private/Skinning/AnimBankTransformProvider.cpp
--- a/Engine/Source/Runtime/Renderer/Private/Skinning/AnimBankTransformProvider.cpp
+++ b/Engine/Source/Runtime/Renderer/Private/Skinning/AnimBankTransformProvider.cpp
@@ -328,7 +328,7 @@ void FAnimBankTransformProvider::ProvideGPUBankTransforms(FSkinningTransformProv
 			const uint32 MaxTransformCount     = ProxyData.SceneExtensionProxy.GetMaxBoneTransformCount();
 			const uint32 FullBlockCount        = MaxTransformCount / BonesPerGroup;
 			const uint32 PartialTransformCount = MaxTransformCount - (FullBlockCount * BonesPerGroup);
-			const uint32 BoneMapOffset         = ProxyData.TransformStorageMode == EBoneTransformStorageMode::BoneMap ? Indirection.HierarchyOffset : ~uint32(0);
+			const uint32 BoneMapOffset         = ProxyData.TransformStorageMode != EBoneTransformStorageMode::Direct ? Indirection.BoneMapOffset : ~uint32(0);
 
 			const uint32 CurrentTransformIndex  = Indirection.CurrentTransformOffset  / sizeof(FCompressedBoneTransform);
 			const uint32 PreviousTransformIndex = Indirection.PreviousTransformOffset / sizeof(FCompressedBoneTransform);
@@ -438,7 +438,7 @@ void FAnimBankTransformProvider::ProvideGPUBankTransforms(FSkinningTransformProv
 	FAnimBankEvaluateCS::FParameters* PassParameters = GraphBuilder.AllocParameters<FAnimBankEvaluateCS::FParameters>();
 	PassParameters->SequenceTransformBuffer = GraphBuilder.CreateSRV(SequenceTransformBuffer);
 	PassParameters->BankBlockBuffer = GraphBuilder.CreateSRV(BankBlockBuffer);
-	PassParameters->BoneMapBuffer = Context.HierarchyBufferSRV;
+	PassParameters->BoneMapBuffer = Context.BoneMapBufferSRV;
 	PassParameters->TransformBuffer = GetCompressedBoneTransformUAV(GraphBuilder, Context.TransformBuffer);
 	PassParameters->bInterpolating = bInterpolating;
 	PassParameters->SampleInterval = 1.0f / ANIM_BANK_SAMPLE_RATE;
@@ -488,6 +488,7 @@ void FAnimBankTransformProvider::ProvideCPUBankTransforms(FSkinningTransformProv
 			TArrayView<FCompressedBoneTransform> DuplicatePreviousTransforms,
 			const FSequence& Sequence,
 			const FMapping& Mapping,
+			uint32 TransformCount,
 			float TimeOffset)
 		{
 			int32 KeyIndex0, KeyIndex1;
@@ -534,15 +535,26 @@ void FAnimBankTransformProvider::ProvideCPUBankTransforms(FSkinningTransformProv
 				return (FMatrix44f)MeshGlobalAnimPoseXform.ToMatrixNoScale();
 			};
 
-			if (Proxy->GetBoneTransformStorageMode() == EBoneTransformStorageMode::BoneMap)
+			if (Proxy->GetBoneTransformStorageMode() == EBoneTransformStorageMode::Direct)
 			{
-				const uint32 NumSourceTransforms = Proxy->GetMaxBoneTransformIndex() + 1;
-				TConstArrayView<uint32> BoneHierarchy = Proxy->GetBoneHierarchy();
-				SourceTransformMap.Init(false, NumSourceTransforms);
-				SourceTransforms.SetNumUninitialized(NumSourceTransforms, EAllowShrinking::No);
+				for (uint32 TransformIndex = 0; TransformIndex < TransformCount; ++TransformIndex)
+				{
+					StoreCompressedBoneTransform(&Transforms[TransformIndex], SampleTransform(TransformIndex));
+
+					if (!DuplicatePreviousTransforms.IsEmpty())
+					{
+						StoreCompressedBoneTransform(&DuplicatePreviousTransforms[TransformIndex], SampleTransform(TransformIndex));
+					}
+				}
+			}
+			else
+			{
+				// Mapping.BoneCount is always the number of mesh skeleton bones.
+				SourceTransformMap.Init(false, Mapping.BoneCount);
+				SourceTransforms.SetNumUninitialized(Mapping.BoneCount, EAllowShrinking::No);
 				int32 DstTransformIndex = 0;
 
-				for (uint32 SrcTransformIndex : Proxy->GetBoneHierarchy())
+				for (uint32 SrcTransformIndex : Proxy->GetBoneMap())
 				{
 					if (!SourceTransformMap[SrcTransformIndex])
 					{
@@ -560,18 +572,6 @@ void FAnimBankTransformProvider::ProvideCPUBankTransforms(FSkinningTransformProv
 					DstTransformIndex++;
 				}
 			}
-			else
-			{
-				for (uint32 TransformIndex = 0; TransformIndex < Mapping.BoneCount; ++TransformIndex)
-				{
-					StoreCompressedBoneTransform(&Transforms[TransformIndex], SampleTransform(TransformIndex));
-
-					if (!DuplicatePreviousTransforms.IsEmpty())
-					{
-						StoreCompressedBoneTransform(&DuplicatePreviousTransforms[TransformIndex], SampleTransform(TransformIndex));
-					}
-				}
-			}
 		};
 
 		const float GlobalTimeScale = CVarAnimBankTimeScale.GetValueOnRenderThread();
@@ -623,12 +623,12 @@ void FAnimBankTransformProvider::ProvideCPUBankTransforms(FSkinningTransformProv
 				if (bNeedsCurrentUpdate)
 				{
 					TArrayView<FCompressedBoneTransform> DstTransforms = ScatterUploader.Add_GetRef<FCompressedBoneTransform>(CurrentTransformIndex + TrackIndex * MaxTotalTransformCount, MaxTransformCount);
-					SampleAnimBank(Proxy, DstTransforms, UpdateMode == EPreviousBoneTransformUpdateMode::DuplicateCurrentToPrevious ? DstTransformsPrevious : TArrayView<FCompressedBoneTransform>{}, SequenceView.Sequence, SequenceView.Mapping, Position);
+					SampleAnimBank(Proxy, DstTransforms, UpdateMode == EPreviousBoneTransformUpdateMode::DuplicateCurrentToPrevious ? DstTransformsPrevious : TArrayView<FCompressedBoneTransform>{}, SequenceView.Sequence, SequenceView.Mapping, MaxTransformCount,  Position);
 				}
 
 				if (UpdateMode == EPreviousBoneTransformUpdateMode::UpdatePrevious)
 				{
-					SampleAnimBank(Proxy, DstTransformsPrevious, {}, SequenceView.Sequence, SequenceView.Mapping, PreviousPosition);
+					SampleAnimBank(Proxy, DstTransformsPrevious, {}, SequenceView.Sequence, SequenceView.Mapping, MaxTransformCount, PreviousPosition);
 				}
 
 				Data.Update(SequenceView.Sequence.PlayLength, SequenceView.Sequence.bLooping, DeltaTime);
diff --git a/Engine/Source/Runtime/Renderer/Private/Skinning/AnimRuntimeTransformProvider.cpp b/Engine/Source/Runtime/Renderer/Private/Skinning/AnimRuntimeTransformProvider.cpp
--- a/Engine/Source/Runtime/Renderer/Private/Skinning/AnimRuntimeTransformProvider.cpp
+++ b/Engine/Source/Runtime/Renderer/Private/Skinning/AnimRuntimeTransformProvider.cpp
@@ -106,7 +106,7 @@ void FAnimRuntimeTransformProvider::ProvideAnimRuntimeTransforms(FSkinningTransf
 						{
 							uint32 DstTransformIndex = 0;
 
-							for (uint32 SrcTransformIndex : Proxy->GetBoneHierarchy())
+							for (uint32 SrcTransformIndex : Proxy->GetBoneMap())
 							{
 								DstTransforms[DstTransformIndex] = SrcTransforms[SrcTransformIndex];
 								DstTransformIndex++;
@@ -126,7 +126,7 @@ void FAnimRuntimeTransformProvider::ProvideAnimRuntimeTransforms(FSkinningTransf
 						{
 							uint32 DstTransformIndex = 0;
 
-							for (uint32 SrcTransformIndex : Proxy->GetBoneHierarchy())
+							for (uint32 SrcTransformIndex : Proxy->GetBoneMap())
 							{
 								SetCompressedBoneTransformIdentity(DstTransforms[DstTransformIndex]);
 								DstTransformIndex++;
diff --git a/Engine/Source/Runtime/Renderer/Private/Skinning/SkinnedMeshDebugView.cpp b/Engine/Source/Runtime/Renderer/Private/Skinning/SkinnedMeshDebugView.cpp
--- a/Engine/Source/Runtime/Renderer/Private/Skinning/SkinnedMeshDebugView.cpp
+++ b/Engine/Source/Runtime/Renderer/Private/Skinning/SkinnedMeshDebugView.cpp
@@ -161,14 +161,14 @@ void FSkinnedMeshDebugViewExtension::RenderSkeletons(FRDGBuilder& GraphBuilder,
 
 		FSkinningSceneExtensionProxy* Proxy = SceneInfo->Proxy->GetSkinningSceneExtensionProxy();
 
-		if (!Proxy || Proxy->GetBoneTransformStorageMode() == EBoneTransformStorageMode::BoneMap)
+		if (!Proxy || Proxy->GetMaxBoneHierarchyCount() == 0)
 		{
 			continue;
 		}
 
 		FSkinnedMeshPrimitive& Primitive = Primitives.AddDefaulted_GetRef();
 		Primitive.Index = SceneInfo->GetPersistentIndex();
-		Primitive.BoneCount = Proxy->GetMaxBoneTransformCount();
+		Primitive.BoneCount = Proxy->GetMaxBoneHierarchyCount();
 		Primitive.InstanceCount = SceneInfo->GetNumInstanceSceneDataEntries();
 	}
 
diff --git a/Engine/Source/Runtime/Renderer/Private/Skinning/SkinningSceneExtension.cpp b/Engine/Source/Runtime/Renderer/Private/Skinning/SkinningSceneExtension.cpp
--- a/Engine/Source/Runtime/Renderer/Private/Skinning/SkinningSceneExtension.cpp
+++ b/Engine/Source/Runtime/Renderer/Private/Skinning/SkinningSceneExtension.cpp
@@ -30,6 +30,7 @@ DECLARE_DWORD_COUNTER_STAT(TEXT("Num Updates"), STAT_SkinningSceneExtension_NumU
 
 DECLARE_SKINNING_MEMORY_STAT(TEXT("Header Buffer Size"), STAT_SkinningSceneExtension_HeaderBufferSize, STATGROUP_SkinningSceneExtension);
 DECLARE_SKINNING_MEMORY_STAT(TEXT("Bone Transform Buffer Size"), STAT_SkinningSceneExtension_BoneTransformBufferSize, STATGROUP_SkinningSceneExtension);
+DECLARE_SKINNING_MEMORY_STAT(TEXT("Bone Map Buffer Size"), STAT_SkinningSceneExtension_BoneMapBufferSize, STATGROUP_SkinningSceneExtension);
 DECLARE_SKINNING_MEMORY_STAT(TEXT("Bone Hierarchy Buffer Size"), STAT_SkinningSceneExtension_BoneHierarchyBufferSize, STATGROUP_SkinningSceneExtension);
 DECLARE_SKINNING_MEMORY_STAT(TEXT("Bone Object Space Buffer Size"), STAT_SkinningSceneExtension_BoneObjectSpaceBufferSize, STATGROUP_SkinningSceneExtension);
 DECLARE_SKINNING_MEMORY_STAT(TEXT("Bone Transform Upload Size"), STAT_SkinningSceneExtension_BoneTransformUploadSize, STATGROUP_SkinningSceneExtension);
@@ -113,6 +114,7 @@ static FAutoConsoleVariableRef CVarSkinningDefaultAnimationMinScreenSize(
 
 BEGIN_UNIFORM_BUFFER_STRUCT(FSkinningSceneParameters, RENDERER_API)
 	SHADER_PARAMETER_RDG_BUFFER_SRV(ByteAddressBuffer, Headers)
+	SHADER_PARAMETER_RDG_BUFFER_SRV(ByteAddressBuffer, BoneMap)
 	SHADER_PARAMETER_RDG_BUFFER_SRV(ByteAddressBuffer, BoneHierarchy)
 	SHADER_PARAMETER_RDG_BUFFER_SRV(ByteAddressBuffer, BoneObjectSpace)
 	SHADER_PARAMETER_RDG_COMPRESSED_BONE_TRANSFORM_SRV(BoneTransforms)
@@ -167,6 +169,7 @@ static void GetDefaultSkinningParameters(FSkinningSceneParameters& OutParameters
 {
 	auto DefaultBuffer = GraphBuilder.CreateSRV(GSystemTextures.GetDefaultByteAddressBuffer(GraphBuilder, 4u));
 	OutParameters.Headers			= DefaultBuffer;
+	OutParameters.BoneMap			= DefaultBuffer;
 	OutParameters.BoneHierarchy		= DefaultBuffer;
 	OutParameters.BoneObjectSpace	= DefaultBuffer;
 	OutParameters.BoneTransforms	= GetCompressedBoneTransformSRV(GraphBuilder, (GetDefaultCompressedBoneTransformBuffer(GraphBuilder)));
@@ -240,7 +243,9 @@ void FSkinningSceneExtension::SetEnabled(bool bEnabled)
 		else
 		{
 			Buffers = nullptr;
-			HierarchyAllocator.Reset();
+			BoneMapAllocator.Reset();
+			BoneHierarchyAllocator.Reset();
+			ObjectSpaceAllocator.Reset();
 			TransformAllocator.Reset();
 			HeaderDatas.Reset();
 			BatchHeaderDatas.Reset();
@@ -258,6 +263,7 @@ void FSkinningSceneExtension::FinishSkinningBufferUpload(FRDGBuilder& GraphBuild
 	}
 
 	FRDGBufferRef HeaderBuffer = nullptr;
+	FRDGBufferRef BoneMapBuffer = nullptr;
 	FRDGBufferRef BoneHierarchyBuffer = nullptr;
 	FRDGBufferRef BoneObjectSpaceBuffer = nullptr;
 	FRDGBufferRef TransformBuffer = nullptr;
@@ -275,7 +281,8 @@ void FSkinningSceneExtension::FinishSkinningBufferUpload(FRDGBuilder& GraphBuild
 
 	const uint32 MinHeaderDataSize = (HeaderDatas.GetMaxIndex() + 1);
 	const uint32 MinTransformDataSize = TransformAllocator.GetMaxSize();
-	const uint32 MinHierarchyDataSize = HierarchyAllocator.GetMaxSize();
+	const uint32 MinBoneMapDataSize = BoneMapAllocator.GetMaxSize();
+	const uint32 MinBoneHierarchyDataSize = BoneHierarchyAllocator.GetMaxSize();
 	const uint32 MinObjectSpaceDataSize = ObjectSpaceAllocator.GetMaxSize();
 
 	RDG_GPU_MASK_SCOPE(GraphBuilder, FRHIGPUMask::All());
@@ -288,10 +295,16 @@ void FSkinningSceneExtension::FinishSkinningBufferUpload(FRDGBuilder& GraphBuild
 			MinHeaderDataSize
 		);
 
+		BoneMapBuffer = Uploader->BoneMapUploader.ResizeAndUploadTo(
+			GraphBuilder,
+			Buffers->BoneMapBuffer,
+			MinBoneMapDataSize
+		);
+
 		BoneHierarchyBuffer = Uploader->BoneHierarchyUploader.ResizeAndUploadTo(
 			GraphBuilder,
 			Buffers->BoneHierarchyBuffer,
-			MinHierarchyDataSize
+			MinBoneHierarchyDataSize
 		);
 
 		BoneObjectSpaceBuffer = Uploader->BoneObjectSpaceUploader.ResizeAndUploadTo(
@@ -311,7 +324,8 @@ void FSkinningSceneExtension::FinishSkinningBufferUpload(FRDGBuilder& GraphBuild
 	else
 	{
 		HeaderBuffer			= Buffers->HeaderDataBuffer.ResizeBufferIfNeeded(GraphBuilder, MinHeaderDataSize);
-		BoneHierarchyBuffer		= Buffers->BoneHierarchyBuffer.ResizeBufferIfNeeded(GraphBuilder, MinHierarchyDataSize);
+		BoneMapBuffer			= Buffers->BoneMapBuffer.ResizeBufferIfNeeded(GraphBuilder, MinBoneMapDataSize);
+		BoneHierarchyBuffer		= Buffers->BoneHierarchyBuffer.ResizeBufferIfNeeded(GraphBuilder, MinBoneHierarchyDataSize);
 		BoneObjectSpaceBuffer	= Buffers->BoneObjectSpaceBuffer.ResizeBufferIfNeeded(GraphBuilder, MinObjectSpaceDataSize);
 		TransformBuffer			= Buffers->TransformDataBuffer.ResizeBufferIfNeeded(GraphBuilder, MinTransformDataSize);
 	}
@@ -321,6 +335,7 @@ void FSkinningSceneExtension::FinishSkinningBufferUpload(FRDGBuilder& GraphBuild
 		INC_DWORD_STAT_BY(STAT_SkinningSceneExtension_NumPrimitives, HeaderDatas.Num());
 		INC_MEMORY_STAT_BY(STAT_SkinningSceneExtension_HeaderBufferSize, HeaderBuffer->Desc.GetSize());
 		INC_MEMORY_STAT_BY(STAT_SkinningSceneExtension_BoneTransformBufferSize, TransformBuffer->Desc.GetSize());
+		INC_MEMORY_STAT_BY(STAT_SkinningSceneExtension_BoneMapBufferSize, BoneMapBuffer->Desc.GetSize());
 		INC_MEMORY_STAT_BY(STAT_SkinningSceneExtension_BoneHierarchyBufferSize, BoneHierarchyBuffer->Desc.GetSize());
 		INC_MEMORY_STAT_BY(STAT_SkinningSceneExtension_BoneObjectSpaceBufferSize, BoneObjectSpaceBuffer->Desc.GetSize());
 		INC_MEMORY_STAT_BY(STAT_SkinningSceneExtension_BoneTransformAllocatorSize, MinTransformDataSize * sizeof(FCompressedBoneTransform));
@@ -336,6 +351,7 @@ void FSkinningSceneExtension::FinishSkinningBufferUpload(FRDGBuilder& GraphBuild
 	if (OutParams != nullptr)
 	{
 		OutParams->Headers			= GraphBuilder.CreateSRV(HeaderBuffer);
+		OutParams->BoneMap			= GraphBuilder.CreateSRV(BoneMapBuffer);
 		OutParams->BoneHierarchy	= GraphBuilder.CreateSRV(BoneHierarchyBuffer);
 		OutParams->BoneObjectSpace	= GraphBuilder.CreateSRV(BoneObjectSpaceBuffer);
 		OutParams->BoneTransforms	= GetCompressedBoneTransformSRV(GraphBuilder, TransformBuffer);
@@ -422,7 +438,7 @@ void FSkinningSceneExtension::PerformSkinning(FSkinningSceneParameters& Paramete
 		{
 			uint32 CurrentTransformOffset;
 			uint32 PreviousTransformOffset;
-			uint32 HierarchyOffset;
+			uint32 BoneMapOffset;
 			EDirtyBoneTransforms DirtyBoneTransforms;
 		};
 
@@ -463,7 +479,7 @@ void FSkinningSceneExtension::PerformSkinning(FSkinningSceneParameters& Paramete
 
 				Offsets[TotalBatches].CurrentTransformOffset  = Header.TransformBufferOffset + ( Header.CurrentTransformSlot * Header.MaxTransformCount);
 				Offsets[TotalBatches].PreviousTransformOffset = Header.TransformBufferOffset + (!Header.CurrentTransformSlot * Header.MaxTransformCount);
-				Offsets[TotalBatches].HierarchyOffset = Header.HierarchyBufferOffset;
+				Offsets[TotalBatches].BoneMapOffset = Header.BoneMapBufferOffset;
 				Offsets[TotalBatches].DirtyBoneTransforms = EDirtyBoneTransforms::All;
 
 				++TotalBatches;
@@ -493,7 +509,7 @@ void FSkinningSceneExtension::PerformSkinning(FSkinningSceneParameters& Paramete
 							TotalBatchIndices,
 							Offsets[TotalBatchIndices].CurrentTransformOffset  * sizeof(FCompressedBoneTransform),
 							Offsets[TotalBatchIndices].PreviousTransformOffset * sizeof(FCompressedBoneTransform),
-							Offsets[TotalBatchIndices].HierarchyOffset * sizeof(uint32),
+							Offsets[TotalBatchIndices].BoneMapOffset * sizeof(uint32),
 							Offsets[TotalBatchIndices].DirtyBoneTransforms
 						);
 						++Range.Count;
@@ -518,7 +534,7 @@ void FSkinningSceneExtension::PerformSkinning(FSkinningSceneParameters& Paramete
 				GraphBuilder,
 				ScatterUploadBuilder,
 				Parameters.BoneTransforms->GetParent(),
-				Parameters.BoneHierarchy
+				Parameters.BoneMap
 			);
 
 			TransformProvider->Broadcast(SkeletonRanges, Context);
@@ -571,7 +587,7 @@ void FSkinningSceneExtension::PerformSkinning(FSkinningSceneParameters& Paramete
 				Proxies[PrimitiveCount] = Header.Proxy;
 				Offsets[PrimitiveCount].CurrentTransformOffset  = Header.TransformBufferOffset + ( Header.CurrentTransformSlot * Header.MaxTransformCount);
 				Offsets[PrimitiveCount].PreviousTransformOffset = Header.TransformBufferOffset + (!Header.CurrentTransformSlot * Header.MaxTransformCount);
-				Offsets[PrimitiveCount].HierarchyOffset = Header.HierarchyBufferOffset;
+				Offsets[PrimitiveCount].BoneMapOffset = Header.BoneMapBufferOffset;
 				Offsets[PrimitiveCount].DirtyBoneTransforms = Header.DirtyBoneTransforms;
 
 				Header.DirtyBoneTransforms = EDirtyBoneTransforms::None;
@@ -596,7 +612,7 @@ void FSkinningSceneExtension::PerformSkinning(FSkinningSceneParameters& Paramete
 					PrimitiveIndex,
 					Offsets[PrimitiveIndex].CurrentTransformOffset * sizeof(FCompressedBoneTransform),
 					Offsets[PrimitiveIndex].PreviousTransformOffset * sizeof(FCompressedBoneTransform),
-					Offsets[PrimitiveIndex].HierarchyOffset * sizeof(uint32),
+					Offsets[PrimitiveIndex].BoneMapOffset * sizeof(uint32),
 					Offsets[PrimitiveIndex].DirtyBoneTransforms
 				);
 				++Range.Count;
@@ -611,7 +627,7 @@ void FSkinningSceneExtension::PerformSkinning(FSkinningSceneParameters& Paramete
 				GraphBuilder,
 				ScatterUploadBuilder,
 				Parameters.BoneTransforms->GetParent(),
-				Parameters.BoneHierarchy
+				Parameters.BoneMap
 			);
 
 			TransformProvider->Broadcast(PrimitiveRanges, Context);
@@ -623,7 +639,8 @@ bool FSkinningSceneExtension::ProcessBufferDefragmentation()
 {
 	// Consolidate spans
 	ObjectSpaceAllocator.Consolidate();
-	HierarchyAllocator.Consolidate();
+	BoneMapAllocator.Consolidate();
+	BoneHierarchyAllocator.Consolidate();
 	TransformAllocator.Consolidate();
 
 	// Decide to defragment the buffer when the used size dips below a certain multiple of the max used size.
@@ -661,7 +678,8 @@ bool FSkinningSceneExtension::ProcessBufferDefragmentation()
 	}
 
 	ObjectSpaceAllocator.Reset();
-	HierarchyAllocator.Reset();
+	BoneMapAllocator.Reset();
+	BoneHierarchyAllocator.Reset();
 	TransformAllocator.Reset();
 	BatchHeaderDatas.Reset();
 	AllocatedHeaderDataIndices.Reset();
@@ -674,6 +692,12 @@ bool FSkinningSceneExtension::ProcessBufferDefragmentation()
 			Data.TransformBufferCount = 0;
 		}
 
+		if (Data.BoneMapBufferOffset != INDEX_NONE)
+		{
+			Data.BoneMapBufferOffset = INDEX_NONE;
+			Data.BoneMapBufferCount = 0;
+		}
+
 		if (Data.HierarchyBufferOffset != INDEX_NONE)
 		{
 			Data.HierarchyBufferOffset = INDEX_NONE;
@@ -738,6 +762,7 @@ void FSkinningSceneExtension::WaitForHeaderDataUpdateTasks() const
 
 FSkinningSceneExtension::FBuffers::FBuffers()
 : HeaderDataBuffer(GSkinningBuffersHeaderDataMinSizeBytes >> 2u, TEXT("Skinning.HeaderData"))
+, BoneMapBuffer(GSkinningBuffersTransformDataMinSizeBytes >> 2u, TEXT("Skinning.BoneMap"))
 , BoneHierarchyBuffer(GSkinningBuffersTransformDataMinSizeBytes >> 2u, TEXT("Skinning.BoneHierarchy"))
 , BoneObjectSpaceBuffer(GSkinningBuffersTransformDataMinSizeBytes >> 2u, TEXT("Skinning.BoneObjectSpace"))
 , TransformDataBuffer(GSkinningBuffersTransformDataMinSizeBytes >> 2u, TEXT("Skinning.BoneTransforms"))
@@ -793,9 +818,14 @@ void FSkinningSceneExtension::FUpdater::PreSceneUpdate(FRDGBuilder& GraphBuilder
 							SceneData->ObjectSpaceAllocator.Free(Data.ObjectSpaceBufferOffset, Data.ObjectSpaceBufferCount);
 						}
 
+						if (Data.BoneMapBufferOffset != INDEX_NONE)
+						{
+							SceneData->BoneMapAllocator.Free(Data.BoneMapBufferOffset, Data.BoneMapBufferCount);
+						}
+
 						if (Data.HierarchyBufferOffset != INDEX_NONE)
 						{
-							SceneData->HierarchyAllocator.Free(Data.HierarchyBufferOffset, Data.HierarchyBufferCount);
+							SceneData->BoneHierarchyAllocator.Free(Data.HierarchyBufferOffset, Data.HierarchyBufferCount);
 						}
 
 						if (Data.TransformBufferOffset != INDEX_NONE)
@@ -882,7 +912,8 @@ void FSkinningSceneExtension::FUpdater::PostSceneUpdate(FRDGBuilder& GraphBuilde
 					NewHeader.Proxy                       = Proxy;
 					NewHeader.MeshObject                  = MeshObject;
 					NewHeader.MaxTransformCount           = Proxy->GetMaxBoneTransformCount();
-					NewHeader.MaxHierarchyCount           = Proxy->GetMaxBoneHierarchyCount();
+					NewHeader.MaxBoneMapCount             = Proxy->GetMaxBoneMapCount();
+					NewHeader.MaxBoneHierarchyCount       = Proxy->GetMaxBoneHierarchyCount();
 					NewHeader.MaxObjectSpaceCount         = Proxy->GetMaxBoneObjectSpaceCount();
 					NewHeader.MaxInfluenceCount           = Proxy->GetMaxBoneInfluenceCount();
 					NewHeader.UniqueAnimationCount        = Proxy->GetUniqueAnimationCount();
@@ -966,16 +997,17 @@ void FSkinningSceneExtension::FUpdater::PostMeshUpdate(FRDGBuilder& GraphBuilder
 					HeaderData.CurrentTransformSlot = !HeaderData.CurrentTransformSlot;
 				}
 
-				const bool bLODChanged                  = LODIndex != INDEX_NONE ? HeaderData.Proxy->SetLOD(LODIndex) : false;
+				const bool bUploadRequested             = LODIndex != INDEX_NONE ? HeaderData.Proxy->SetLOD(LODIndex) : false;
 				const bool bUniqueAnimationCountChanged = HeaderData.UniqueAnimationCount != HeaderData.Proxy->GetUniqueAnimationCount();
 
-				if (bLODChanged || bUniqueAnimationCountChanged)
+				if (bUploadRequested || bUniqueAnimationCountChanged)
 				{
-					NumAllocations += bLODChanged ? 1 : 0;
-					HeaderData.UniqueAnimationCount = HeaderData.Proxy->GetUniqueAnimationCount();
-					HeaderData.MaxTransformCount    = HeaderData.Proxy->GetMaxBoneTransformCount();
-					HeaderData.MaxHierarchyCount    = HeaderData.Proxy->GetMaxBoneHierarchyCount();
-					HeaderData.MaxObjectSpaceCount  = HeaderData.Proxy->GetMaxBoneObjectSpaceCount();
+					NumAllocations += bUploadRequested ? 1 : 0;
+					HeaderData.UniqueAnimationCount  = HeaderData.Proxy->GetUniqueAnimationCount();
+					HeaderData.MaxTransformCount     = HeaderData.Proxy->GetMaxBoneTransformCount();
+					HeaderData.MaxBoneMapCount       = HeaderData.Proxy->GetMaxBoneMapCount();
+					HeaderData.MaxBoneHierarchyCount = HeaderData.Proxy->GetMaxBoneHierarchyCount();
+					HeaderData.MaxObjectSpaceCount   = HeaderData.Proxy->GetMaxBoneObjectSpaceCount();
 
 					AllocationUpdateList.Emplace(PrimitiveIndex);
 				}
@@ -1086,6 +1118,8 @@ void FSkinningSceneExtension::FUpdater::PostMeshUpdate(FRDGBuilder& GraphBuilder
 
 				Data.ObjectSpaceBufferOffset	= SrcHeaderData.ObjectSpaceBufferOffset;
 				Data.ObjectSpaceBufferCount		= SrcHeaderData.ObjectSpaceBufferCount;
+				Data.BoneMapBufferOffset		= SrcHeaderData.BoneMapBufferOffset;
+				Data.BoneMapBufferCount			= SrcHeaderData.BoneMapBufferCount;
 				Data.HierarchyBufferOffset		= SrcHeaderData.HierarchyBufferOffset;
 				Data.HierarchyBufferCount		= SrcHeaderData.HierarchyBufferCount;
 				Data.TransformBufferOffset		= SrcHeaderData.TransformBufferOffset;
@@ -1103,9 +1137,14 @@ void FSkinningSceneExtension::FUpdater::PostMeshUpdate(FRDGBuilder& GraphBuilder
 						Data.ObjectSpaceBufferCount, (Data.MaxObjectSpaceCount * ObjectSpaceFloatCount),
 						*SkinnedAssetName, *SkeletonName);
 
-					checkf(Data.HierarchyBufferCount == Data.MaxHierarchyCount,
-						TEXT("Mismatch between HierarchyBufferCount=%d and MaxHierarchyCount=%d for mesh %s with skeleton %s."),
-						Data.HierarchyBufferCount, Data.MaxHierarchyCount,
+					checkf(Data.BoneMapBufferCount == Data.MaxBoneMapCount,
+						TEXT("Mismatch between BoneMapBufferCount=%d and MaxBoneMapCount=%d for mesh %s with skeleton %s."),
+						Data.BoneMapBufferCount, Data.MaxBoneMapCount,
+						*SkinnedAssetName, *SkeletonName);
+
+					checkf(Data.HierarchyBufferCount == Data.MaxBoneHierarchyCount,
+						TEXT("Mismatch between HierarchyBufferCount=%d and MaxBoneHierarchyCount=%d for mesh %s with skeleton %s."),
+						Data.HierarchyBufferCount, Data.MaxBoneHierarchyCount,
 						*SkinnedAssetName, *SkeletonName);
 
 					checkf(Data.TransformBufferCount == (Data.UniqueAnimationCount * Data.MaxTransformCount * 2u),
@@ -1138,15 +1177,32 @@ void FSkinningSceneExtension::FUpdater::PostMeshUpdate(FRDGBuilder& GraphBuilder
 			}
 		}
 
-		const uint32 HierarchyNeededSize = Data.MaxHierarchyCount;
+		const uint32 BoneMapNeededSize = Data.MaxBoneMapCount;
+		if (BoneMapNeededSize != Data.BoneMapBufferCount)
+		{
+			if (Data.BoneMapBufferCount > 0)
+			{
+				SceneData->BoneMapAllocator.Free(Data.BoneMapBufferOffset, Data.BoneMapBufferCount);
+			}
+
+			Data.BoneMapBufferOffset = BoneMapNeededSize > 0 ? SceneData->BoneMapAllocator.Allocate(BoneMapNeededSize) : INDEX_NONE;
+			Data.BoneMapBufferCount = BoneMapNeededSize;
+
+			if (!bForceFullUpload)
+			{
+				bRequireUpload = true;
+			}
+		}
+
+		const uint32 HierarchyNeededSize = Data.MaxBoneHierarchyCount;
 		if (HierarchyNeededSize != Data.HierarchyBufferCount)
 		{
 			if (Data.HierarchyBufferCount > 0)
 			{
-				SceneData->HierarchyAllocator.Free(Data.HierarchyBufferOffset, Data.HierarchyBufferCount);
+				SceneData->BoneHierarchyAllocator.Free(Data.HierarchyBufferOffset, Data.HierarchyBufferCount);
 			}
 
-			Data.HierarchyBufferOffset = HierarchyNeededSize > 0 ? SceneData->HierarchyAllocator.Allocate(HierarchyNeededSize) : INDEX_NONE;
+			Data.HierarchyBufferOffset = HierarchyNeededSize > 0 ? SceneData->BoneHierarchyAllocator.Allocate(HierarchyNeededSize) : INDEX_NONE;
 			Data.HierarchyBufferCount = HierarchyNeededSize;
 
 			if (!bForceFullUpload)
@@ -1242,7 +1298,7 @@ void FSkinningSceneExtension::FUpdater::PostMeshUpdate(FRDGBuilder& GraphBuilder
 
 		// Catch when/if no transform buffer data is allocated for a primitive we're tracking.
 		// This should be indicative of a bug.
-		ensure(Data.HierarchyBufferCount != 0 && Data.TransformBufferCount != 0);
+		ensure(Data.TransformBufferCount != 0);
 
 		check(SceneData->Uploader.IsValid()); // Sanity check
 		SceneData->Uploader->HeaderDataUploader.Add(Data.Pack(), PersistentIndex);
@@ -1291,11 +1347,32 @@ void FSkinningSceneExtension::FUpdater::PostMeshUpdate(FRDGBuilder& GraphBuilder
 
 	auto UploadHierarchyData = [this](const FHeaderData& Data)
 	{
+		// Bone Map
+		if (Data.MaxBoneMapCount > 0)
+		{
+			TConstArrayView<uint16> BoneMap = Data.Proxy->GetBoneMap();
+			check(BoneMap.Num() == Data.MaxBoneMapCount);
+			check(Data.BoneMapBufferCount == Data.MaxBoneMapCount);
+			check(SceneData->Uploader.IsValid());
+
+			auto UploadData = SceneData->Uploader->BoneMapUploader.AddMultiple_GetRef(
+				Data.BoneMapBufferOffset,
+				Data.BoneMapBufferCount
+			);
+
+			uint32* DstBoneMapPtr = UploadData.GetData();
+			for (int32 BoneIndex = 0; BoneIndex < Data.MaxBoneMapCount; ++BoneIndex)
+			{
+				DstBoneMapPtr[BoneIndex] = BoneMap[BoneIndex];
+			}
+		}
+
 		// Bone Hierarchy
-		if (Data.MaxHierarchyCount > 0)
+		if (Data.MaxBoneHierarchyCount > 0)
 		{
 			TConstArrayView<uint32> BoneHierarchy = Data.Proxy->GetBoneHierarchy();
-			check(BoneHierarchy.Num() == Data.MaxHierarchyCount);
+			check(BoneHierarchy.Num() == Data.MaxBoneHierarchyCount);
+			check(Data.HierarchyBufferCount == Data.MaxBoneHierarchyCount);
 			check(SceneData->Uploader.IsValid());
 
 			auto UploadData = SceneData->Uploader->BoneHierarchyUploader.AddMultiple_GetRef(
@@ -1304,7 +1381,7 @@ void FSkinningSceneExtension::FUpdater::PostMeshUpdate(FRDGBuilder& GraphBuilder
 			);
 
 			uint32* DstBoneHierarchyPtr = UploadData.GetData();
-			for (int32 BoneIndex = 0; BoneIndex < Data.MaxHierarchyCount; ++BoneIndex)
+			for (int32 BoneIndex = 0; BoneIndex < Data.MaxBoneHierarchyCount; ++BoneIndex)
 			{
 				DstBoneHierarchyPtr[BoneIndex] = BoneHierarchy[BoneIndex];
 			}
@@ -1316,6 +1393,7 @@ void FSkinningSceneExtension::FUpdater::PostMeshUpdate(FRDGBuilder& GraphBuilder
 			TConstArrayView<float> BoneObjectSpace = Data.Proxy->GetBoneObjectSpace();
 			const uint32 FloatCount = Data.Proxy->GetObjectSpaceFloatCount();
 			check(BoneObjectSpace.Num() == Data.MaxObjectSpaceCount * FloatCount);
+			check(Data.ObjectSpaceBufferCount == Data.MaxObjectSpaceCount * FloatCount);
 
 			auto UploadData = SceneData->Uploader->BoneObjectSpaceUploader.AddMultiple_GetRef(
 				Data.ObjectSpaceBufferOffset,
@@ -1730,7 +1808,7 @@ void FSkinningSceneExtension::ProvideMeshObjectTransforms(FSkinningTransformProv
 
 					if (Proxy->GetBoneTransformStorageMode() == EBoneTransformStorageMode::BoneMap)
 					{
-						for (uint32 SrcTransformIndex : Proxy->GetBoneHierarchy())
+						for (uint32 SrcTransformIndex : Proxy->GetBoneMap())
 						{
 							StoreCompressedBoneTransform(&DstCurrentTransforms[DstTransformIndex], Transforms.CurrentTransforms[SrcTransformIndex]);
 							DstTransformIndex++;
@@ -1754,7 +1832,7 @@ void FSkinningSceneExtension::ProvideMeshObjectTransforms(FSkinningTransformProv
 
 					if (Proxy->GetBoneTransformStorageMode() == EBoneTransformStorageMode::BoneMap)
 					{
-						for (uint32 SrcTransformIndex : Proxy->GetBoneHierarchy())
+						for (uint32 SrcTransformIndex : Proxy->GetBoneMap())
 						{
 							StoreCompressedBoneTransform(&DstPreviousTransforms[DstTransformIndex], Transforms.PreviousTransforms[SrcTransformIndex]);
 							DstTransformIndex++;
diff --git a/Engine/Source/Runtime/Renderer/Private/Skinning/SkinningSceneExtension.h b/Engine/Source/Runtime/Renderer/Private/Skinning/SkinningSceneExtension.h
--- a/Engine/Source/Runtime/Renderer/Private/Skinning/SkinningSceneExtension.h
+++ b/Engine/Source/Runtime/Renderer/Private/Skinning/SkinningSceneExtension.h
@@ -111,6 +111,8 @@ private:
 		uint32 UpdateCounter                     = 0;
 		uint32 InstanceSceneDataOffset           = 0;
 		uint32 NumInstanceSceneDataEntries       = 0;
+		uint32 BoneMapBufferOffset               = INDEX_NONE;
+		uint32 BoneMapBufferCount                = 0;
 		uint32 ObjectSpaceBufferOffset           = INDEX_NONE;
 		uint32 ObjectSpaceBufferCount            = 0;
 		uint32 HierarchyBufferOffset             = INDEX_NONE;
@@ -119,7 +121,8 @@ private:
 		uint32 TransformBufferCount              = 0;
 		uint32 UniqueAnimationCount              = 1;
 		uint16 MaxTransformCount                 = 0;
-		uint16 MaxHierarchyCount                 = 0;
+		uint16 MaxBoneMapCount                   = 0;
+		uint16 MaxBoneHierarchyCount             = 0;
 		uint16 MaxObjectSpaceCount               = 0;
 		EDirtyBoneTransforms DirtyBoneTransforms = EDirtyBoneTransforms::All;
 		uint8  MaxInfluenceCount                 = 0;
@@ -153,8 +156,8 @@ private:
 
 		void Validate() const
 		{
-			check(HierarchyBufferOffset	<= SKINNING_BUFFER_HIERARCHY_OFFSET_MAX);
 			check(TransformBufferOffset	<= SKINNING_BUFFER_TRANSFORM_OFFSET_MAX);
+			check(HierarchyBufferOffset == INDEX_NONE || HierarchyBufferOffset <= SKINNING_BUFFER_HIERARCHY_OFFSET_MAX);
 			check(ObjectSpaceBufferOffset == INDEX_NONE || ObjectSpaceBufferOffset <= SKINNING_BUFFER_OBJECT_SPACE_OFFSET_MAX);
 			check(MaxInfluenceCount		<= SKINNING_BUFFER_INFLUENCE_MAX);
 		}
@@ -166,6 +169,7 @@ private:
 		FBuffers();
 
 		TPersistentByteAddressBuffer<FSkinningHeader> HeaderDataBuffer;
+		TPersistentByteAddressBuffer<uint32> BoneMapBuffer;
 		TPersistentByteAddressBuffer<uint32> BoneHierarchyBuffer;
 		TPersistentByteAddressBuffer<float> BoneObjectSpaceBuffer;
 		FPersistentCompressedBoneTransformBuffer TransformDataBuffer;
@@ -177,6 +181,7 @@ private:
 	{
 	public:
 		TByteAddressBufferScatterUploader<FSkinningHeader> HeaderDataUploader;
+		TByteAddressBufferScatterUploader<uint32> BoneMapUploader;
 		TByteAddressBufferScatterUploader<uint32> BoneHierarchyUploader;
 		TByteAddressBufferScatterUploader<float> BoneObjectSpaceUploader;
 		FCompressedBoneTransformBufferScatterUploader TransformDataUploader;
@@ -213,8 +218,9 @@ private:
 private:
 	bool GetHeaderFromMeshObject(const FSkeletalMeshObject* MeshObject, FHeaderData*& OutHeaderData, int32& OutPrimitiveIndex);
 
+	FSpanAllocator BoneMapAllocator;
+	FSpanAllocator BoneHierarchyAllocator;
 	FSpanAllocator ObjectSpaceAllocator;
-	FSpanAllocator HierarchyAllocator;
 	FSpanAllocator TransformAllocator;
 	TSparseArray<FHeaderData> HeaderDatas;
 	TSet<int32> AllocatedHeaderDataIndices;
diff --git a/Engine/Source/Runtime/Renderer/Private/Skinning/SkinningTransformProvider.h b/Engine/Source/Runtime/Renderer/Private/Skinning/SkinningTransformProvider.h
--- a/Engine/Source/Runtime/Renderer/Private/Skinning/SkinningTransformProvider.h
+++ b/Engine/Source/Runtime/Renderer/Private/Skinning/SkinningTransformProvider.h
@@ -76,18 +76,18 @@ public:
 
 	struct FProviderIndirection
 	{
-		FProviderIndirection(uint32 InIndex, uint32 InCurrentTransformOffset, uint32 InPreviousTransformOffset, uint32 InHierarchyOffset, EDirtyBoneTransforms InDirtyBoneTransforms)
+		FProviderIndirection(uint32 InIndex, uint32 InCurrentTransformOffset, uint32 InPreviousTransformOffset, uint32 InBoneMapOffset, EDirtyBoneTransforms InDirtyBoneTransforms)
 			: Index(InIndex)
 			, CurrentTransformOffset(InCurrentTransformOffset)
 			, PreviousTransformOffset(InPreviousTransformOffset)
-			, HierarchyOffset(InHierarchyOffset)
+			, BoneMapOffset(InBoneMapOffset)
 			, DirtyBoneTransforms(InDirtyBoneTransforms)
 		{}
 
 		uint32 Index = 0;
 		uint32 CurrentTransformOffset = 0;
 		uint32 PreviousTransformOffset = 0;
-		uint32 HierarchyOffset = 0;
+		uint32 BoneMapOffset = 0;
 		EDirtyBoneTransforms DirtyBoneTransforms = EDirtyBoneTransforms::None;
 	};
 
@@ -102,7 +102,7 @@ public:
 			FRDGBuilder& InGraphBuilder,
 			FRDGScatterUploadBuilder& InScatterUploadBuilder,
 			FRDGBufferRef InTransformBuffer,
-			FRDGBufferSRVRef InHierarchyBufferSRV
+			FRDGBufferSRVRef InBoneMapBufferSRV
 		)
 		: Primitives(InPrimitives)
 		, Proxies(InProxies)
@@ -111,7 +111,7 @@ public:
 		, GraphBuilder(InGraphBuilder)
 		, ScatterUploadBuilder(InScatterUploadBuilder)
 		, TransformBuffer(InTransformBuffer)
-		, HierarchyBufferSRV(InHierarchyBufferSRV)
+		, BoneMapBufferSRV(InBoneMapBufferSRV)
 		, ViewFamily(InViewFamily)
 		{
 		}
@@ -124,7 +124,7 @@ public:
 		FRDGBuilder& GraphBuilder;
 		FRDGScatterUploadBuilder& ScatterUploadBuilder;
 		FRDGBufferRef TransformBuffer;
-		FRDGBufferSRVRef HierarchyBufferSRV;
+		FRDGBufferSRVRef BoneMapBufferSRV;
 
 		const FSceneViewFamily& ViewFamily;
 	};
