[HWRT] Fix an issue with indirect TLAS build being enabled even though it should be disabled.

#rnx

#Synced-CL 50205011

diff --git a/Engine/Source/Runtime/Renderer/Private/RayTracing/RayTracingScene.cpp b/Engine/Source/Runtime/Renderer/Private/RayTracing/RayTracingScene.cpp
--- a/Engine/Source/Runtime/Renderer/Private/RayTracing/RayTracingScene.cpp
+++ b/Engine/Source/Runtime/Renderer/Private/RayTracing/RayTracingScene.cpp
@@ -81,6 +81,11 @@ bool IsRayTracingFeedbackEnabled(const FSceneViewFamily& ViewFamily)
 	return !HasRayTracedOverlay(ViewFamily) && CVarRayTracingSceneUseTracingFeedback.GetValueOnRenderThread() && GRHISupportsInlineRayTracing;
 }
 
+bool IsRayTracingSceneCompactInstancesEnabled()
+{
+	return CVarRayTracingSceneCompactInstances.GetValueOnRenderThread();
+}
+
 BEGIN_SHADER_PARAMETER_STRUCT(FBuildInstanceBufferPassParams, )
 	SHADER_PARAMETER_RDG_BUFFER_UAV(RWStructuredBuffer, InstanceBuffer)
 	SHADER_PARAMETER_RDG_BUFFER_UAV(RWStructuredBuffer, HitGroupContributionsBuffer)
@@ -247,7 +252,7 @@ void FRayTracingScene::Update(FRDGBuilder& GraphBuilder, FSceneUniformBuffer& Sc
 		AddClearUAVPass(GraphBuilder, InstanceStatsBufferUAV, 0, ComputePassFlags);
 	}
 
-	const bool bCompactInstanceBuffer = CVarRayTracingSceneCompactInstances.GetValueOnRenderThread();
+	const bool bCompactInstanceBuffer = IsRayTracingSceneCompactInstancesEnabled();
 
 	for (uint32 LayerIndex = 0; LayerIndex < NumLayers; ++LayerIndex)
 	{
@@ -518,6 +523,7 @@ void FRayTracingScene::Build(FRDGBuilder& GraphBuilder, ERDGPassFlags ComputePas
 		[PassParams, this](FRHICommandList& RHICmdList)
 		{
 			const bool bUseBatchedBuild = CVarRayTracingSceneBatchedBuild.GetValueOnRenderThread();
+			const bool bCompactInstanceBuffer = IsRayTracingSceneCompactInstancesEnabled();
 			const uint32 MaxNumViews = ActiveViews.GetMaxIndex();
 
 			TArray<FRayTracingSceneBuildParams> BatchedBuildParams;
@@ -545,7 +551,7 @@ void FRayTracingScene::Build(FRDGBuilder& GraphBuilder, ERDGPassFlags ComputePas
 						BuildParams.HitGroupContributionsBufferOffset = 0;
 					}
 
-					if (GRHIGlobals.RayTracing.SupportsBuildTopLevelAccelerationStructureIndirect)
+					if (bCompactInstanceBuffer && GRHIGlobals.RayTracing.SupportsBuildTopLevelAccelerationStructureIndirect)
 					{
 						BuildParams.InstanceCountBuffer = InstanceStatsBuffer ? InstanceStatsBuffer->GetRHI() : nullptr;
 						BuildParams.InstanceCountBufferOffset = (LayerIndex * MaxNumViews + ViewIndex) * sizeof(FInstanceBufferStats);
