Height fog - added support for sky light capture to affect distant height fog only (not volumetric fog or local fog volume since it is too complex to manage with lumen gi integration GI and occluded skylight). Also having that kind of scattering on VFog or LFV could lead to sky light in scattering leaking indoors.

#rnx
#tests PC editor dx12

#Synced-CL 50391932

diff --git a/Engine/Shaders/Private/HeightFogCommon.ush b/Engine/Shaders/Private/HeightFogCommon.ush
--- a/Engine/Shaders/Private/HeightFogCommon.ush
+++ b/Engine/Shaders/Private/HeightFogCommon.ush
@@ -10,6 +10,7 @@
 #include "/Engine/Shared/EnvironmentComponentsFlags.h"
 #include "/Engine/Private/SkyAtmosphereCommon.ush"
 #include "ParticipatingMediaCommon.ush"
+#include "ReflectionEnvironmentShared.ush"
 
 #ifndef PERMUTATION_SUPPORT_FOG_INSCATTERING_TEXTURE
 #define PERMUTATION_SUPPORT_FOG_INSCATTERING_TEXTURE 1
@@ -78,6 +79,9 @@ struct FogStructData
 	float VolumetricFogNearFadeInDistanceInv;
 	float3 VolumetricFogAlbedo;
 	float VolumetricFogPhaseG;
+	float SkyLightCaptureAffectsHeightFogStrength;
+	float SkyLightCaptureAffectsHeightFogRoughness;
+	float SkyLightCaptureAffectsHeightFogMipCount;
 	// Spir-v does not support copying texture types so 
 	// we can't use FogStructData to store them.
 	//TextureCube FogInscatteringColorCubemap;
@@ -107,6 +111,9 @@ FogStructData GetFogStructData(uint EyeIndex)
 		Out.VolumetricFogNearFadeInDistanceInv = FogStruct.VolumetricFogNearFadeInDistanceInv;
 		Out.VolumetricFogAlbedo = FogStruct.VolumetricFogAlbedo;
 		Out.VolumetricFogPhaseG = FogStruct.VolumetricFogPhaseG;
+		Out.SkyLightCaptureAffectsHeightFogStrength  = FogStruct.SkyLightCaptureAffectsHeightFogStrength;
+		Out.SkyLightCaptureAffectsHeightFogRoughness = FogStruct.SkyLightCaptureAffectsHeightFogRoughness;
+		Out.SkyLightCaptureAffectsHeightFogMipCount  = FogStruct.SkyLightCaptureAffectsHeightFogMipCount;
 		// Intentionally commented out to avoid spir-v generation issue - copying texture types is not supported
 		//Out.FogInscatteringColorCubemap = FogStruct.FogInscatteringColorCubemap;
 		//Out.IntegratedLightScattering = FogStruct.IntegratedLightScattering;
@@ -129,6 +136,9 @@ FogStructData GetFogStructData(uint EyeIndex)
 		Out.VolumetricFogNearFadeInDistanceInv = FogStructISR.VolumetricFogNearFadeInDistanceInv;
 		Out.VolumetricFogAlbedo = FogStructISR.VolumetricFogAlbedo;
 		Out.VolumetricFogPhaseG = FogStructISR.VolumetricFogPhaseG;
+		Out.SkyLightCaptureAffectsHeightFogStrength  = FogStruct.SkyLightCaptureAffectsHeightFogStrength;
+		Out.SkyLightCaptureAffectsHeightFogRoughness = FogStruct.SkyLightCaptureAffectsHeightFogRoughness;
+		Out.SkyLightCaptureAffectsHeightFogMipCount  = FogStruct.SkyLightCaptureAffectsHeightFogMipCount;
 		// Intentionally commented out to avoid spir-v generation issue - copying texture types is not supported
 		//Out.FogInscatteringColorCubemap = FogStructISR.FogInscatteringColorCubemap;
 		//Out.IntegratedLightScattering = FogStructISR.IntegratedLightScattering;
@@ -169,6 +179,7 @@ float3 ComputeInscatteringColor(float3 CameraToReceiver, float CameraToReceiverL
 #endif
 		Inscattering *= lerp(NonDirectionalColor, DirectionalColor, FadeAlpha);
 	}
+//	else
 #endif
 
 #if PROJECT_SUPPORT_SKY_ATMOSPHERE_AFFECTS_HEIGHFOG
@@ -209,6 +220,21 @@ float PreComputeFogOriginFactor(float OriginHeight, float FogHeight, float FogFa
 	return FogDensity * pow(2.0f, CollapsedFogParameterPower);
 }
 
+#define REFLECTION_CAPTURE_ROUGHEST_MIP 1
+#define REFLECTION_CAPTURE_ROUGHNESS_MIP_SCALE 1.2
+
+////////////////////**
+/////////////////// * Compute absolute mip for a reflection capture cubemap given a roughness.
+/////////////////// */
+///////////////////half ComputeReflectionCaptureMipFromRoughness222(half Roughness, half CubemapMaxMip)
+///////////////////{
+///////////////////	// Heuristic that maps roughness to mip level
+///////////////////	// This is done in a way such that a certain mip level will always have the same roughness, regardless of how many mips are in the texture
+///////////////////	// Using more mips in the cubemap just allows sharper reflections to be supported
+///////////////////	half LevelFrom1x1 = REFLECTION_CAPTURE_ROUGHEST_MIP - REFLECTION_CAPTURE_ROUGHNESS_MIP_SCALE * log2(max(Roughness, 0.001));
+///////////////////	return CubemapMaxMip - 1 - LevelFrom1x1;
+///////////////////}
+
 // @param WorldPositionRelativeToCamera = WorldPosition - InCameraPosition. It is the vector from the camera to the world position to shade (independent from TranslatedPosition).
 half4 GetExponentialHeightFog(float3 WorldPositionRelativeToCamera, float ExcludeDistance, uint EyeIndex, ViewState InView, bool bOverrideOrigin = false, float3 OverrideRayStartRelativeToCamera=0.0f.xxx, float3 OverrideRayDir = 0.0f.xxx, float OverrideRayLength = 0.0f)
 {
@@ -354,6 +380,23 @@ half4 GetExponentialHeightFog(float3 WorldPositionRelativeToCamera, float Exclud
 	}
 #endif
 
+#if SUPPORTS_SKYLIGHTCAPTURE_AFFECTS_HEIGHTFOGINSCATTERING 
+	if (FogData.SkyLightCaptureAffectsHeightFogStrength > 0.0f && View.RenderingReflectionCaptureMask == 0)
+	{
+		float3 CubemapLookupVector = CameraToReceiver;
+		float AbsoluteSpecularMip = ComputeReflectionCaptureMipFromRoughness(FogData.SkyLightCaptureAffectsHeightFogRoughness, FogData.SkyLightCaptureAffectsHeightFogMipCount);
+		float3 SkyLightInscatteringColor = View.SkyLightColor.rgb * TextureCubeSampleLevel(FogStruct.FogInscatteringColorCubemap, SharedFogInColorScatteringSampler, CubemapLookupVector, AbsoluteSpecularMip).xyz;
+
+		// Lerp based on influence to not account for twice the energy.
+		// No applied when rendering capture otherwise it could cause infinite feedback loop and luminance/exposure explosion.
+		// It is not applied on volumetric fog or local fog volumes since:
+		//  - It is too complex to manage with lumen GI mixed with occluded skylight in a single color.
+		//  - Also having that kind of scattering on VFog or LFV could lead to sky light in scattering leaking indoors. (it can happen with height fog but it usually has a larger start distance).
+		DirectionalInscattering = lerp(DirectionalInscattering, 0.0f, FogData.SkyLightCaptureAffectsHeightFogStrength);
+		InscatteringColor = lerp(InscatteringColor, SkyLightInscatteringColor, FogData.SkyLightCaptureAffectsHeightFogStrength);
+	}
+#endif
+
 #if PROJECT_EXPFOG_MATCHES_VFOG
 	DirectionalInscattering *= FogData.VolumetricFogAlbedo;
 #endif
diff --git a/Engine/Source/Runtime/Engine/Classes/Components/ExponentialHeightFogComponent.h b/Engine/Source/Runtime/Engine/Classes/Components/ExponentialHeightFogComponent.h
--- a/Engine/Source/Runtime/Engine/Classes/Components/ExponentialHeightFogComponent.h
+++ b/Engine/Source/Runtime/Engine/Classes/Components/ExponentialHeightFogComponent.h
@@ -53,6 +53,14 @@ class UExponentialHeightFogComponent : public USceneComponent
 	UPROPERTY(EditAnywhere, BlueprintReadOnly, Category=InscatteringTexture)
 	TObjectPtr<class UTextureCube> InscatteringColorCubemap;
 
+	/** When >0 and when r.SupportsSkyLightCaptureAffectsHeightFogInScattering is true, this will make the sky light capture inject scattered light into the height fog. */
+	UPROPERTY(EditAnywhere, BlueprintReadOnly, Category = InscatteringTexture, meta = (DisplayName = "Sky Light Capture In-Scattering Strength", UIMin = "0", UIMax = "1", ClampMin = "0"))
+	float SkyLightCaptureAffectsHeightFogStrength;
+
+	/** This is the roughness at which the sky light capture cubemap will be sample to simulate scattering dispersion due to a phase function. */
+	UPROPERTY(EditAnywhere, BlueprintReadOnly, Category = InscatteringTexture, meta = (DisplayName = "Sky Light Capture In-Scattering Roughness", UIMin = "0", UIMax = "1", ClampMin = "0", ClampMax = "1"))
+	float SkyLightCaptureAffectsHeightFogRoughness;
+
 	/** Angle to rotate the InscatteringColorCubemap around the Z axis. */
 	UPROPERTY(EditAnywhere, BlueprintReadOnly, Category=InscatteringTexture, meta=(UIMin = "0", UIMax = "360"))
 	float InscatteringColorCubemapAngle;
diff --git a/Engine/Source/Runtime/Engine/Private/Components/ExponentialHeightFogComponent.cpp b/Engine/Source/Runtime/Engine/Private/Components/ExponentialHeightFogComponent.cpp
--- a/Engine/Source/Runtime/Engine/Private/Components/ExponentialHeightFogComponent.cpp
+++ b/Engine/Source/Runtime/Engine/Private/Components/ExponentialHeightFogComponent.cpp
@@ -15,6 +15,7 @@
 #include "Net/UnrealNetwork.h"
 #include "Components/BillboardComponent.h"
 #include "UObject/UE5MainStreamObjectVersion.h"
+#include "RenderUtils.h"
 
 #include UE_INLINE_GENERATED_CPP_BY_NAME(ExponentialHeightFogComponent)
 
@@ -37,6 +38,8 @@ constexpr bool UseExponentialHeightFogStateStream = true;
 	UEXPONENTIALHEIGHTFOG_VARIABLE(bool, b, RenderInMainPass) \
 	UEXPONENTIALHEIGHTFOG_VARIABLE(bool, b, VisibleInReflectionCaptures) \
 	UEXPONENTIALHEIGHTFOG_VARIABLE(bool, b, VisibleInRealTimeSkyCaptures) \
+	UEXPONENTIALHEIGHTFOG_VARIABLE(float, , SkyLightCaptureAffectsHeightFogStrength) \
+	UEXPONENTIALHEIGHTFOG_VARIABLE(float, , SkyLightCaptureAffectsHeightFogRoughness) \
 	UEXPONENTIALHEIGHTFOG_VARIABLES_WITH_SETFUNCTION
 
 #define UEXPONENTIALHEIGHTFOG_VARIABLES_WITH_SETFUNCTION \
@@ -78,6 +81,9 @@ UExponentialHeightFogComponent::UExponentialHeightFogComponent(const FObjectInit
 	FullyDirectionalInscatteringColorDistance = 100000.0f;
 	NonDirectionalInscatteringColorDistance = 1000.0f;
 
+	SkyLightCaptureAffectsHeightFogRoughness = 0.15f;
+	SkyLightCaptureAffectsHeightFogStrength = 0.0f;
+
 	FogDensity = 0.02f;
 	FogHeightFalloff = 0.2f;
 	// No influence from the second fog as default
@@ -179,6 +185,22 @@ bool UExponentialHeightFogComponent::CanEditChange(const FProperty* InProperty)
 			}
 		}
 
+		const bool bSupportsSkyLightCaptureAffectsHeightFogInScattering = SupportsSkyLightCaptureAffectsHeightFogInScattering(GMaxRHIShaderPlatform);
+		if (PropertyName == GET_MEMBER_NAME_STRING_CHECKED(UExponentialHeightFogComponent, SkyLightCaptureAffectsHeightFogStrength))
+		{
+			return bSupportsSkyLightCaptureAffectsHeightFogInScattering;
+		}
+		const bool UsesSkyLightCaptureAffectsHeightFogInScattering = bSupportsSkyLightCaptureAffectsHeightFogInScattering && SkyLightCaptureAffectsHeightFogStrength > 0.0f;
+		if (PropertyName == GET_MEMBER_NAME_STRING_CHECKED(UExponentialHeightFogComponent, SkyLightCaptureAffectsHeightFogRoughness))
+		{
+			return UsesSkyLightCaptureAffectsHeightFogInScattering;
+		}
+
+		if (PropertyName == GET_MEMBER_NAME_STRING_CHECKED(UExponentialHeightFogComponent, InscatteringColorCubemap))
+		{
+			return !UsesSkyLightCaptureAffectsHeightFogInScattering;
+		}
+
 		if (PropertyName == GET_MEMBER_NAME_STRING_CHECKED(UExponentialHeightFogComponent, DirectionalInscatteringExponent) ||
 			PropertyName == GET_MEMBER_NAME_STRING_CHECKED(UExponentialHeightFogComponent, DirectionalInscatteringStartDistance) ||
 			PropertyName == GET_MEMBER_NAME_STRING_CHECKED(UExponentialHeightFogComponent, DirectionalInscatteringLuminance) ||
@@ -192,7 +214,7 @@ bool UExponentialHeightFogComponent::CanEditChange(const FProperty* InProperty)
 			PropertyName == GET_MEMBER_NAME_STRING_CHECKED(UExponentialHeightFogComponent, InscatteringTextureTint) ||
 			PropertyName == GET_MEMBER_NAME_STRING_CHECKED(UExponentialHeightFogComponent, InscatteringColorCubemapAngle))
 		{
-			return InscatteringColorCubemap != NULL;
+			return InscatteringColorCubemap != NULL && !UsesSkyLightCaptureAffectsHeightFogInScattering;
 		}
 
 		if (PropertyName == GET_MEMBER_NAME_STRING_CHECKED(UExponentialHeightFogComponent, FogInscatteringLuminance))
diff --git a/Engine/Source/Runtime/Engine/Private/Components/SkyLightComponent.cpp b/Engine/Source/Runtime/Engine/Private/Components/SkyLightComponent.cpp
--- a/Engine/Source/Runtime/Engine/Private/Components/SkyLightComponent.cpp
+++ b/Engine/Source/Runtime/Engine/Private/Components/SkyLightComponent.cpp
@@ -90,6 +90,13 @@ FAutoConsoleVariableRef CVarSkylightIntensityMultiplier(
 	ECVF_Scalability | ECVF_RenderThreadSafe
 	);
 
+// The project setting for the sky light component to affect the height fog in scattering (disable runtime and shader code)
+static TAutoConsoleVariable<int32> CVarSupportsSkyLightCaptureAffectsHeightFogInScattering(
+	TEXT("r.SupportsSkyLightCaptureAffectsHeightFogInScattering"),
+	0,
+	TEXT("Enables the inscattering of the sky light capture into the distant height fog."),
+	ECVF_ReadOnly | ECVF_RenderThreadSafe);
+
 void OnChangeSkylightRealTimeReflectionCapture(IConsoleVariable* Var)
 {
 	if (!GIsRHIInitialized)
diff --git a/Engine/Source/Runtime/Engine/Private/ShaderCompiler/ShaderCompiler.cpp b/Engine/Source/Runtime/Engine/Private/ShaderCompiler/ShaderCompiler.cpp
--- a/Engine/Source/Runtime/Engine/Private/ShaderCompiler/ShaderCompiler.cpp
+++ b/Engine/Source/Runtime/Engine/Private/ShaderCompiler/ShaderCompiler.cpp
@@ -3456,6 +3456,8 @@ void GlobalBeginCompileShader(
 
 	SET_SHADER_DEFINE(Input.Environment, FORWARD_SHADING_FORCES_SKYLIGHT_CUBEMAPS_BLENDING, ForwardShadingForcesSkyLightCubemapBlending(ShaderPlatform) ? 1 : 0);
 
+	SET_SHADER_DEFINE(Input.Environment, SUPPORTS_SKYLIGHTCAPTURE_AFFECTS_HEIGHTFOGINSCATTERING, SupportsSkyLightCaptureAffectsHeightFogInScattering(ShaderPlatform) ? 1 : 0);
+
 	// Enables HLSL 2021
 	uint32 EnablesHLSL2021ByDefault = FDataDrivenShaderPlatformInfo::GetEnablesHLSL2021ByDefault(EShaderPlatform(Target.Platform));
 	bool bInlineRayTracing = Input.Environment.CompilerFlags.Contains(CFLAG_InlineRayTracing);
diff --git a/Engine/Source/Runtime/Engine/Public/StateStream/ExponentialHeightFogStateStream.h b/Engine/Source/Runtime/Engine/Public/StateStream/ExponentialHeightFogStateStream.h
--- a/Engine/Source/Runtime/Engine/Public/StateStream/ExponentialHeightFogStateStream.h
+++ b/Engine/Source/Runtime/Engine/Public/StateStream/ExponentialHeightFogStateStream.h
@@ -45,6 +45,10 @@ struct FExponentialHeightFogDynamicState
 	UPROPERTY()
 	TObjectPtr<UTextureCube> InscatteringColorCubemap;
 	UPROPERTY()
+	float SkyLightCaptureAffectsHeightFogRoughness = 0.0f;
+	UPROPERTY()
+	float SkyLightCaptureAffectsHeightFogStrength = 0.0f;
+	UPROPERTY()
 	float InscatteringColorCubemapAngle = 0.0f;
 	UPROPERTY()
 	FLinearColor InscatteringTextureTint = FLinearColor::White;
diff --git a/Engine/Source/Runtime/RenderCore/Private/RenderUtils.cpp b/Engine/Source/Runtime/RenderCore/Private/RenderUtils.cpp
--- a/Engine/Source/Runtime/RenderCore/Private/RenderUtils.cpp
+++ b/Engine/Source/Runtime/RenderCore/Private/RenderUtils.cpp
@@ -2568,6 +2568,12 @@ bool ForwardShadingForcesSkyLightCubemapBlending(const FStaticShaderPlatform Pla
 	return CVar.Get(Platform) > 0;
 }
 
+bool SupportsSkyLightCaptureAffectsHeightFogInScattering(const FStaticShaderPlatform Platform)
+{
+	static FShaderPlatformCachedIniValue<int32> CVar(TEXT("r.SupportsSkyLightCaptureAffectsHeightFogInScattering"));
+	return CVar.Get(Platform) > 0;
+}
+
 bool HardwareVariableRateShadingSupportedByPlatform(EShaderPlatform ShaderPlatform)
 {
 	const static auto SupportVRSCVar = IConsoleManager::Get().FindTConsoleVariableDataInt(TEXT("r.VRS.Support"));
diff --git a/Engine/Source/Runtime/RenderCore/Private/Shader.cpp b/Engine/Source/Runtime/RenderCore/Private/Shader.cpp
--- a/Engine/Source/Runtime/RenderCore/Private/Shader.cpp
+++ b/Engine/Source/Runtime/RenderCore/Private/Shader.cpp
@@ -1955,6 +1955,12 @@ void ShaderMapAppendKey(EShaderPlatform Platform, FShaderKeyGenerator& KeyGen)
 		KeyGen.Append(TEXT("FwdSkyBlnd"));
 	}
 
+	if (SupportsSkyLightCaptureAffectsHeightFogInScattering(Platform))
+	{
+		KeyGen.AppendSeparator();
+		KeyGen.Append(TEXT("SkyCapFog"));
+	}
+
 	if (IsMobilePlatform(Platform))
 	{
 		{
@@ -2348,7 +2354,7 @@ void ShaderMapAppendKey(EShaderPlatform Platform, FShaderKeyGenerator& KeyGen)
 			KeyGen.AppendSeparator();
 			KeyGen.Append(TEXT("SKYATM"));
 
-			static const auto CVarHeightFog = IConsoleManager::Get().FindTConsoleVariableDataInt(TEXT("r.SupportSkyAtmosphereAffectsHeightFog"));
+			static const auto CVarHeightFog = IConsoleManager::Get().FindTConsoleVariableDataInt(TEXT("r.SupportsSkyAtmosphereAffectsHeightFog"));
 			if (CVarHeightFog && CVarHeightFog->GetValueOnAnyThread() > 0)
 			{
 				KeyGen.AppendSeparator();
diff --git a/Engine/Source/Runtime/RenderCore/Public/RenderUtils.h b/Engine/Source/Runtime/RenderCore/Public/RenderUtils.h
--- a/Engine/Source/Runtime/RenderCore/Public/RenderUtils.h
+++ b/Engine/Source/Runtime/RenderCore/Public/RenderUtils.h
@@ -346,6 +346,9 @@ RENDERCORE_API bool IsForwardShadingEnabled(const FStaticShaderPlatform Platform
 /** Return true if all forward shaded material should blend the interpolated sky boxes for higher quality. */
 RENDERCORE_API bool ForwardShadingForcesSkyLightCubemapBlending(const FStaticShaderPlatform Platform);
 
+/** Return true the sky light capture can affect the height fog. */
+RENDERCORE_API bool SupportsSkyLightCaptureAffectsHeightFogInScattering(const FStaticShaderPlatform Platform);
+
 /** Returns if the GBuffer is used. Only valid for the current platform. */
 RENDERCORE_API bool IsUsingGBuffers(const FStaticShaderPlatform Platform);
 
diff --git a/Engine/Source/Runtime/Renderer/Private/FogRendering.cpp b/Engine/Source/Runtime/Renderer/Private/FogRendering.cpp
--- a/Engine/Source/Runtime/Renderer/Private/FogRendering.cpp
+++ b/Engine/Source/Runtime/Renderer/Private/FogRendering.cpp
@@ -13,6 +13,7 @@
 #include "TextureResource.h"
 #include "PostProcess/PostProcessing.h" // IsPostProcessingWithAlphaChannelSupported
 #include "EnvironmentComponentsFlags.h"
+#include "SceneProxies/SkyLightSceneProxy.h"
 
 DECLARE_GPU_DRAWCALL_STAT(Fog);
 
@@ -82,6 +83,41 @@ static void SetupDummyExponentialHeightFogUniformParameters(FFogUniformParameter
 	OutParameters.FogInscatteringColorSampler = TStaticSamplerState<SF_Trilinear, AM_Clamp, AM_Clamp, AM_Clamp>::GetRHI();
 }
 
+// Returns the sky light texture to use in when sky light is present, height fog is present and all that data is there ready.
+static FTextureRHIRef IsSkyLightCaptureAffectsHeightFogTextureEnabled(const FViewInfo& View)
+{
+	FTextureRHIRef SkyLightTextureResource = nullptr;
+
+	const FScene* Scene = View.Family->Scene ? View.Family->Scene->GetRenderScene() : nullptr;
+	if (Scene && Scene->HasAnyExponentialHeightFog())
+	{
+		const FExponentialHeightFogSceneInfo& HFog = Scene->ExponentialFogs[0];
+
+		const bool bApplySkyLight = View.Family->EngineShowFlags.SkyLighting;
+		if (bApplySkyLight && Scene && Scene->SkyLight 
+			&& SupportsSkyLightCaptureAffectsHeightFogInScattering(View.GetShaderPlatform())
+			&& HFog.SkyLightCaptureAffectsHeightFogStrength > 0.0f )
+		{
+			if ((Scene->SkyLight->ProcessedTexture || Scene->CanSampleSkyLightRealTimeCaptureData())
+				&& bApplySkyLight)
+			{
+				const FSkyLightSceneProxy& SkyLight = *Scene->SkyLight;
+
+				if (Scene->CanSampleSkyLightRealTimeCaptureData())
+				{
+					// Cannot blend with this capture mode as of today.
+					SkyLightTextureResource = Scene->ConvolvedSkyRenderTarget[Scene->ConvolvedSkyRenderTargetReadyIndex]->GetRHI();
+				}
+				else if (SkyLight.ProcessedTexture)
+				{
+					SkyLightTextureResource = SkyLight.ProcessedTexture->TextureRHI;
+				}
+			}
+		}
+	}
+	return SkyLightTextureResource;
+}
+
 void SetupFogUniformParameters(FRDGBuilder& GraphBuilder, const FViewInfo& View, FFogUniformParameters& OutParameters, bool bForRealtimeSkyCapture)
 {
 	// Exponential Height Fog
@@ -93,24 +129,46 @@ void SetupFogUniformParameters(FRDGBuilder& GraphBuilder, const FViewInfo& View,
 	}
 	else
 	{
-		const FTexture* Cubemap = GWhiteTextureCube;
+		OutParameters.SkyLightCaptureAffectsHeightFogStrength = 0.0f;
+		OutParameters.SkyLightCaptureAffectsHeightFogRoughness = 0.0f;
+		OutParameters.SkyLightCaptureAffectsHeightFogMipCount = 0.0f;
 
-		if (View.FogInscatteringColorCubemap)
+		FTextureRHIRef SkyLightTextureResource = IsSkyLightCaptureAffectsHeightFogTextureEnabled(View);
+		if (Scene && SkyLightTextureResource != nullptr)
 		{
-			Cubemap = View.FogInscatteringColorCubemap->GetResource();
+			const FExponentialHeightFogSceneInfo& HFog = Scene->ExponentialFogs[0];
+
+			OutParameters.SkyLightCaptureAffectsHeightFogStrength = HFog.SkyLightCaptureAffectsHeightFogStrength;
+			OutParameters.SkyLightCaptureAffectsHeightFogRoughness = HFog.SkyLightCaptureAffectsHeightFogRoughness;
+
+			const int32 CubemapWidth = SkyLightTextureResource->GetDesc().Extent.X;
+			const float SkyMipCount = FMath::Log2(static_cast<float>(CubemapWidth)) + 1.0f;
+			OutParameters.SkyLightCaptureAffectsHeightFogMipCount = SkyMipCount;
+		}
+		else if (View.FogInscatteringColorCubemap)
+		{
+			SkyLightTextureResource = View.FogInscatteringColorCubemap->GetResource()->GetTextureRHI();
+		}
+		else
+		{
+			SkyLightTextureResource = GWhiteTextureCube->GetTextureRHI();
 		}
 
 		OutParameters.ExponentialFogParameters = View.ExponentialFogParameters;
 		OutParameters.ExponentialFogColorParameter = FVector4f(View.ExponentialFogColor, 1.0f - View.FogMaxOpacity);
 		OutParameters.ExponentialFogParameters2 = View.ExponentialFogParameters2;
 		OutParameters.ExponentialFogParameters3 = View.ExponentialFogParameters3;
+		if (OutParameters.SkyLightCaptureAffectsHeightFogStrength > 0.0f)
+		{
+			OutParameters.ExponentialFogParameters3.Z = 0.0f;	// Disables in scattering texture when it is re-used for sky light in scattering
+		}
 		OutParameters.SkyAtmosphereAmbientContributionColorScale = View.SkyAtmosphereAmbientContributionColorScale;
 		OutParameters.SinCosInscatteringColorCubemapRotation = View.SinCosInscatteringColorCubemapRotation;
 		OutParameters.FogInscatteringTextureParameters = (FVector3f)View.FogInscatteringTextureParameters;
 		OutParameters.InscatteringLightDirection = (FVector3f)View.InscatteringLightDirection;
 		OutParameters.InscatteringLightDirection.W = View.bUseDirectionalInscattering ? FMath::Max(0.f, View.DirectionalInscatteringStartDistance) : -1.f;
 		OutParameters.DirectionalInscatteringColor = FVector4f(FVector3f(View.DirectionalInscatteringColor), FMath::Clamp(View.DirectionalInscatteringExponent, 0.000001f, 1000.0f));
-		OutParameters.FogInscatteringColorCubemap = Cubemap->TextureRHI;
+		OutParameters.FogInscatteringColorCubemap = SkyLightTextureResource;
 		OutParameters.FogInscatteringColorSampler = TStaticSamplerState<SF_Trilinear, AM_Clamp, AM_Clamp, AM_Clamp>::GetRHI();
 		OutParameters.EndDistance = View.FogEndDistance;
 	}
@@ -436,7 +494,7 @@ static void RenderViewFog(
 	GraphicsPSOInit.BoundShaderState.VertexDeclarationRHI = GFogVertexDeclaration.VertexDeclarationRHI;
 	GraphicsPSOInit.BoundShaderState.VertexShaderRHI = VertexShader.GetVertexShader();
 
-	const bool bUseFogInscatteringColorCubemap = View.FogInscatteringColorCubemap != nullptr;
+	const bool bUseFogInscatteringColorCubemap = View.FogInscatteringColorCubemap != nullptr && !IsSkyLightCaptureAffectsHeightFogTextureEnabled(View);
 	FExponentialHeightFogPS::FPermutationDomain PsPermutationVector;
 	PsPermutationVector.Set<FExponentialHeightFogPS::FSupportFogInScatteringTexture>(bUseFogInscatteringColorCubemap);
 	PsPermutationVector.Set<FExponentialHeightFogPS::FSupportFogDirectionalLightInScattering>(!bUseFogInscatteringColorCubemap && View.bUseDirectionalInscattering);
diff --git a/Engine/Source/Runtime/Renderer/Private/FogRendering.h b/Engine/Source/Runtime/Renderer/Private/FogRendering.h
--- a/Engine/Source/Runtime/Renderer/Private/FogRendering.h
+++ b/Engine/Source/Runtime/Renderer/Private/FogRendering.h
@@ -23,6 +23,9 @@ BEGIN_GLOBAL_SHADER_PARAMETER_STRUCT(FFogUniformParameters,)
 	SHADER_PARAMETER(FVector4f, DirectionalInscatteringColor)
 	SHADER_PARAMETER(FVector2f, SinCosInscatteringColorCubemapRotation)
 	SHADER_PARAMETER(FVector3f, FogInscatteringTextureParameters)
+	SHADER_PARAMETER(float, SkyLightCaptureAffectsHeightFogStrength) // When >0 and r.SupportsSkyLightCaptureAffectsHeightFogInScattering>0, uses FogInscatteringColorCubemap to sample the skylight at some roughness.
+	SHADER_PARAMETER(float, SkyLightCaptureAffectsHeightFogRoughness)
+	SHADER_PARAMETER(float, SkyLightCaptureAffectsHeightFogMipCount)
 	SHADER_PARAMETER(float, EndDistance)
 	SHADER_PARAMETER(float, ApplyVolumetricFog)
 	SHADER_PARAMETER(float, VolumetricFogStartDistance)
diff --git a/Engine/Source/Runtime/Renderer/Private/SceneCore.cpp b/Engine/Source/Runtime/Renderer/Private/SceneCore.cpp
--- a/Engine/Source/Runtime/Renderer/Private/SceneCore.cpp
+++ b/Engine/Source/Runtime/Renderer/Private/SceneCore.cpp
@@ -407,6 +407,8 @@ FExponentialHeightFogSceneInfo::FExponentialHeightFogSceneInfo(uint64 InId, cons
 
 	FogColor = State.InscatteringColorCubemap ? State.InscatteringTextureTint : State.FogInscatteringLuminance;
 	InscatteringColorCubemap = State.InscatteringColorCubemap;
+	SkyLightCaptureAffectsHeightFogRoughness = State.SkyLightCaptureAffectsHeightFogRoughness;
+	SkyLightCaptureAffectsHeightFogStrength = State.SkyLightCaptureAffectsHeightFogStrength;
 	InscatteringColorCubemapAngle = State.InscatteringColorCubemapAngle * (PI / 180.f);
 	FullyDirectionalInscatteringColorDistance = State.FullyDirectionalInscatteringColorDistance;
 	NonDirectionalInscatteringColorDistance = State.NonDirectionalInscatteringColorDistance;
diff --git a/Engine/Source/Runtime/Renderer/Private/SceneCore.h b/Engine/Source/Runtime/Renderer/Private/SceneCore.h
--- a/Engine/Source/Runtime/Renderer/Private/SceneCore.h
+++ b/Engine/Source/Runtime/Renderer/Private/SceneCore.h
@@ -168,6 +168,8 @@ public:
 	float DirectionalInscatteringStartDistance;
 	FLinearColor DirectionalInscatteringColor;
 	UTextureCube* InscatteringColorCubemap;
+	float SkyLightCaptureAffectsHeightFogRoughness;
+	float SkyLightCaptureAffectsHeightFogStrength;
 	float InscatteringColorCubemapAngle;
 	float FullyDirectionalInscatteringColorDistance;
 	float NonDirectionalInscatteringColorDistance;
