Fixed incorrect lower bound in unbounded LOD range for hierarhical CPU instance culling.

#rnx
#Synced-CL 50394525

diff --git a/Engine/Source/Runtime/Renderer/Private/InstanceCulling/InstanceCullingContext.cpp b/Engine/Source/Runtime/Renderer/Private/InstanceCulling/InstanceCullingContext.cpp
--- a/Engine/Source/Runtime/Renderer/Private/InstanceCulling/InstanceCullingContext.cpp
+++ b/Engine/Source/Runtime/Renderer/Private/InstanceCulling/InstanceCullingContext.cpp
@@ -1651,7 +1651,7 @@ void FInstanceCullingContext::SetupDrawCommands(
 					if (const FSceneCulling::FPrimitiveCullingInfo* PrimitiveCullingData = CullingQuery ? CullingQuery->GetCullingInfo(PrimitiveSceneInfo->GetPersistentIndex()) : nullptr)
 					{
 						FVector2f MinMaxLODScreenSize(
-							UnpackScreenSizeHelper(VisibleMeshDrawCommand.CullingPayload.MinScreenSize, CurrentCullingPayloadFlags, EMeshDrawCommandCullingPayloadFlags::MinScreenSizeCull, std::numeric_limits<float>::min()),
+							UnpackScreenSizeHelper(VisibleMeshDrawCommand.CullingPayload.MinScreenSize, CurrentCullingPayloadFlags, EMeshDrawCommandCullingPayloadFlags::MinScreenSizeCull, std::numeric_limits<float>::lowest()),
 							UnpackScreenSizeHelper(VisibleMeshDrawCommand.CullingPayload.MaxScreenSize, CurrentCullingPayloadFlags, EMeshDrawCommandCullingPayloadFlags::MaxScreenSizeCull, std::numeric_limits<float>::max())
 						);
 						
