[HWRT] "Show Far Field" option in Ray Tracing Debug visualizations.
- ray tracing debug visualization will now trace against Near OR Far Field TLAS depending on r.RayTracing.Visualize.FarField (which can be controlled in editor via the view mode UI).
- removed logic to trace against far field when near field trace doesn't hit anything.
- removed logic to tint far field hits in debug visualization.
- all visualization modes now work with both near and far field.

#Synced-CL 50401438

diff --git a/Engine/Shaders/Private/RayTracing/RayTracingDebug.usf b/Engine/Shaders/Private/RayTracing/RayTracingDebug.usf
--- a/Engine/Shaders/Private/RayTracing/RayTracingDebug.usf
+++ b/Engine/Shaders/Private/RayTracing/RayTracingDebug.usf
@@ -31,7 +31,6 @@ RWTexture2D<float4> Output;
 RWTexture2D<float> OutputDepth;
 
 RaytracingAccelerationStructure TLAS;
-RaytracingAccelerationStructure FarFieldTLAS;
 
 #if RAYTRACING_DEBUG_INLINE
 	StructuredBuffer<FRayTracingSceneMetadataRecord> RayTracingSceneMetadata;
@@ -40,13 +39,8 @@ RaytracingAccelerationStructure FarFieldTLAS;
 uint VisualizationMode;
 uint PickerDomain;
 uint OpaqueOnly;
-uint TraceFarField;
-uint TintFarField;
 float TimingScale;
 
-float MaxTraceDistance;
-float FarFieldMaxTraceDistance;
-
 float TriangleHitCountMaxThreshold;
 float TriangleHitCountPerInstanceMaxThreshold;
 uint TopKMostHitInstances;
@@ -72,7 +66,7 @@ float4 RoughnessDebugColor(float InRoughness)
 	return float4(pow(InRoughness, 2.2f).xxx, 0.0f);
 }
 
-float4 DebugRayTracingMaterial(uint2 PixelCoord, FRayDesc Ray, bool bIsFarField)
+float4 DebugRayTracingMaterial(uint2 PixelCoord, FRayDesc Ray)
 {
 	FRayCone RayCone = (FRayCone)0;
 	RayCone.SpreadAngle = View.EyeToPixelSpreadAngle;
@@ -82,19 +76,6 @@ float4 DebugRayTracingMaterial(uint2 PixelCoord, FRayDesc Ray, bool bIsFarField)
 	const bool bEnableSkyLightContribution = true;
 	const bool bIgnoreTranslucentMaterials = false;
 	const bool bCameraRay = true;
-
-	RaytracingAccelerationStructure TLASToTrace = TLAS;
-	if (bIsFarField)
-	{
-		TLASToTrace = FarFieldTLAS;
-
-		Ray.TMin = MaxTraceDistance;
-		Ray.TMax = FarFieldMaxTraceDistance;
-	}
-	else
-	{
-		Ray.TMax = (MaxTraceDistance > 0.0f) ? MaxTraceDistance : Ray.TMax;
-	}
 	
 	FMaterialClosestHitPayload Payload = (FMaterialClosestHitPayload)0;
 
@@ -123,7 +104,7 @@ float4 DebugRayTracingMaterial(uint2 PixelCoord, FRayDesc Ray, bool bIsFarField)
 #endif
 
 		Payload = TraceMaterialRay(
-			TLASToTrace,
+			TLAS,
 			RayFlags,
 			InstanceInclusionMask,
 			Ray,
@@ -235,7 +216,7 @@ float4 DebugRayTracingMaterial(uint2 PixelCoord, FRayDesc Ray, bool bIsFarField)
 
 // Add functions that use the debug payload only
 
-float4 DebugRayTracingInstance(FRayDesc Ray, bool bIsFarField, out float OutDepth)
+float4 DebugRayTracingInstance(FRayDesc Ray, out float OutDepth)
 {
 	float4 Result = float4(0, 0, 0, 1);
 
@@ -253,24 +234,11 @@ float4 DebugRayTracingInstance(FRayDesc Ray, bool bIsFarField, out float OutDept
 		RayFlags |= RAY_FLAG_FORCE_OPAQUE;
 	}
 
-	RaytracingAccelerationStructure TLASToTrace = TLAS;
-	if (bIsFarField)
-	{
-		TLASToTrace = FarFieldTLAS;
-
-		Ray.TMin = MaxTraceDistance;
-		Ray.TMax = FarFieldMaxTraceDistance;
-	}
-	else
-	{
-		Ray.TMax = (MaxTraceDistance > 0.0f) ? MaxTraceDistance : Ray.TMax;
-	}
-
 #if RAYTRACING_DEBUG_INLINE
 	FTraceRayInlineContext TraceRayInlineContext = CreateTraceRayInlineContext();
 
 	FTraceRayInlineResult TraceRayResult = TraceRayInline(
-		TLASToTrace,
+		TLAS,
 		RayFlags,
 		InstanceInclusionMask,
 		Ray.GetNativeDesc(),
@@ -333,7 +301,7 @@ float4 DebugRayTracingInstance(FRayDesc Ray, bool bIsFarField, out float OutDept
 	}
 #else
 	TraceRay(
-		TLASToTrace,
+		TLAS,
 		RayFlags  /*RayFlags*/,
 		InstanceInclusionMask /*InstanceInclusionMask*/,
 		0 /*RayContributionToHitGroupIndex*/,
@@ -459,44 +427,17 @@ RAY_TRACING_ENTRY_RAYGEN_OR_INLINE(RayTracingDebugMain)
 	float2 RenderTargetUV = (float2(PixelCoord) + .5f) * View.BufferSizeAndInvSize.zw;
 
 	FRayDesc Ray = CreatePrimaryRay(RenderTargetUV);
-
-	float4 Result = (float4)0;
 	
 	float OutDepth = 0.0f;
 
-	bool bIsFarField = false;
-
 	bool bWriteColorOuput = VisualizationMode != RAY_TRACING_DEBUG_VIZ_SUBSTRATE_DATA;
 	bool bWriteDepthOutput = VisualizationMode == RAY_TRACING_DEBUG_VIZ_INSTANCE_OVERLAP;
 #if USE_DEBUG_CHS || RAYTRACING_DEBUG_INLINE
-	{
-		Result = DebugRayTracingInstance(Ray, bIsFarField, OutDepth);
-		
-		const bool bNearFieldMiss = Result.w < 0.0;
-		if (bNearFieldMiss && TraceFarField > 0)
-		{
-			bIsFarField = true;
-			Result = DebugRayTracingInstance(Ray, bIsFarField, OutDepth);
-		}
-	}
+	float4 Result = DebugRayTracingInstance(Ray, OutDepth);
 #else // USE_DEBUG_CHS
-	{
-		Result = DebugRayTracingMaterial(PixelCoord, Ray, bIsFarField);
-		
-		const bool bNearFieldMiss = Result.w < 0.0;
-		if (bNearFieldMiss && TraceFarField > 0)
-		{
-			bIsFarField = true;
-			Result = DebugRayTracingMaterial(PixelCoord, Ray, bIsFarField);
-		}
-	}
+	float4 Result = DebugRayTracingMaterial(PixelCoord, Ray);
 #endif // USE_DEBUG_CHS
 
-	if (bIsFarField && TintFarField)
-	{
-		Result *= float4(1.0f, 0.5f, 0.5f, 1.0f);
-	}
-
 	// When the visualization mode does not require exposure, we disable it via RaytracingDebugViewModeNeedsTonemapping which makes this factor 1.0
 	Result.rgb *= View.PreExposure;
 
diff --git a/Engine/Source/Editor/CommonMenuExtensions/Private/RayTracingDebugVisualizationMenuCommands.cpp b/Engine/Source/Editor/CommonMenuExtensions/Private/RayTracingDebugVisualizationMenuCommands.cpp
--- a/Engine/Source/Editor/CommonMenuExtensions/Private/RayTracingDebugVisualizationMenuCommands.cpp
+++ b/Engine/Source/Editor/CommonMenuExtensions/Private/RayTracingDebugVisualizationMenuCommands.cpp
@@ -92,6 +92,9 @@ void FRayTracingDebugVisualizationMenuCommands::BuildCommandMap()
 			break;
 		}
 	}
+
+	UI_COMMAND(ShowNearFieldCommand, "Near Field", "Visualizations will show the near field layer of the ray tracing scene (disables r.RayTracing.Visualize.FarField).", EUserInterfaceActionType::RadioButton, FInputChord());
+	UI_COMMAND(ShowFarFieldCommand, "Far Field", "Visualizations will show the far field layer of the ray tracing scene (enables r.RayTracing.Visualize.FarField).", EUserInterfaceActionType::RadioButton, FInputChord());
 }
 
 void FRayTracingDebugVisualizationMenuCommands::BuildVisualisationSubMenu(UToolMenu* InMenu)
@@ -140,6 +143,11 @@ void FRayTracingDebugVisualizationMenuCommands::BuildVisualisationSubMenu(UToolM
 		FToolMenuSection& OtherSection = InMenu->AddSection("LevelViewportRayTracingVisualizationOther", LOCTEXT("RayTracingVisualizationOther", "Other"));
 		Commands.AddCommandTypeToSection(OtherSection, FVisualizationType::Other);
 	}
+
+	{
+		FToolMenuSection& LayersSection = InMenu->AddSection("LevelViewportRayTracingVisualizationLayers", LOCTEXT("RayTracingVisualizationLayers", "Layers"));
+		Commands.AddLayerCommands(LayersSection);
+	}
 }
 
 bool FRayTracingDebugVisualizationMenuCommands::AddCommandTypeToSection(FToolMenuSection& InSection, const FVisualizationType Type) const
@@ -159,6 +167,14 @@ bool FRayTracingDebugVisualizationMenuCommands::AddCommandTypeToSection(FToolMen
 	return bAddedCommands;
 }
 
+bool FRayTracingDebugVisualizationMenuCommands::AddLayerCommands(FToolMenuSection& InSection) const
+{
+	InSection.AddMenuEntry(ShowNearFieldCommand);
+	InSection.AddMenuEntry(ShowFarFieldCommand);
+
+	return true;
+}
+
 FRayTracingDebugVisualizationMenuCommands::TCommandConstIterator FRayTracingDebugVisualizationMenuCommands::CreateCommandConstIterator() const
 {
 	return CommandMap.CreateConstIterator();
@@ -182,6 +198,40 @@ void FRayTracingDebugVisualizationMenuCommands::BindCommands(FUICommandList& Com
 			FIsActionChecked::CreateStatic(&FRayTracingDebugVisualizationMenuCommands::IsRayTracingDebugVisualizationModeSelected, Client.ToWeakPtr(), Record.Name)
 		);
 	}
+
+	TWeakPtr<FEditorViewportClient> WeakClient = Client;
+	IConsoleVariable* FarFieldCVar = IConsoleManager::Get().FindConsoleVariable(TEXT("r.RayTracing.Visualize.FarField"));
+	check(FarFieldCVar);
+
+	CommandList.MapAction(
+		ShowNearFieldCommand,
+		FExecuteAction::CreateLambda([=]()
+			{
+				FarFieldCVar->Set(false, ECVF_SetByConsole);
+
+				if (TSharedPtr<FEditorViewportClient> Client = WeakClient.Pin())
+				{
+					Client->Invalidate();
+				}
+			}),
+		FCanExecuteAction::CreateLambda([]() { return true; }),
+		FIsActionChecked::CreateLambda([=]() { return !FarFieldCVar->GetBool(); })
+	);
+
+	CommandList.MapAction(
+		ShowFarFieldCommand,
+		FExecuteAction::CreateLambda([=]()
+			{
+				FarFieldCVar->Set(true, ECVF_SetByConsole);
+				
+				if (TSharedPtr<FEditorViewportClient> Client = WeakClient.Pin())
+				{
+					Client->Invalidate();
+				}
+			}),
+		FCanExecuteAction::CreateLambda([]() { return true; }),
+		FIsActionChecked::CreateLambda([=]() { return FarFieldCVar->GetBool(); })
+	);
 }
 
 void FRayTracingDebugVisualizationMenuCommands::ChangeRayTracingDebugVisualizationMode(TWeakPtr<FEditorViewportClient> WeakClient, FName InName)
diff --git a/Engine/Source/Editor/CommonMenuExtensions/Public/RayTracingDebugVisualizationMenuCommands.h b/Engine/Source/Editor/CommonMenuExtensions/Public/RayTracingDebugVisualizationMenuCommands.h
--- a/Engine/Source/Editor/CommonMenuExtensions/Public/RayTracingDebugVisualizationMenuCommands.h
+++ b/Engine/Source/Editor/CommonMenuExtensions/Public/RayTracingDebugVisualizationMenuCommands.h
@@ -65,11 +65,15 @@ private:
 	void BuildCommandMap();
 
 	bool AddCommandTypeToSection(FToolMenuSection& InSection, const FVisualizationType Type) const;
+	bool AddLayerCommands(FToolMenuSection& InSection) const;
 
 	static void ChangeRayTracingDebugVisualizationMode(TWeakPtr<FEditorViewportClient> WeakClient, FName InName);
 	static bool IsRayTracingDebugVisualizationModeSelected(TWeakPtr<FEditorViewportClient> WeakClient, FName InName);
 
 	TVisualizationModeCommandMap CommandMap;
+
+	TSharedPtr<FUICommandInfo> ShowNearFieldCommand;
+	TSharedPtr<FUICommandInfo> ShowFarFieldCommand;
 };
 
 #undef UE_API
diff --git a/Engine/Source/Runtime/Renderer/Private/RayTracing/RayTracingBarycentrics.cpp b/Engine/Source/Runtime/Renderer/Private/RayTracing/RayTracingBarycentrics.cpp
--- a/Engine/Source/Runtime/Renderer/Private/RayTracing/RayTracingBarycentrics.cpp
+++ b/Engine/Source/Runtime/Renderer/Private/RayTracing/RayTracingBarycentrics.cpp
@@ -139,11 +139,11 @@ IMPLEMENT_GLOBAL_SHADER(FRayTracingBarycentricsCS, "/Engine/Private/RayTracing/R
 
 extern float GetRayTracingDebugTimingScale();
 
-void RenderRayTracingBarycentricsCS(FRDGBuilder& GraphBuilder, const FViewInfo& View, FRDGTextureRef SceneColor, bool bVisualizeProceduralPrimitives, bool bOutputTiming)
+void RenderRayTracingBarycentricsCS(FRDGBuilder& GraphBuilder, const FViewInfo& View, FRDGTextureRef SceneColor, ERayTracingSceneLayer LayerToTrace, bool bVisualizeProceduralPrimitives, bool bOutputTiming)
 {
 	FRayTracingBarycentricsCS::FParameters* PassParameters = GraphBuilder.AllocParameters<FRayTracingBarycentricsCS::FParameters>();
 
-	PassParameters->TLAS = View.GetRayTracingSceneLayerViewChecked(ERayTracingSceneLayer::Base);
+	PassParameters->TLAS = View.GetRayTracingSceneLayerViewChecked(LayerToTrace);
 	PassParameters->Output = GraphBuilder.CreateUAV(SceneColor);
 	PassParameters->ViewUniformBuffer = View.ViewUniformBuffer;
 	PassParameters->Scene = View.GetSceneUniforms().GetBuffer(GraphBuilder);
@@ -167,7 +167,7 @@ void RenderRayTracingBarycentricsCS(FRDGBuilder& GraphBuilder, const FViewInfo&
 	FComputeShaderUtils::AddPass(GraphBuilder, RDG_EVENT_NAME("Barycentrics"), ComputeShader, PassParameters, GroupCount);
 }
 
-void RenderRayTracingBarycentricsRGS(FRDGBuilder& GraphBuilder, const FScene& Scene, const FViewInfo& View, FRDGTextureRef SceneColor, bool bOutputTiming)
+void RenderRayTracingBarycentricsRGS(FRDGBuilder& GraphBuilder, const FScene& Scene, const FViewInfo& View, FRDGTextureRef SceneColor, ERayTracingSceneLayer LayerToTrace, bool bOutputTiming)
 {
 	FRayTracingBarycentricsRGS::FPermutationDomain PermutationVector;
 	PermutationVector.Set<FRayTracingBarycentricsRGS::FOutputTiming>(bOutputTiming);
@@ -199,7 +199,7 @@ void RenderRayTracingBarycentricsRGS(FRDGBuilder& GraphBuilder, const FScene& Sc
 	   
 	FRayTracingBarycentricsRGS::FParameters* RayGenParameters = GraphBuilder.AllocParameters<FRayTracingBarycentricsRGS::FParameters>();
 
-	RayGenParameters->TLAS = View.GetRayTracingSceneLayerViewChecked(ERayTracingSceneLayer::Base);
+	RayGenParameters->TLAS = View.GetRayTracingSceneLayerViewChecked(LayerToTrace);
 	RayGenParameters->ViewUniformBuffer = View.ViewUniformBuffer;
 	RayGenParameters->Scene = View.GetSceneUniforms().GetBuffer(GraphBuilder);
 	RayGenParameters->NaniteRayTracing = Nanite::GRayTracingManager.GetUniformBuffer();
@@ -230,7 +230,7 @@ void RenderRayTracingBarycentricsRGS(FRDGBuilder& GraphBuilder, const FScene& Sc
 	});
 }
 
-void RenderRayTracingBarycentrics(FRDGBuilder& GraphBuilder, const FScene& Scene, const FViewInfo& View, FRDGTextureRef SceneColor, bool bVisualizeProceduralPrimitives, bool bOutputTiming)
+void RenderRayTracingBarycentrics(FRDGBuilder& GraphBuilder, const FScene& Scene, const FViewInfo& View, FRDGTextureRef SceneColor, ERayTracingSceneLayer LayerToTrace, bool bVisualizeProceduralPrimitives, bool bOutputTiming)
 {
 	const bool bRayTracingInline = ShouldRenderRayTracingEffect(true, ERayTracingPipelineCompatibilityFlags::Inline, View);
 	const bool bRayTracingPipeline = ShouldRenderRayTracingEffect(true, ERayTracingPipelineCompatibilityFlags::FullPipeline, View);
@@ -238,11 +238,11 @@ void RenderRayTracingBarycentrics(FRDGBuilder& GraphBuilder, const FScene& Scene
 	// currently bOutputTiming is only supported on Nvidia when using RGS
 	if (bRayTracingInline && (!bOutputTiming || !IsRHIDeviceNVIDIA()))
 	{
-		RenderRayTracingBarycentricsCS(GraphBuilder, View, SceneColor, bVisualizeProceduralPrimitives, bOutputTiming);
+		RenderRayTracingBarycentricsCS(GraphBuilder, View, SceneColor, LayerToTrace, bVisualizeProceduralPrimitives, bOutputTiming);
 	}
 	else if (bRayTracingPipeline)
 	{
-		RenderRayTracingBarycentricsRGS(GraphBuilder, Scene, View, SceneColor, bOutputTiming);
+		RenderRayTracingBarycentricsRGS(GraphBuilder, Scene, View, SceneColor, LayerToTrace, bOutputTiming);
 	}
 }
 #endif
diff --git a/Engine/Source/Runtime/Renderer/Private/RayTracing/RayTracingDebug.cpp b/Engine/Source/Runtime/Renderer/Private/RayTracing/RayTracingDebug.cpp
--- a/Engine/Source/Runtime/Renderer/Private/RayTracing/RayTracingDebug.cpp
+++ b/Engine/Source/Runtime/Renderer/Private/RayTracing/RayTracingDebug.cpp
@@ -77,15 +77,8 @@ static TAutoConsoleVariable<int32> CVarRayTracingVisualizeOpaqueOnly(
 
 static TAutoConsoleVariable<bool> CVarRayTracingVisualizeFarField(
 	TEXT("r.RayTracing.Visualize.FarField"),
-	true,
-	TEXT("Whether to include far field geometry in ray tracing visualizations."),
-	ECVF_RenderThreadSafe
-);
-
-static TAutoConsoleVariable<bool> CVarRayTracingVisualizeFarFieldTint(
-	TEXT("r.RayTracing.Visualize.FarField.Tint"),
-	true,
-	TEXT("Whether to tint far field geometry in ray tracing visualizations."),
+	false,
+	TEXT("Whether to show far field geometry instead of near field in ray tracing visualizations."),
 	ECVF_RenderThreadSafe
 );
 
@@ -199,10 +192,6 @@ BEGIN_SHADER_PARAMETER_STRUCT(FRayTracingDebugParameters, )
 	SHADER_PARAMETER(uint32, VisualizationMode)
 	SHADER_PARAMETER(uint32, PickerDomain)
 	SHADER_PARAMETER(float, TimingScale)
-	SHADER_PARAMETER(float, MaxTraceDistance)
-	SHADER_PARAMETER(float, FarFieldMaxTraceDistance)
-	SHADER_PARAMETER(uint32, TraceFarField)
-	SHADER_PARAMETER(uint32, TintFarField)
 	SHADER_PARAMETER(uint32, OpaqueOnly)
 	SHADER_PARAMETER(float, TriangleHitCountMaxThreshold)
 	SHADER_PARAMETER(float, TriangleHitCountPerInstanceMaxThreshold)
@@ -210,7 +199,6 @@ BEGIN_SHADER_PARAMETER_STRUCT(FRayTracingDebugParameters, )
 	SHADER_PARAMETER(uint32, NumTotalInstances)
 	SHADER_PARAMETER(uint32, SubstrateDebugDataSizeInUints)
 	SHADER_PARAMETER_RDG_BUFFER_SRV(RaytracingAccelerationStructure, TLAS)
-	SHADER_PARAMETER_RDG_BUFFER_SRV(RaytracingAccelerationStructure, FarFieldTLAS)
 	SHADER_PARAMETER_RDG_TEXTURE_UAV(RWTexture2D<float4>, Output)
 	SHADER_PARAMETER_RDG_TEXTURE_UAV(RWTexture2D<float>, OutputDepth)
 	SHADER_PARAMETER_RDG_BUFFER_SRV(StructuredBuffer, InstancesExtraData)
@@ -993,7 +981,7 @@ static FRDGBufferRef RayTracingPerformPicking(
 	const FScene& Scene,
 	const FViewInfo& View,
 	TConstArrayView<FRayTracingShaderBindingData> VisibleRayTracingShaderBindings,
-	FRayTracingPickingFeedback& PickingFeedback, bool bInlineRayTracing)
+	FRayTracingPickingFeedback& PickingFeedback, ERayTracingSceneLayer LayerToTrace, bool bInlineRayTracing)
 {
 	const FRayTracingScene& RayTracingScene = Scene.RayTracingScene;
 
@@ -1002,11 +990,11 @@ static FRDGBufferRef RayTracingPerformPicking(
 	FRDGBufferRef PickingBuffer = GraphBuilder.CreateBuffer(PickingBufferDesc, TEXT("RayTracingDebug.PickingBuffer"));
 
 	FRayTracingPickingParameters SharedParameters = {};
-	SharedParameters.InstancesExtraData = GraphBuilder.CreateSRV(RayTracingScene.GetInstanceExtraDataBuffer(ERayTracingSceneLayer::Base, View.GetRayTracingSceneViewHandle()));
-	SharedParameters.InstancesDebugData = GraphBuilder.CreateSRV(RayTracingScene.GetInstanceDebugBuffer(ERayTracingSceneLayer::Base));
-	SharedParameters.TLAS = RayTracingScene.GetLayerView(ERayTracingSceneLayer::Base, View.GetRayTracingSceneViewHandle());
+	SharedParameters.InstancesExtraData = GraphBuilder.CreateSRV(RayTracingScene.GetInstanceExtraDataBuffer(LayerToTrace, View.GetRayTracingSceneViewHandle()));
+	SharedParameters.InstancesDebugData = GraphBuilder.CreateSRV(RayTracingScene.GetInstanceDebugBuffer(LayerToTrace));
+	SharedParameters.TLAS = RayTracingScene.GetLayerView(LayerToTrace, View.GetRayTracingSceneViewHandle());
 	SharedParameters.OpaqueOnly = CVarRayTracingVisualizeOpaqueOnly.GetValueOnRenderThread();
-	SharedParameters.InstanceBuffer = GraphBuilder.CreateSRV(RayTracingScene.GetInstanceBuffer(ERayTracingSceneLayer::Base, View.GetRayTracingSceneViewHandle()));
+	SharedParameters.InstanceBuffer = GraphBuilder.CreateSRV(RayTracingScene.GetInstanceBuffer(LayerToTrace, View.GetRayTracingSceneViewHandle()));
 	SharedParameters.ViewUniformBuffer = View.ViewUniformBuffer;
 	SharedParameters.SceneUniformBuffer = GetSceneUniformBufferRef(GraphBuilder, View); // TODO: use a separate params structure
 	SharedParameters.NaniteRayTracingUniformBuffer = Nanite::GRayTracingManager.GetUniformBuffer();
@@ -1274,7 +1262,12 @@ static void PrintTopKMostHitMessage(FRDGBuilder& GraphBuilder, const FScene& Sce
 	}
 }
 
-static FRDGBufferRef RayTracingPerformHitStatsPerPrimitive(FRDGBuilder& GraphBuilder, const FScene& Scene, const FViewInfo& View, TConstArrayView<FRayTracingShaderBindingData> VisibleRayTracingShaderBindings)
+static FRDGBufferRef RayTracingPerformHitStatsPerPrimitive(
+	FRDGBuilder& GraphBuilder,
+	const FScene& Scene,
+	const FViewInfo& View,
+	TConstArrayView<FRayTracingShaderBindingData> VisibleRayTracingShaderBindings,
+	ERayTracingSceneLayer LayerToTrace)
 {
 	const FRayTracingScene& RayTracingScene = Scene.RayTracingScene;
 
@@ -1309,7 +1302,7 @@ static FRDGBufferRef RayTracingPerformHitStatsPerPrimitive(FRDGBuilder& GraphBui
 	FShaderBindingTableRHIRef HitStatsSBT = Scene.RayTracingSBT.AllocateTransientRHI(GraphBuilder.RHICmdList, ERayTracingShaderBindingMode::RTPSO, ERayTracingHitGroupIndexingMode::Allow, Initializer.GetMaxLocalBindingDataSize());
 
 	// TODO: Should check RayTracingScene for actual number of instances instead of max number in FRHIRayTracingScene initializer
-	const uint32 NumInstancesInTLAS = FMath::Max(RayTracingScene.GetRHIRayTracingSceneChecked(ERayTracingSceneLayer::Base, View.GetRayTracingSceneViewHandle())->GetInitializer().MaxNumInstances, (uint32)CVarRayTracingVisualizeHitCountTopKHits.GetValueOnRenderThread());
+	const uint32 NumInstancesInTLAS = FMath::Max(RayTracingScene.GetRHIRayTracingSceneChecked(LayerToTrace, View.GetRayTracingSceneViewHandle())->GetInitializer().MaxNumInstances, (uint32)CVarRayTracingVisualizeHitCountTopKHits.GetValueOnRenderThread());
 	FRDGBufferDesc HitStatsPerPrimitiveBufferDesc = FRDGBufferDesc::CreateStructuredDesc(sizeof(FRayTracingHitStatsEntry), NumInstancesInTLAS);
 	HitStatsPerPrimitiveBufferDesc.Usage = EBufferUsageFlags(HitStatsPerPrimitiveBufferDesc.Usage | BUF_SourceCopy);
 	FRDGBufferRef HitStatsBuffer = GraphBuilder.CreateBuffer(HitStatsPerPrimitiveBufferDesc, TEXT("RayTracingDebug.HitStatsBuffer"));
@@ -1319,7 +1312,7 @@ static FRDGBufferRef RayTracingPerformHitStatsPerPrimitive(FRDGBuilder& GraphBui
 	DebugHitStatsUniformBuffer = GraphBuilder.CreateUniformBuffer(DebugHitStatsUniformBufferParameters);
 
 	FRayTracingDebugHitStatsRGS::FParameters* RayGenParameters = GraphBuilder.AllocParameters<FRayTracingDebugHitStatsRGS::FParameters>();
-	RayGenParameters->TLAS = RayTracingScene.GetLayerView(ERayTracingSceneLayer::Base, View.GetRayTracingSceneViewHandle());
+	RayGenParameters->TLAS = RayTracingScene.GetLayerView(LayerToTrace, View.GetRayTracingSceneViewHandle());
 	RayGenParameters->OpaqueOnly = CVarRayTracingVisualizeOpaqueOnly.GetValueOnRenderThread();
 	RayGenParameters->ViewUniformBuffer = View.ViewUniformBuffer;
 	RayGenParameters->SceneUniformBuffer = GetSceneUniformBufferRef(GraphBuilder, View); // TODO: use a separate params structure
@@ -1471,7 +1464,7 @@ static void RayTracingDrawInstances(FRDGBuilder& GraphBuilder, const FViewInfo&
 		});
 }
 
-static void DrawInstanceOverlap(FRDGBuilder& GraphBuilder, const FScene& Scene, const FViewInfo& View, FRDGTextureRef SceneColorTexture, FRDGTextureRef InputDepthTexture)
+static void DrawInstanceOverlap(FRDGBuilder& GraphBuilder, const FScene& Scene, const FViewInfo& View, FRDGTextureRef SceneColorTexture, FRDGTextureRef InputDepthTexture, ERayTracingSceneLayer LayerToTrace)
 {
 	FRDGTextureRef SceneDepthTexture = GraphBuilder.CreateTexture(
 		FRDGTextureDesc::Create2D(
@@ -1511,8 +1504,8 @@ static void DrawInstanceOverlap(FRDGBuilder& GraphBuilder, const FScene& Scene,
 		View,
 		InstanceOverlapTexture,
 		SceneDepthTexture,
-		Scene.RayTracingScene.GetInstanceExtraDataBuffer(ERayTracingSceneLayer::Base, View.GetRayTracingSceneViewHandle()),
-		Scene.RayTracingScene.GetNumNativeInstances(ERayTracingSceneLayer::Base, View.GetRayTracingSceneViewHandle()),
+		Scene.RayTracingScene.GetInstanceExtraDataBuffer(LayerToTrace, View.GetRayTracingSceneViewHandle()),
+		Scene.RayTracingScene.GetNumNativeInstances(LayerToTrace, View.GetRayTracingSceneViewHandle()),
 		false);
 
 	// Calculate heatmap of instance overlap and blend it on top of ray tracing debug output
@@ -1546,8 +1539,8 @@ static void DrawInstanceOverlap(FRDGBuilder& GraphBuilder, const FScene& Scene,
 			View,
 			SceneColorTexture,
 			SceneDepthTexture,
-			Scene.RayTracingScene.GetInstanceExtraDataBuffer(ERayTracingSceneLayer::Base, View.GetRayTracingSceneViewHandle()),
-			Scene.RayTracingScene.GetNumNativeInstances(ERayTracingSceneLayer::Base, View.GetRayTracingSceneViewHandle()),
+			Scene.RayTracingScene.GetInstanceExtraDataBuffer(LayerToTrace, View.GetRayTracingSceneViewHandle()),
+			Scene.RayTracingScene.GetNumNativeInstances(LayerToTrace, View.GetRayTracingSceneViewHandle()),
 			true);
 	}
 }
@@ -1594,12 +1587,12 @@ bool HasRaytracingDebugViewModeRaytracedOverlay(const FSceneViewFamily& ViewFami
 	return !bAnySubstrate && !bAnyTraversalSecondary;
 }
 
-extern void RenderRayTracingBarycentrics(FRDGBuilder& GraphBuilder, const FScene& Scene, const FViewInfo& View, FRDGTextureRef SceneColor, bool bVisualizeProceduralPrimitives, bool bOutputTiming);
+extern void RenderRayTracingBarycentrics(FRDGBuilder& GraphBuilder, const FScene& Scene, const FViewInfo& View, FRDGTextureRef SceneColor, ERayTracingSceneLayer LayerToTrace, bool bVisualizeProceduralPrimitives, bool bOutputTiming);
 
 extern void RenderRayTracingPrimaryRaysView(
 	FRDGBuilder& GraphBuilder, const FViewInfo& View, const FSceneTextures& SceneTextures,
 	FRDGTextureRef* InOutColorTexture, FRDGTextureRef* InOutRayHitDistanceTexture,
-	int32 SamplePerPixel, int32 HeightFog, float ResolutionFraction, ERayTracingPrimaryRaysFlag Flags);
+	ERayTracingSceneLayer LayerToTrace, int32 SamplePerPixel, int32 HeightFog, float ResolutionFraction, ERayTracingPrimaryRaysFlag Flags);
 
 void RenderRayTracingDebug(FRDGBuilder& GraphBuilder, const FScene& Scene, const FViewInfo& View, FSceneTextures& SceneTextures, TConstArrayView<FRayTracingShaderBindingData> VisibleRayTracingShaderBindings, FRayTracingPickingFeedback& PickingFeedback)
 {
@@ -1608,6 +1601,8 @@ void RenderRayTracingDebug(FRDGBuilder& GraphBuilder, const FScene& Scene, const
 	FSceneTextureParameters SceneTextureParameters = GetSceneTextureParameters(GraphBuilder, SceneTextures.UniformBuffer);
 	FRDGTextureRef SceneColorTexture = SceneTextures.Color.Target;
 
+	const ERayTracingSceneLayer LayerToTrace = CVarRayTracingVisualizeFarField.GetValueOnRenderThread() ? ERayTracingSceneLayer::FarField : ERayTracingSceneLayer::Base;
+
 	const uint32 DebugVisualizationMode = GetRaytracingDebugViewModeID(View);
 	const bool bSubstratePixelDebugEnable = DebugVisualizationMode == RAY_TRACING_DEBUG_VIZ_SUBSTRATE_DATA && View.ViewState && Substrate::IsSubstrateEnabled();
 
@@ -1620,7 +1615,7 @@ void RenderRayTracingDebug(FRDGBuilder& GraphBuilder, const FScene& Scene, const
 
 	if (DebugVisualizationMode == RAY_TRACING_DEBUG_VIZ_BARYCENTRICS || DebugVisualizationMode == RAY_TRACING_DEBUG_VIZ_TIMING_TRAVERSAL)
 	{
-		return RenderRayTracingBarycentrics(GraphBuilder, Scene, View, SceneColorTexture, (bool)GVisualizeProceduralPrimitives, /*bOutputTiming*/ DebugVisualizationMode == RAY_TRACING_DEBUG_VIZ_TIMING_TRAVERSAL);
+		return RenderRayTracingBarycentrics(GraphBuilder, Scene, View, SceneColorTexture, LayerToTrace, (bool)GVisualizeProceduralPrimitives, /*bOutputTiming*/ DebugVisualizationMode == RAY_TRACING_DEBUG_VIZ_TIMING_TRAVERSAL);
 	}
 
 	const FRayTracingScene& RayTracingScene = Scene.RayTracingScene;
@@ -1632,7 +1627,7 @@ void RenderRayTracingDebug(FRDGBuilder& GraphBuilder, const FScene& Scene, const
 
 		FRayTracingDebugTraversalCS::FParameters* PassParameters = GraphBuilder.AllocParameters<FRayTracingDebugTraversalCS::FParameters>();
 		PassParameters->Output = GraphBuilder.CreateUAV(SceneColorTexture);
-		PassParameters->TLAS = RayTracingScene.GetLayerView(ERayTracingSceneLayer::Base, View.GetRayTracingSceneViewHandle());
+		PassParameters->TLAS = RayTracingScene.GetLayerView(LayerToTrace, View.GetRayTracingSceneViewHandle());
 		PassParameters->ViewUniformBuffer = View.ViewUniformBuffer;
 		PassParameters->Scene = View.GetSceneUniforms().GetBuffer(GraphBuilder);
 		PassParameters->NaniteRasterUniformBuffer = CreateDebugNaniteRasterUniformBuffer(GraphBuilder);
@@ -1698,7 +1693,7 @@ void RenderRayTracingDebug(FRDGBuilder& GraphBuilder, const FScene& Scene, const
 
 		RenderRayTracingPrimaryRaysView(
 			GraphBuilder, View, SceneTextures, 
-			&OutputColor, &HitDistanceTexture, 1, 1, 1, ERayTracingPrimaryRaysFlag::PrimaryView);
+			&OutputColor, &HitDistanceTexture, LayerToTrace, 1, 1, 1, ERayTracingPrimaryRaysFlag::PrimaryView);
 
 		AddDrawTexturePass(GraphBuilder, View, OutputColor, SceneColorTexture, View.ViewRect.Min, View.ViewRect.Min, View.ViewRect.Size());
 		return;
@@ -1707,10 +1702,10 @@ void RenderRayTracingDebug(FRDGBuilder& GraphBuilder, const FScene& Scene, const
 	FRDGBufferRef PickingBuffer = nullptr;
 	FRDGBufferRef StatsBuffer = nullptr;
 	if (IsRayTracingPickingEnabled(DebugVisualizationMode)
-		&& RayTracingScene.GetInstanceExtraDataBuffer(ERayTracingSceneLayer::Base, View.GetRayTracingSceneViewHandle()) != nullptr
-		&& RayTracingScene.GetInstanceDebugBuffer(ERayTracingSceneLayer::Base) != nullptr)
+		&& RayTracingScene.GetInstanceExtraDataBuffer(LayerToTrace, View.GetRayTracingSceneViewHandle()) != nullptr
+		&& RayTracingScene.GetInstanceDebugBuffer(LayerToTrace) != nullptr)
 	{
-		PickingBuffer = RayTracingPerformPicking(GraphBuilder, Scene, View, VisibleRayTracingShaderBindings, PickingFeedback, bInlineRayTracing);
+		PickingBuffer = RayTracingPerformPicking(GraphBuilder, Scene, View, VisibleRayTracingShaderBindings, PickingFeedback, LayerToTrace, bInlineRayTracing);
 	}
 	else
 	{
@@ -1719,7 +1714,7 @@ void RenderRayTracingDebug(FRDGBuilder& GraphBuilder, const FScene& Scene, const
 
 	if (DebugVisualizationMode == RAY_TRACING_DEBUG_VIZ_HITCOUNT_PER_INSTANCE)
 	{
-		StatsBuffer = RayTracingPerformHitStatsPerPrimitive(GraphBuilder, Scene, View, VisibleRayTracingShaderBindings);
+		StatsBuffer = RayTracingPerformHitStatsPerPrimitive(GraphBuilder, Scene, View, VisibleRayTracingShaderBindings, LayerToTrace);
 	}
 	else 
 	{
@@ -1732,19 +1727,19 @@ void RenderRayTracingDebug(FRDGBuilder& GraphBuilder, const FScene& Scene, const
 		AddClearUAVPass(GraphBuilder, DebugHitStatsUniformBufferParameters->HitStatsOutput, 0);
 	}
 
-	FRDGBufferRef InstanceExtraDataBuffer = RayTracingScene.GetInstanceExtraDataBuffer(ERayTracingSceneLayer::Base, View.GetRayTracingSceneViewHandle());
+	FRDGBufferRef InstanceExtraDataBuffer = RayTracingScene.GetInstanceExtraDataBuffer(LayerToTrace, View.GetRayTracingSceneViewHandle());
 	if (InstanceExtraDataBuffer == nullptr)
 	{
 		InstanceExtraDataBuffer = GSystemTextures.GetDefaultStructuredBuffer(GraphBuilder, sizeof(FRayTracingInstanceExtraData));
 	}
 
-	FRDGBufferRef InstanceDebugBuffer = RayTracingScene.GetInstanceDebugBuffer(ERayTracingSceneLayer::Base);
+	FRDGBufferRef InstanceDebugBuffer = RayTracingScene.GetInstanceDebugBuffer(LayerToTrace);
 	if (InstanceDebugBuffer == nullptr)
 	{
 		InstanceDebugBuffer = GSystemTextures.GetDefaultStructuredBuffer(GraphBuilder, sizeof(FRayTracingInstanceDebugData));
 	}
 
-	FRDGBufferRef InstanceBuffer = RayTracingScene.GetInstanceBuffer(ERayTracingSceneLayer::Base, View.GetRayTracingSceneViewHandle());
+	FRDGBufferRef InstanceBuffer = RayTracingScene.GetInstanceBuffer(LayerToTrace, View.GetRayTracingSceneViewHandle());
 	if (InstanceBuffer == nullptr)
 	{
 		InstanceBuffer = GSystemTextures.GetDefaultStructuredBuffer(GraphBuilder, sizeof(uint32));
@@ -1752,15 +1747,13 @@ void RenderRayTracingDebug(FRDGBuilder& GraphBuilder, const FScene& Scene, const
 
 	const bool bRequiresDebugCHS = RequiresRayTracingDebugCHS(DebugVisualizationMode);
 
-	const uint32 NumInstances = RayTracingScene.GetNumNativeInstances(ERayTracingSceneLayer::Base, View.GetRayTracingSceneViewHandle());
+	const uint32 NumInstances = RayTracingScene.GetNumNativeInstances(LayerToTrace, View.GetRayTracingSceneViewHandle());
 
 	FRayTracingDebugParameters SharedParameters;
 	SharedParameters.VisualizationMode = DebugVisualizationMode;
 	SharedParameters.PickerDomain = CVarRayTracingVisualizePickerDomain.GetValueOnRenderThread();
 	SharedParameters.TimingScale = GetRayTracingDebugTimingScale();
 	SharedParameters.OpaqueOnly = CVarRayTracingVisualizeOpaqueOnly.GetValueOnRenderThread();
-	SharedParameters.TraceFarField = CVarRayTracingVisualizeFarField.GetValueOnRenderThread() ? 1 : 0;
-	SharedParameters.TintFarField = CVarRayTracingVisualizeFarFieldTint.GetValueOnRenderThread() ? 1 : 0;
 	SharedParameters.TriangleHitCountMaxThreshold = FMath::Clamp((float)CVarRayTracingVisualizeHitCountMaxThreshold.GetValueOnRenderThread(), 1, 100000);
 	SharedParameters.TriangleHitCountPerInstanceMaxThreshold = FMath::Max(1, CVarRayTracingVisualizeHitCountPerInstanceMaxThreshold.GetValueOnRenderThread());
 	SharedParameters.RayTracingDebugHitStatsUniformBuffer = DebugHitStatsUniformBuffer;
@@ -1794,24 +1787,12 @@ void RenderRayTracingDebug(FRDGBuilder& GraphBuilder, const FScene& Scene, const
 		TexCreate_ShaderResource | TexCreate_UAV);
 	FRDGTextureRef OutputDepthTexture = GraphBuilder.CreateTexture(OutputDepthTextureDesc, TEXT("RayTracingDebug::Depth"));
 	SharedParameters.OutputDepth = GraphBuilder.CreateUAV(OutputDepthTexture);
-
-	if (Lumen::UseFarField(*View.Family) || MegaLights::UseFarField(*View.Family))
-	{
-		SharedParameters.MaxTraceDistance = Lumen::GetMaxTraceDistance(View);
-		SharedParameters.FarFieldMaxTraceDistance = Lumen::GetFarFieldMaxTraceDistance();
-	}
-	else
-	{
-		SharedParameters.MaxTraceDistance = 0.0f;
-		SharedParameters.FarFieldMaxTraceDistance = 0.0f;
-	}
 	
 	SharedParameters.InstancesExtraData = GraphBuilder.CreateSRV(InstanceExtraDataBuffer);
 	SharedParameters.InstancesDebugData = GraphBuilder.CreateSRV(InstanceDebugBuffer);
 	SharedParameters.InstanceBuffer = GraphBuilder.CreateSRV(InstanceBuffer);
 	SharedParameters.PickingBuffer = GraphBuilder.CreateSRV(PickingBuffer);
-	SharedParameters.TLAS = RayTracingScene.GetLayerView(ERayTracingSceneLayer::Base, View.GetRayTracingSceneViewHandle());
-	SharedParameters.FarFieldTLAS = RayTracingScene.GetLayerView(ERayTracingSceneLayer::FarField, View.GetRayTracingSceneViewHandle());
+	SharedParameters.TLAS = RayTracingScene.GetLayerView(LayerToTrace, View.GetRayTracingSceneViewHandle());
 	SharedParameters.ViewUniformBuffer = View.ViewUniformBuffer;
 	SharedParameters.Output = GraphBuilder.CreateUAV(SceneColorTexture);
 	SharedParameters.TopKHitStats = GraphBuilder.CreateSRV(StatsBuffer);
@@ -1929,7 +1910,7 @@ void RenderRayTracingDebug(FRDGBuilder& GraphBuilder, const FScene& Scene, const
 
 	if (DebugVisualizationMode == RAY_TRACING_DEBUG_VIZ_INSTANCE_OVERLAP)
 	{
-		DrawInstanceOverlap(GraphBuilder, Scene, View, SceneColorTexture, OutputDepthTexture);
+		DrawInstanceOverlap(GraphBuilder, Scene, View, SceneColorTexture, OutputDepthTexture, LayerToTrace);
 	}
 
 	if (bSubstratePixelDebugEnable)
@@ -1963,6 +1944,8 @@ void FDeferredShadingSceneRenderer::RayTracingDisplayPicking(const FRayTracingPi
 	{
 		return;
 	}
+
+	const ERayTracingSceneLayer LayerToTrace = CVarRayTracingVisualizeFarField.GetValueOnRenderThread() ? ERayTracingSceneLayer::FarField : ERayTracingSceneLayer::Base;
 	
 	int32 PickerDomain = CVarRayTracingVisualizePickerDomain.GetValueOnRenderThread();
 	switch (PickerDomain)
@@ -2008,7 +1991,7 @@ void FDeferredShadingSceneRenderer::RayTracingDisplayPicking(const FRayTracingPi
 	Writer.EmptyLine();
 
 	FRHIRayTracingGeometry* Geometry = nullptr;
-	for (const FRayTracingGeometryInstance& Instance : Scene->RayTracingScene.GetInstances(ERayTracingSceneLayer::Base))
+	for (const FRayTracingGeometryInstance& Instance : Scene->RayTracingScene.GetInstances(LayerToTrace))
 	{
 		if (Instance.GeometryRHI)
 		{
diff --git a/Engine/Source/Runtime/Renderer/Private/RayTracing/RayTracingPrimaryRays.cpp b/Engine/Source/Runtime/Renderer/Private/RayTracing/RayTracingPrimaryRays.cpp
--- a/Engine/Source/Runtime/Renderer/Private/RayTracing/RayTracingPrimaryRays.cpp
+++ b/Engine/Source/Runtime/Renderer/Private/RayTracing/RayTracingPrimaryRays.cpp
@@ -110,6 +110,7 @@ void RenderRayTracingPrimaryRaysView(
 	const FSceneTextures& SceneTextures,
 	FRDGTextureRef* InOutColorTexture,
 	FRDGTextureRef* InOutRayHitDistanceTexture,
+	ERayTracingSceneLayer LayerToTrace,
 	int32 SamplePerPixel,
 	int32 HeightFog,
 	float ResolutionFraction,
@@ -159,7 +160,7 @@ void RenderRayTracingPrimaryRaysView(
 	PassParameters->TranslucencyRefraction = TranslucencyOptions.EnableRefraction >= 0 ? TranslucencyOptions.EnableRefraction : View.FinalPostProcessSettings.RayTracingTranslucencyRefraction;
 	PassParameters->MaxNormalBias = GetRaytracingMaxNormalBias();
 	PassParameters->PrimaryRayFlags = (uint32)Flags;
-	PassParameters->TLAS = View.GetRayTracingSceneLayerViewChecked(ERayTracingSceneLayer::Base);
+	PassParameters->TLAS = View.GetRayTracingSceneLayerViewChecked(LayerToTrace);
 	PassParameters->ViewUniformBuffer = View.ViewUniformBuffer;
 
 	PassParameters->LightGridPacked = View.RayTracingLightGridUniformBuffer;
diff --git a/Engine/Source/Runtime/Renderer/Private/RayTracing/RayTracingTranslucency.cpp b/Engine/Source/Runtime/Renderer/Private/RayTracing/RayTracingTranslucency.cpp
--- a/Engine/Source/Runtime/Renderer/Private/RayTracing/RayTracingTranslucency.cpp
+++ b/Engine/Source/Runtime/Renderer/Private/RayTracing/RayTracingTranslucency.cpp
@@ -170,7 +170,7 @@ bool ShouldRenderRayTracingTranslucency(const FViewInfo& View)
 extern void RenderRayTracingPrimaryRaysView(
 	FRDGBuilder& GraphBuilder, const FViewInfo& View, const FSceneTextures& SceneTextures,
 	FRDGTextureRef* InOutColorTexture, FRDGTextureRef* InOutRayHitDistanceTexture,
-	int32 SamplePerPixel, int32 HeightFog, float ResolutionFraction, ERayTracingPrimaryRaysFlag Flags);
+	ERayTracingSceneLayer LayerToTrace, int32 SamplePerPixel, int32 HeightFog, float ResolutionFraction, ERayTracingPrimaryRaysFlag Flags);
 
 void FDeferredShadingSceneRenderer::RenderRayTracingTranslucency(FRDGBuilder& GraphBuilder, FRDGTextureMSAA SceneColorTexture)
 {
@@ -209,7 +209,7 @@ void FDeferredShadingSceneRenderer::RenderRayTracingTranslucency(FRDGBuilder& Gr
 			RenderRayTracingPrimaryRaysView(
 				GraphBuilder, View, GetActiveSceneTextures(),
 				&DenoiserInputs.Color, &DenoiserInputs.RayHitDistance,
-				TranslucencySPP, GRayTracingTranslucencyHeightFog, ResolutionFraction, 
+				ERayTracingSceneLayer::Base, TranslucencySPP, GRayTracingTranslucencyHeightFog, ResolutionFraction,
 				ERayTracingPrimaryRaysFlag::AllowSkipSkySample | ERayTracingPrimaryRaysFlag::UseGBufferForMaxDistance);
 
 			const FScreenPassTexture SceneColor(DenoiserInputs.Color, View.ViewRect);
