Light function atlas - Directional light can now use light function atlas for color light functions.

Removed limitation of directional lights from CanLightUsesAtlasForUnbatchedLight as this is not longer needed: we are now preventing light function material sampling UV coordinate (before it was when using UMaterialExpressionTextureCoordinate and this CL adds when a texture sample node has its texcoord input plugged in). Except when the use is forcing the compatibility with the atlas on the material itself. So all cases should be covered.

History:
- light function added in 32540109 but we were only checking vertex position and depth for validity. So we disabled all dir light to be sure.
- Then in 32971273 naming was clarified and in between those CL more checks we added and the force compatible with atlas bool on the material. With that now we should be in a safe spot to allow light function on directional lights.

[RN] directional light can now use light functionstored in atlas so that color light function can be used.

#tests Tested dir light function with color, without texcoord and with tex coord. Then with the compatibility boolean set to true.
#Synced-CL 50398316

diff --git a/Engine/Source/Runtime/Engine/Private/Materials/MaterialExpressions.cpp b/Engine/Source/Runtime/Engine/Private/Materials/MaterialExpressions.cpp
--- a/Engine/Source/Runtime/Engine/Private/Materials/MaterialExpressions.cpp
+++ b/Engine/Source/Runtime/Engine/Private/Materials/MaterialExpressions.cpp
@@ -4092,6 +4092,11 @@ int32 UMaterialExpressionTextureSampleParameter::Compile(class FMaterialCompiler
 		return UMaterialExpressionTextureSample::Compile(Compiler, OutputIndex);
 	}
 
+	if (Coordinates.GetTracedInput().Expression)
+	{
+		Compiler->SetPotentiallyManipulateTexCoords();
+	}
+
 	return CompileTextureSample(
 		Compiler,
 		Texture,
diff --git a/Engine/Source/Runtime/Renderer/Private/LightRendering.cpp b/Engine/Source/Runtime/Renderer/Private/LightRendering.cpp
--- a/Engine/Source/Runtime/Renderer/Private/LightRendering.cpp
+++ b/Engine/Source/Runtime/Renderer/Private/LightRendering.cpp
@@ -1334,13 +1334,6 @@ bool CanLightUsesAtlasForUnbatchedLight(ERHIFeatureLevel::Type FeatureLevel, con
 		return false;
 	}
 
-	// For now, we prevent directional light to use the light function atlas because atlas tiles needs to be repeatable.
-	// And if a texcoordinate node is not scale as a integer multiplier of the uv in [0,1] then the tile will become visible.
-	if (Proxy->GetLightType() == LightType_Directional)
-	{
-		return false;
-	}
-
 	// Material shoud also be compatible with light function atlas, i.e. not sample depth or world position.
 	// We do not check that for other systems (translucent, water, volume fog, clustered, Lumen), 
 	// because light functions were never available there before the atlas. So those LF aare still added into the atlas.
