[TSR] Adjust when translucent clipped depth pass can be triggered. It should not rely on thin geo pass.
#tests translucency clipped depth pass runs in test scene.

#Synced-CL 50303332

diff --git a/Engine/Source/Runtime/Renderer/Private/PostProcess/TemporalAA.h b/Engine/Source/Runtime/Renderer/Private/PostProcess/TemporalAA.h
--- a/Engine/Source/Runtime/Renderer/Private/PostProcess/TemporalAA.h
+++ b/Engine/Source/Runtime/Renderer/Private/PostProcess/TemporalAA.h
@@ -182,6 +182,9 @@ bool NeedTSRAntiFlickeringPass(const FViewInfo& View);
 
 bool NeedTSRThinGeometryDetectionPass(const FViewInfo& View);
 
+/** Returns whether TSR needs the bits set by translucent clipped depth pass*/
+bool NeedTSRTranslucentClippedDepthPass(const FViewInfo& View);
+
 /** Returns whether TSR internal visualization is enabled on the view. */
 bool IsVisualizeTSREnabled(const FViewInfo& View);
 
diff --git a/Engine/Source/Runtime/Renderer/Private/PostProcess/TemporalSuperResolution.cpp b/Engine/Source/Runtime/Renderer/Private/PostProcess/TemporalSuperResolution.cpp
--- a/Engine/Source/Runtime/Renderer/Private/PostProcess/TemporalSuperResolution.cpp
+++ b/Engine/Source/Runtime/Renderer/Private/PostProcess/TemporalSuperResolution.cpp
@@ -1671,6 +1671,12 @@ bool NeedTSRThinGeometryDetectionPass(const FViewInfo& View)
 		View.Family->EngineShowFlags.PostProcessing;
 }
 
+bool NeedTSRTranslucentClippedDepthPass(const FViewInfo& View)
+{
+	return GetMainTAAPassConfig(View) == EMainTAAPassConfig::TSR &&
+		View.Family->EngineShowFlags.PostProcessing;
+}
+
 bool ShouldAddTSRMainFlickeringLumaPass()
 {
 	return  (CVarTSRFlickeringEnable.GetValueOnRenderThread()!=0 && CVarTSRFlickeringPeriod.GetValueOnRenderThread() != 0.0f) || ShouldApplySkyRelaxation();
diff --git a/Engine/Source/Runtime/Renderer/Private/SceneVisibility.cpp b/Engine/Source/Runtime/Renderer/Private/SceneVisibility.cpp
--- a/Engine/Source/Runtime/Renderer/Private/SceneVisibility.cpp
+++ b/Engine/Source/Runtime/Renderer/Private/SceneVisibility.cpp
@@ -1433,7 +1433,7 @@ void FRelevancePacket::ComputeRelevance(FDynamicPrimitiveIndexList& DynamicPrimi
 	const bool bMobileBasePassAlwaysUsesCSM = (ShadingPath == EShadingPath::Mobile) && MobileBasePassAlwaysUsesCSM(Scene.GetShaderPlatform());
 	const bool bVelocityPassWritesDepth = Scene.EarlyZPassMode == DDM_AllOpaqueNoVelocity;
 	const bool bIsTranslucentHoldoutEnabled = IsPrimitiveAlphaHoldoutEnabled(View);
-	const bool bIsTranslucentClippedDepthEnabled = NeedTSRThinGeometryDetectionPass(View);
+	const bool bIsTranslucentClippedDepthEnabled = NeedTSRTranslucentClippedDepthPass(View);
 	const bool bHLODActive = Scene.SceneLODHierarchy.IsActive();
 	const bool bUsingSingleLayerWaterFallback = IsUsingSingleLayerWaterFallback(View.GetShaderPlatform());
 	const FHLODVisibilityState* const HLODState = bHLODActive && ViewState ? &ViewState->HLODVisibilityState : nullptr;
@@ -2405,7 +2405,7 @@ static void ComputeDynamicMeshRelevance(
 	const FPrimitiveBounds& Bounds)
 {
 	const int32 NumElements = MeshBatch.Mesh->Elements.Num();
-	const bool bIsTranslucentClippedDepthEnabled = NeedTSRThinGeometryDetectionPass(View);
+	const bool bIsTranslucentClippedDepthEnabled = NeedTSRTranslucentClippedDepthPass(View);
 
 	if (ViewRelevance.bDrawRelevance && (ViewRelevance.bRenderInMainPass || ViewRelevance.bRenderCustomDepth || ViewRelevance.bRenderInDepthPass))
 	{
