Move FinishRelevantDecals call to the render thread instead of running it inside an RDG pass on the mobile renderer.

#rnx

#Synced-CL 50356853

diff --git a/Engine/Source/Runtime/Renderer/Private/DecalRenderingShared.h b/Engine/Source/Runtime/Renderer/Private/DecalRenderingShared.h
--- a/Engine/Source/Runtime/Renderer/Private/DecalRenderingShared.h
+++ b/Engine/Source/Runtime/Renderer/Private/DecalRenderingShared.h
@@ -144,17 +144,19 @@ private:
 
 struct FDecalViewInfo
 {
-	FDecalViewInfo(int32 InViewOffset, int32 InViewIndex, EDecalRenderStage InRenderStage, EDecalRenderTargetMode InRenderTargetMode) :
+	FDecalViewInfo(int32 InViewOffset, int32 InViewIndex, EDecalRenderStage InRenderStage, EDecalRenderTargetMode InRenderTargetMode, const TConstArrayView<const FVisibleDecal*>& InRelevantDeferredDecals) :
 		ViewOffset(InViewOffset),
 		ViewIndex(InViewIndex),
 		RenderStage(InRenderStage),
-		RenderTargetMode(InRenderTargetMode)
+		RenderTargetMode(InRenderTargetMode),
+		RelevantDeferredDecals(InRelevantDeferredDecals)
 	{
 	}
 	int32 ViewOffset;
 	int32 ViewIndex;
 	EDecalRenderStage RenderStage;
 	EDecalRenderTargetMode RenderTargetMode;
+	TConstArrayView<const FVisibleDecal*> RelevantDeferredDecals;
 };
 /**
  * Shared deferred decal functionality.
diff --git a/Engine/Source/Runtime/Renderer/Private/MobileDecalRendering.cpp b/Engine/Source/Runtime/Renderer/Private/MobileDecalRendering.cpp
--- a/Engine/Source/Runtime/Renderer/Private/MobileDecalRendering.cpp
+++ b/Engine/Source/Runtime/Renderer/Private/MobileDecalRendering.cpp
@@ -45,6 +45,11 @@ static bool DoesPlatformSupportDecals(EShaderPlatform ShaderPlatform)
 	return true;
 }
 
+TConstArrayView<const FVisibleDecal*> FMobileSceneRenderer::GetRelevantDeferredDecals(int32 DecalViewOffset, int32 ViewIndex, EDecalRenderStage DecalRenderStage)
+{
+	return DecalVisibilityTaskData ? DecalVisibilityTaskData->FinishRelevantDecals(DecalViewOffset, ViewIndex, DecalRenderStage) : TConstArrayView<const FVisibleDecal*>();
+}
+
 void FMobileSceneRenderer::RenderDecals(FRHICommandList& RHICmdList, FViewInfo& View, const FDecalViewInfo& DecalViewInfo, const FInstanceCullingDrawParams* InstanceCullingDrawParams)
 {
 	if (!DoesPlatformSupportDecals(View.GetShaderPlatform()) || !View.Family->EngineShowFlags.Decals || View.bIsPlanarReflection)
@@ -64,7 +69,7 @@ void FMobileSceneRenderer::RenderDecals(FRHICommandList& RHICmdList, FViewInfo&
 	}
 
 	// Deferred decals
-	if (DecalVisibilityTaskData != nullptr)
+	if (DecalViewInfo.RelevantDeferredDecals.Num() > 0)
 	{
 		SCOPED_DRAW_EVENT(RHICmdList, Decals);
 		RenderDeferredDecalsMobile(RHICmdList, View, DecalViewInfo);
@@ -73,10 +78,9 @@ void FMobileSceneRenderer::RenderDecals(FRHICommandList& RHICmdList, FViewInfo&
 
 void FMobileSceneRenderer::RenderDeferredDecalsMobile(FRHICommandList& RHICmdList, const FViewInfo& View, const FDecalViewInfo& DecalViewInfo)
 {
-	TConstArrayView<const FVisibleDecal*> RelevantDecals = DecalVisibilityTaskData->FinishRelevantDecals(DecalViewInfo.ViewOffset, DecalViewInfo.ViewIndex, DecalViewInfo.RenderStage);
 	FRelevantDecalList SortedDecals;
 
-	for (const FVisibleDecal* VisibleDecal : RelevantDecals)
+	for (const FVisibleDecal* VisibleDecal : DecalViewInfo.RelevantDeferredDecals)
 	{
 		if (DecalViewInfo.RenderTargetMode == DecalRendering::GetRenderTargetMode(View.GetShaderPlatform(), VisibleDecal->BlendDesc, DecalViewInfo.RenderStage))
 		{
diff --git a/Engine/Source/Runtime/Renderer/Private/MobileShadingRenderer.cpp b/Engine/Source/Runtime/Renderer/Private/MobileShadingRenderer.cpp
--- a/Engine/Source/Runtime/Renderer/Private/MobileShadingRenderer.cpp
+++ b/Engine/Source/Runtime/Renderer/Private/MobileShadingRenderer.cpp
@@ -1023,9 +1023,11 @@ void FMobileSceneRenderer::RenderCustomRenderPassBasePass(FRDGBuilder& GraphBuil
 			}
 		}
 
+		EDecalRenderStage DecalStage = bRequiresDBufferDecals ? EDecalRenderStage::Emissive : EDecalRenderStage::Mobile;
 		FDecalViewInfo DecalViewInfo = FDecalViewInfo(ViewContext.DecalViewOffset, ViewContext.ViewIndex,
-			bRequiresDBufferDecals ? EDecalRenderStage::Emissive : EDecalRenderStage::Mobile,
-			bDeferredShading ? EDecalRenderTargetMode::SceneColorAndGBuffer : EDecalRenderTargetMode::SceneColor);
+													  DecalStage,
+													  bDeferredShading ? EDecalRenderTargetMode::SceneColorAndGBuffer : EDecalRenderTargetMode::SceneColor,
+													  GetRelevantDeferredDecals(ViewContext.DecalViewOffset, ViewContext.ViewIndex, DecalStage));
 
 		GraphBuilder.AddPass(
 			RDG_EVENT_NAME("RenderMobileBasePass"),
@@ -2093,9 +2095,11 @@ void FMobileSceneRenderer::RenderForwardSinglePass(FRDGBuilder& GraphBuilder, FM
 	const bool bDoOcclusionInMainPass = bDoOcclusionQueries && !bVulkanAdrenoOcclusionMode;
 	PassParameters->RenderTargets.NumOcclusionQueries = bVulkanAdrenoOcclusionMode ? 0u : NumOcclusionQueries;
 	
+	EDecalRenderStage DecalStage = bRequiresDBufferDecals ? EDecalRenderStage::Emissive : EDecalRenderStage::Mobile;
 	FDecalViewInfo DecalViewInfo = FDecalViewInfo(ViewContext.DecalViewOffset, ViewContext.ViewIndex,
-		bRequiresDBufferDecals ? EDecalRenderStage::Emissive : EDecalRenderStage::Mobile,
-		EDecalRenderTargetMode::SceneColor);
+												  DecalStage,
+												  EDecalRenderTargetMode::SceneColor,
+												  GetRelevantDeferredDecals(ViewContext.DecalViewOffset, ViewContext.ViewIndex, DecalStage));
 
 	GraphBuilder.AddPass(
 		RDG_EVENT_NAME("SceneColorRendering"),
@@ -2213,10 +2217,11 @@ void FMobileSceneRenderer::RenderForwardMultiPass(FRDGBuilder& GraphBuilder, FMo
 		BuildMeshRenderingCommands(GraphBuilder, EMeshPass::DebugViewMode, View, Scene->GPUScene, InstanceCullingManager, ParameterCollection->DebugViewModeInstanceCullingDrawParams);
 	}
 
-	FDecalViewInfo DecalViewInfo = FDecalViewInfo(ViewContext.DecalViewOffset,
-												  ViewContext.ViewIndex,
-												  bRequiresDBufferDecals ? EDecalRenderStage::Emissive : EDecalRenderStage::Mobile,
-												  EDecalRenderTargetMode::SceneColor);
+	EDecalRenderStage DecalStage = bRequiresDBufferDecals ? EDecalRenderStage::Emissive : EDecalRenderStage::Mobile;
+	FDecalViewInfo DecalViewInfo = FDecalViewInfo(ViewContext.DecalViewOffset, ViewContext.ViewIndex, 
+												  DecalStage,
+												  EDecalRenderTargetMode::SceneColor,
+												  GetRelevantDeferredDecals(ViewContext.DecalViewOffset, ViewContext.ViewIndex, DecalStage));
 
 	GraphBuilder.AddPass(
 		RDG_EVENT_NAME("SceneColorRendering"),
@@ -2545,10 +2550,11 @@ void FMobileSceneRenderer::RenderDeferredSinglePass(FRDGBuilder& GraphBuilder, F
 			BuildMeshRenderingCommands(GraphBuilder, StandardTranslucencyMeshPass, View, Scene->GPUScene, InstanceCullingManager, ParameterCollection->TranslucencyInstanceCullingDrawParams);
 		}
 
-		FDecalViewInfo DecalViewInfo = FDecalViewInfo(ViewContext.DecalViewOffset,
-													  ViewContext.ViewIndex,
-													  bRequiresDBufferDecals ? EDecalRenderStage::Emissive : EDecalRenderStage::MobileBeforeLighting,
-													  EDecalRenderTargetMode::SceneColorAndGBuffer);
+		EDecalRenderStage DecalStage = bRequiresDBufferDecals ? EDecalRenderStage::Emissive : EDecalRenderStage::MobileBeforeLighting;
+		FDecalViewInfo DecalViewInfo = FDecalViewInfo(ViewContext.DecalViewOffset, ViewContext.ViewIndex,
+													  DecalStage,
+													  EDecalRenderTargetMode::SceneColorAndGBuffer,
+													  GetRelevantDeferredDecals(ViewContext.DecalViewOffset, ViewContext.ViewIndex, DecalStage));
 
 		GraphBuilder.AddPass(
 			RDG_EVENT_NAME("SceneColorRendering"),
@@ -2614,13 +2620,13 @@ void FMobileSceneRenderer::RenderDeferredMultiPassDecals(FRDGBuilder& GraphBuild
 	const EDecalRenderTargetMode DecalRenderTargetModes[] = { EDecalRenderTargetMode::SceneColorAndGBuffer, EDecalRenderTargetMode::SceneColorAndGBufferNoNormal };
 	const EDecalRenderStage DecalRenderStage = bRequiresDBufferDecals ? EDecalRenderStage::Emissive : EDecalRenderStage::MobileBeforeLighting;
 	constexpr const TCHAR* PassNames[2] = { TEXT("SceneColorAndGBuffer"), TEXT("SceneColorAndGBufferNoNormal")};
-	TConstArrayView<const FVisibleDecal*> RelevantDecals = DecalVisibilityTaskData ? DecalVisibilityTaskData->FinishRelevantDecals(DecalViewOffset, ViewIndex, DecalRenderStage) : TConstArrayView<const FVisibleDecal*>();
+	TConstArrayView<const FVisibleDecal*> RelevantDeferredDecals = GetRelevantDeferredDecals(DecalViewOffset, ViewIndex, DecalRenderStage);
 
 	for (const EDecalRenderTargetMode DecalRenderTargetMode : DecalRenderTargetModes)
 	{
 		EMeshPass::Type MeshDecalsPass = DecalRendering::GetMeshPassType(DecalRenderTargetMode);
 
-		const FVisibleDecal* const* FirstDecal = RelevantDecals.FindByPredicate( [&View, &DecalRenderTargetMode, &DecalRenderStage](const FVisibleDecal* Decal) { return DecalRenderTargetMode == DecalRendering::GetRenderTargetMode(View.GetShaderPlatform(), Decal->BlendDesc, DecalRenderStage); });
+		const FVisibleDecal* const* FirstDecal = RelevantDeferredDecals.FindByPredicate( [&View, &DecalRenderTargetMode, &DecalRenderStage](const FVisibleDecal* Decal) { return DecalRenderTargetMode == DecalRendering::GetRenderTargetMode(View.GetShaderPlatform(), Decal->BlendDesc, DecalRenderStage); });
 		if (FirstDecal == nullptr && !HasAnyDraw(View.ParallelMeshDrawCommandPasses[MeshDecalsPass]))
 		{
 			continue;
@@ -2657,7 +2663,11 @@ void FMobileSceneRenderer::RenderDeferredMultiPassDecals(FRDGBuilder& GraphBuild
 			BuildMeshRenderingCommands(GraphBuilder, MeshDecalsPass, View, Scene->GPUScene, InstanceCullingManager, PassParameters->InstanceCullingDrawParams);
 		}
 
-		FDecalViewInfo DecalViewInfo = FDecalViewInfo(DecalViewOffset, ViewIndex, DecalRenderStage, DecalRenderTargetMode);
+		FDecalViewInfo DecalViewInfo = FDecalViewInfo(DecalViewOffset, 
+			ViewIndex,
+			DecalRenderStage,
+			DecalRenderTargetMode,
+			RelevantDeferredDecals);
 
 		GraphBuilder.AddPass(
 			RDG_EVENT_NAME("Decals_%s", DecalRendering::GetRenderTargetModeName(DecalRenderTargetMode)),
diff --git a/Engine/Source/Runtime/Renderer/Private/SceneRendering.h b/Engine/Source/Runtime/Renderer/Private/SceneRendering.h
--- a/Engine/Source/Runtime/Renderer/Private/SceneRendering.h
+++ b/Engine/Source/Runtime/Renderer/Private/SceneRendering.h
@@ -2919,6 +2919,7 @@ protected:
 
 
 	/** Renders decals. */
+	TConstArrayView<const FVisibleDecal*> GetRelevantDeferredDecals(int32 DecalViewOffset, int32 ViewIndex, EDecalRenderStage DecalRenderStage);
 	void RenderDecals(FRHICommandList& RHICmdList, FViewInfo& View, const FDecalViewInfo& DecalViewInfo, const FInstanceCullingDrawParams* InstanceCullingDrawParams);
 	void RenderDeferredDecalsMobile(FRHICommandList& RHICmdList, const FViewInfo& View, const FDecalViewInfo& DecalViewInfo);
 	void RenderDBuffer(FRDGBuilder& GraphBuilder, TArray<FViewInfo>& InViews, FSceneTextures& SceneTextures, FDBufferTextures& DBufferTextures, FInstanceCullingManager& InstanceCullingManager);
