[HWRT] Check if ray tracing shaders are compiled for the project to avoid a crash when switching to preview mode of platforms that do not enable/support them.

#rnx

#Synced-CL 50366487

diff --git a/Engine/Source/Runtime/Renderer/Private/DeferredShadingRenderer.cpp b/Engine/Source/Runtime/Renderer/Private/DeferredShadingRenderer.cpp
--- a/Engine/Source/Runtime/Renderer/Private/DeferredShadingRenderer.cpp
+++ b/Engine/Source/Runtime/Renderer/Private/DeferredShadingRenderer.cpp
@@ -4094,9 +4094,10 @@ static bool AnyRayTracingPassEnabled(const FScene* Scene, const FViewInfo& View,
 		|| MegaLights::UseHardwareRayTracing(*View.Family);
 }
 
-static bool ShouldRenderRayTracingEffectInternal(bool bEffectEnabled, ERayTracingPipelineCompatibilityFlags CompatibilityFlags)
+static bool ShouldRenderRayTracingEffectInternal(bool bEffectEnabled, ERayTracingPipelineCompatibilityFlags CompatibilityFlags, EShaderPlatform ShaderPlatform)
 {
-	const bool bAllowPipeline = GRHISupportsRayTracingShaders && 
+	const bool bAllowPipeline = GRHISupportsRayTracingShaders &&
+								ShouldCompileRayTracingShadersForProject(ShaderPlatform) && 
 								CVarRayTracingAllowPipeline.GetValueOnRenderThread() &&
 								EnumHasAnyFlags(CompatibilityFlags, ERayTracingPipelineCompatibilityFlags::FullPipeline);
 
@@ -4122,14 +4123,31 @@ static bool ShouldRenderRayTracingEffectInternal(bool bEffectEnabled, ERayTracin
 	}
 }
 
+// Most ray tracing effects can be enabled or disabled per view, but the ray tracing sky light effect specifically requires base pass shaders
+// in the FScene to be configured differently, and thus can't work if ray tracing is disabled.  There is logic in FScene::Update where
+// bCachedShouldRenderSkylightInBasePass is updated based on the result of ShouldRenderSkylightInBasePass(), which is affected by whether sky light
+// ray tracing is enabled.  When this value changes, bScenesPrimitivesNeedStaticMeshElementUpdate is set to true, forcing a rebuild of all static mesh
+// elements in the scene.  This can't be done per frame (never mind per view), which would be required to allow this setting to vary, at least with
+// the current implementation.  Sky light ray tracing is often used for cinematic capture, and not in games, so hopefully this isn't a big limitation.
+
+bool ShouldRenderRayTracingEffect(bool bEffectEnabled, ERayTracingPipelineCompatibilityFlags CompatibilityFlags, EShaderPlatform ShaderPlatform)
+{
+	if (!IsRayTracingEnabled(ShaderPlatform))
+	{
+		return false;
+	}
+
+	return ShouldRenderRayTracingEffectInternal(bEffectEnabled, CompatibilityFlags, ShaderPlatform);
+}
+
 bool ShouldRenderRayTracingEffect(bool bEffectEnabled, ERayTracingPipelineCompatibilityFlags CompatibilityFlags, const FSceneView& View)
 {
-	if (!IsRayTracingEnabled(View.GetShaderPlatform()) || !View.IsRayTracingAllowedForView())
+	if (!View.IsRayTracingAllowedForView())
 	{
 		return false;
 	}
 
-	return ShouldRenderRayTracingEffectInternal(bEffectEnabled, CompatibilityFlags);
+	return ShouldRenderRayTracingEffect(bEffectEnabled, CompatibilityFlags, View.GetShaderPlatform());
 }
 
 bool ShouldRenderRayTracingEffect(bool bEffectEnabled, ERayTracingPipelineCompatibilityFlags CompatibilityFlags, const FSceneViewFamily& ViewFamily)
@@ -4141,20 +4159,6 @@ bool ShouldRenderRayTracingEffect(bool bEffectEnabled, ERayTracingPipelineCompat
 	return ShouldRenderRayTracingEffect(bEffectEnabled, CompatibilityFlags, *ViewFamily.Views[0]);
 }
 
-// Most ray tracing effects can be enabled or disabled per view, but the ray tracing sky light effect specifically requires base pass shaders
-// in the FScene to be configured differently, and thus can't work if ray tracing is disabled.  There is logic in FScene::Update where
-// bCachedShouldRenderSkylightInBasePass is updated based on the result of ShouldRenderSkylightInBasePass(), which is affected by whether sky light
-// ray tracing is enabled.  When this value changes, bScenesPrimitivesNeedStaticMeshElementUpdate is set to true, forcing a rebuild of all static mesh
-// elements in the scene.  This can't be done per frame (never mind per view), which would be required to allow this setting to vary, at least with
-// the current implementation.  Sky light ray tracing is often used for cinematic capture, and not in games, so hopefully this isn't a big limitation.
-// 
-// This forces ray tracing on, but other ray tracing features are still disabled.  This is its own function to allow ShouldRenderRayTracingEffectInternal
-// to be kept private, as all other effects should provide a view or view family, to allow IsRayTracingAllowedForView to be tested.
-bool ShouldRenderRayTracingSkyLightEffect()
-{
-	return ShouldRenderRayTracingEffectInternal(true, ERayTracingPipelineCompatibilityFlags::FullPipeline);
-}
-
 bool HasRaytracingDebugViewModeRaytracedOverlay(const FSceneViewFamily& ViewFamily);
 bool HasRayTracedOverlay(const FSceneViewFamily& ViewFamily)
 {
diff --git a/Engine/Source/Runtime/Renderer/Private/RayTracing/RaytracingOptions.h b/Engine/Source/Runtime/Renderer/Private/RayTracing/RaytracingOptions.h
--- a/Engine/Source/Runtime/Renderer/Private/RayTracing/RaytracingOptions.h
+++ b/Engine/Source/Runtime/Renderer/Private/RayTracing/RaytracingOptions.h
@@ -76,9 +76,9 @@ namespace RayTracing
 #if RHI_RAYTRACING
 
 // Whether a particular effect should be used, taking into account debug override
+extern bool ShouldRenderRayTracingEffect(bool bEffectEnabled, ERayTracingPipelineCompatibilityFlags CompatibilityFlags, EShaderPlatform ShaderPlatform);
 extern bool ShouldRenderRayTracingEffect(bool bEffectEnabled, ERayTracingPipelineCompatibilityFlags CompatibilityFlags, const FSceneView& View);
 extern bool ShouldRenderRayTracingEffect(bool bEffectEnabled, ERayTracingPipelineCompatibilityFlags CompatibilityFlags, const FSceneViewFamily& ViewFamily);
-extern bool ShouldRenderRayTracingSkyLightEffect();
 
 extern bool AnyInlineRayTracingPassEnabled(const FScene* Scene, const FViewInfo& View);
 extern FRayTracingPrimaryRaysOptions GetRayTracingTranslucencyOptions(const FViewInfo& View);
@@ -109,17 +109,17 @@ extern bool CanUseRayTracingAMDHitToken();
 
 #else // RHI_RAYTRACING
 
-FORCEINLINE bool ShouldRenderRayTracingEffect(bool bEffectEnabled, ERayTracingPipelineCompatibilityFlags CompatibilityFlags, const FSceneView& View)
+FORCEINLINE bool ShouldRenderRayTracingEffect(bool bEffectEnabled, ERayTracingPipelineCompatibilityFlags CompatibilityFlags, EShaderPlatform ShaderPlatform)
 {
 	return false;
 }
 
-FORCEINLINE bool ShouldRenderRayTracingEffect(bool bEffectEnabled, ERayTracingPipelineCompatibilityFlags CompatibilityFlags, const FSceneViewFamily& ViewFamily)
+FORCEINLINE bool ShouldRenderRayTracingEffect(bool bEffectEnabled, ERayTracingPipelineCompatibilityFlags CompatibilityFlags, const FSceneView& View)
 {
 	return false;
 }
 
-FORCEINLINE bool ShouldRenderRayTracingSkyLightEffect()
+FORCEINLINE bool ShouldRenderRayTracingEffect(bool bEffectEnabled, ERayTracingPipelineCompatibilityFlags CompatibilityFlags, const FSceneViewFamily& ViewFamily)
 {
 	return false;
 }
diff --git a/Engine/Source/Runtime/Renderer/Private/RayTracing/RaytracingSkylight.cpp b/Engine/Source/Runtime/Renderer/Private/RayTracing/RaytracingSkylight.cpp
--- a/Engine/Source/Runtime/Renderer/Private/RayTracing/RaytracingSkylight.cpp
+++ b/Engine/Source/Runtime/Renderer/Private/RayTracing/RaytracingSkylight.cpp
@@ -139,7 +139,7 @@ bool ShouldRenderRayTracingSkyLight(const FSkyLightSceneProxy* SkyLightSceneProx
 	bool bRayTracingSkyEnabled = (GRayTracingSkyLight  > 0 && SkyLightSceneProxy->CastRayTracedShadow == ECastRayTracedShadow::UseProjectSetting)
 								||  SkyLightSceneProxy->CastRayTracedShadow == ECastRayTracedShadow::Enabled;
 
-	return bRayTracingSkyEnabled && ShouldRenderRayTracingSkyLightEffect() && (GetSkyLightSamplesPerPixel(SkyLightSceneProxy) > 0);
+	return bRayTracingSkyEnabled && ShouldRenderRayTracingEffect(true, ERayTracingPipelineCompatibilityFlags::FullPipeline, ShaderPlatform) && (GetSkyLightSamplesPerPixel(SkyLightSceneProxy) > 0);
 }
 
 IMPLEMENT_GLOBAL_SHADER_PARAMETER_STRUCT(FSkyLightData, "SkyLight");
