CinematicPrestream: Deal with case where virtual texture has a lod bias applied during the cook.
When this happens using the recorded page request directly will produce incorrect page IDs some of which are outside of the playback allocated VT space.
We deal with this by stroring the original VT size in the page request and comparing with the runtime size. Then we adjust the page request based on this.
#rnx

#Synced-CL 50484185

diff --git a/Engine/Source/Runtime/Renderer/Private/VT/VirtualTextureSystem.cpp b/Engine/Source/Runtime/Renderer/Private/VT/VirtualTextureSystem.cpp
--- a/Engine/Source/Runtime/Renderer/Private/VT/VirtualTextureSystem.cpp
+++ b/Engine/Source/Runtime/Renderer/Private/VT/VirtualTextureSystem.cpp
@@ -3382,9 +3382,12 @@ void FVirtualTextureSystem::RecordPageRequests(FUniquePageList const* UniquePage
 			continue;
 		}
 
-		const uint32 LocalAddress = vAddress - AllocatedVT->GetVirtualAddress();
+		// PersistentHash is used to find associated AllocatedVT at playback.
 		const uint32 PersistentHash = AllocatedVT->GetPersistentHash();
-		const uint64 ExportPageId = ((uint64)PersistentHash << 32)| ((uint64)vLevelPlusOne << 28) | (uint64)LocalAddress;
+		// Log2TileCountX is used at playback to detect and adjust for any TextureLODGroup bias that has been applied.
+		const uint32 Log2TileCountX = FMath::CeilLogTwo(AllocatedVT->GetBlockWidthInTiles()) & 0xf;
+		const uint32 LocalAddress = (vAddress - AllocatedVT->GetVirtualAddress()) & 0xffffff;
+		const uint64 ExportPageId = ((uint64)PersistentHash << 32) | ((uint64)vLevelPlusOne << 28) | ((uint64)Log2TileCountX << 24) | (uint64)LocalAddress;
 		
 		OutPages.Add(ExportPageId);
 	}
@@ -3427,7 +3430,7 @@ void FVirtualTextureSystem::AddRequestedTilesTask(const FAddRequestedTilesParame
 	FUniquePageList* RESTRICT RequestedPageList = Parameters.UniquePageList;
 	const uint64* RESTRICT Buffer = Parameters.RequestBuffer;
 	const uint32 BufferSize = Parameters.NumRequests;
-	const uint32 LevelBias = Parameters.LevelBias;
+	const uint32 ViewLevelBias = Parameters.LevelBias;
 
 	uint32 CurrentPersistentHash = ~0u;
 	IAllocatedVirtualTexture* AllocatedVT = nullptr;
@@ -3448,17 +3451,26 @@ void FVirtualTextureSystem::AddRequestedTilesTask(const FAddRequestedTilesParame
 
 		if (AllocatedVT != nullptr && AllocatedVT->GetSpaceID() < VIRTUALTEXTURE_MAX_FEEDBACK_SPACES)
 		{
-			const uint32 LevelPlusOne = (uint32)(PageRequest >> 28) & 0xf;
-			const uint32 LocalAddress = (uint32)(PageRequest & 0xffffff);
+			const uint32 SourceLevelPlusOne = (uint32)(PageRequest >> 28) & 0xf;
+			const uint32 SourceLevel = FMath::Max(SourceLevelPlusOne, 1u) - 1u;
+			const uint32 SourceLog2TileCountX = (uint32)(PageRequest >> 24) & 0xf;
+			const uint32 SourceLocalAddress = (uint32)(PageRequest & 0xffffff);
 
+			// Deal with the case of texture size being reduced from recording to playback due to LOD bias in cook.
+			const uint32 Log2TileCountX = FMath::CeilLogTwo(AllocatedVT->GetBlockWidthInTiles());
+			const uint32 TextureSizeBias = SourceLog2TileCountX > Log2TileCountX ? SourceLog2TileCountX - Log2TileCountX : 0;
+
+			const uint32 LocalAddress = SourceLocalAddress >> (TextureSizeBias * 2u);
 			const uint32 BaseAddress = AllocatedVT->GetVirtualAddress();
 			const uint32 Address = BaseAddress + LocalAddress;
-			const uint32 BaseLevel = FMath::Max(LevelPlusOne, 1u) - 1;
+
 			const uint32 MaxLevel = AllocatedVT->GetMaxLevel();
-			const uint32 Level = FMath::Min(BaseLevel + LevelBias, MaxLevel);
+			const uint32 Level = FMath::Clamp((int32)SourceLevel - (int32)TextureSizeBias + (int32)ViewLevelBias, 0, MaxLevel);
+
 			const uint32 PageX = FMath::ReverseMortonCode2(Address) >> Level;
 			const uint32 PageY = FMath::ReverseMortonCode2(Address >> 1) >> Level;
 			const uint32 SpaceId = AllocatedVT->GetSpaceID();
+
 			const uint32 PageId = PageX | (PageY << 12) | ((Level+1) << 24) | (SpaceId << 28);
 
 			RequestedPageList->Add(PageId, 0xffff);
