Prevent calling FDecalVisibilityTaskData::FinishRelevantDecals when DecalViewOffset is invalid to avoid accessing an invalid TMap.

#rnx

#Synced-CL 50479750

diff --git a/Engine/Source/Runtime/Renderer/Private/MobileDecalRendering.cpp b/Engine/Source/Runtime/Renderer/Private/MobileDecalRendering.cpp
--- a/Engine/Source/Runtime/Renderer/Private/MobileDecalRendering.cpp
+++ b/Engine/Source/Runtime/Renderer/Private/MobileDecalRendering.cpp
@@ -47,7 +47,12 @@ static bool DoesPlatformSupportDecals(EShaderPlatform ShaderPlatform)
 
 TConstArrayView<const FVisibleDecal*> FMobileSceneRenderer::GetRelevantDeferredDecals(int32 DecalViewOffset, int32 ViewIndex, EDecalRenderStage DecalRenderStage)
 {
-	return DecalVisibilityTaskData ? DecalVisibilityTaskData->FinishRelevantDecals(DecalViewOffset, ViewIndex, DecalRenderStage) : TConstArrayView<const FVisibleDecal*>();
+	if (DecalVisibilityTaskData != nullptr && DecalViewOffset != INDEX_NONE)
+	{
+		return DecalVisibilityTaskData->FinishRelevantDecals(DecalViewOffset, ViewIndex, DecalRenderStage);
+	}
+
+	return TConstArrayView<const FVisibleDecal*>();
 }
 
 void FMobileSceneRenderer::RenderDecals(FRHICommandList& RHICmdList, FViewInfo& View, const FDecalViewInfo& DecalViewInfo, const FInstanceCullingDrawParams* InstanceCullingDrawParams)
