Cinematic Prestreaming: Apply an automatic texture bias to the prestreaming data to account for any difference between capture and playback resolution.
Previously we used a manual cvar "r.VT.PlaybackMipBias" to account for this.
The cvar remains but is now only intended to be used as an additional debug tool.

#Synced-CL 50505473

diff --git a/Engine/Plugins/Experimental/CinematicPrestreaming/Source/CinematicPrestreaming/Private/TrackInstances/MovieSceneCinePrestreamingTrackInstance.cpp b/Engine/Plugins/Experimental/CinematicPrestreaming/Source/CinematicPrestreaming/Private/TrackInstances/MovieSceneCinePrestreamingTrackInstance.cpp
--- a/Engine/Plugins/Experimental/CinematicPrestreaming/Source/CinematicPrestreaming/Private/TrackInstances/MovieSceneCinePrestreamingTrackInstance.cpp
+++ b/Engine/Plugins/Experimental/CinematicPrestreaming/Source/CinematicPrestreaming/Private/TrackInstances/MovieSceneCinePrestreamingTrackInstance.cpp
@@ -58,6 +58,7 @@ void UMovieSceneCinePrestreamingTrackInstance::OnAnimate()
 	// Gather active prestreaming sections.
 	FCinePrestreamingVTData AccumulatedVirtualTextureData;
 	FCinePrestreamingNaniteData AccumulatedNaniteData;
+	FIntPoint RecordedResolution = FIntPoint::ZeroValue;
 
 	for (const FMovieSceneTrackInstanceInput& Input : GetInputs())
 	{
@@ -123,14 +124,18 @@ void UMovieSceneCinePrestreamingTrackInstance::OnAnimate()
 		// todo: Fetch the data from some time range of keyframes?
 		AccumulatedVirtualTextureData.PageIds.Append(PrestreamingData->VirtualTextureDatas[Index].PageIds);
 		AccumulatedNaniteData.RequestData.Append(PrestreamingData->NaniteDatas[Index].RequestData);
+
+		// Should be safe to assume that if we have multiple tracks then they were recorded at the same resolution.
+		ensure(RecordedResolution == FIntPoint::ZeroValue || RecordedResolution == PrestreamingData->RecordedResolution);
+		RecordedResolution = PrestreamingData->RecordedResolution;
 	}
 
 	// Submit virtual texture prestreaming data to the renderer.
 	if (AccumulatedVirtualTextureData.PageIds.Num() > 0)
 	{
-		ENQUEUE_RENDER_COMMAND(RequestTiles)([PageIds = MoveTemp(AccumulatedVirtualTextureData.PageIds)](FRHICommandListImmediate& RHICmdList) mutable
+		ENQUEUE_RENDER_COMMAND(RequestTiles)([PageIds = MoveTemp(AccumulatedVirtualTextureData.PageIds), RecordedResolution](FRHICommandListImmediate& RHICmdList) mutable
 		{
-			GetRendererModule().RequestVirtualTextureTiles(MoveTemp(PageIds));
+			GetRendererModule().RequestVirtualTextureTiles(MoveTemp(PageIds), RecordedResolution);
 		});
 	}
 
diff --git a/Engine/Source/Runtime/RenderCore/Public/RendererInterface.h b/Engine/Source/Runtime/RenderCore/Public/RendererInterface.h
--- a/Engine/Source/Runtime/RenderCore/Public/RendererInterface.h
+++ b/Engine/Source/Runtime/RenderCore/Public/RendererInterface.h
@@ -862,8 +862,11 @@ public:
 	virtual void SetVirtualTextureRequestRecordBuffer(uint64 Handle) = 0;
 	/** Fetch the virtual texture page requests recorded since the last call to SetVirtualTextureRequestRecordBuffer. Returns the handle that was passed in. */
 	virtual uint64 GetVirtualTextureRequestRecordBuffer(TSet<uint64>& OutPageRequests) = 0;
-	/**	Request an array of virtual texture page requests that was captured with SetVirtualTextureRequestRecordBuffer. Note that the array will be moved and ownership is taken. */
-	virtual void RequestVirtualTextureTiles(TArray<uint64>&& InPageRequests) = 0;
+	/**	
+	 * Request an array of virtual texture page requests that was captured with SetVirtualTextureRequestRecordBuffer. Note that the array will be moved and ownership is taken.
+	 * Provide an optional InCaptureResolution so that a bias can be applied to account for the current playback resolution if necessary.
+	 */
+	virtual void RequestVirtualTextureTiles(TArray<uint64>&& InPageRequests, FIntPoint const& InCaptureResolution = FIntPoint::ZeroValue) = 0;
 
 	/** Evict all data from virtual texture caches. */
 	virtual void FlushVirtualTextureCache() = 0;
diff --git a/Engine/Source/Runtime/Renderer/Private/DeferredShadingRenderer.cpp b/Engine/Source/Runtime/Renderer/Private/DeferredShadingRenderer.cpp
--- a/Engine/Source/Runtime/Renderer/Private/DeferredShadingRenderer.cpp
+++ b/Engine/Source/Runtime/Renderer/Private/DeferredShadingRenderer.cpp
@@ -1797,9 +1797,7 @@ void FDeferredShadingSceneRenderer::Render(FRDGBuilder& GraphBuilder, const FSce
 	// Virtual texturing isn't needed for depth prepass
 	if (bUseVirtualTexturing && RendererOutput != ERendererOutput::DepthPrepassOnly)
 	{
-		FVirtualTextureUpdateSettings Settings;
-		Settings.EnableThrottling(!ViewFamily.bOverrideVirtualTextureThrottle);
-
+		FVirtualTextureUpdateSettings Settings(ViewFamily);
 		VirtualTextureUpdater = FVirtualTextureSystem::Get().BeginUpdate(GraphBuilder, FeatureLevel, this, Settings);
 		VirtualTextureFeedbackBegin(GraphBuilder, Views, GetActiveSceneTexturesConfig().Extent);
 	}
diff --git a/Engine/Source/Runtime/Renderer/Private/MobileShadingRenderer.cpp b/Engine/Source/Runtime/Renderer/Private/MobileShadingRenderer.cpp
--- a/Engine/Source/Runtime/Renderer/Private/MobileShadingRenderer.cpp
+++ b/Engine/Source/Runtime/Renderer/Private/MobileShadingRenderer.cpp
@@ -1158,9 +1158,7 @@ void FMobileSceneRenderer::Render(FRDGBuilder& GraphBuilder, const FSceneRenderU
 
 	if (bUseVirtualTexturing)
 	{
-		FVirtualTextureUpdateSettings Settings;
-		Settings.EnableThrottling(!ViewFamily.bOverrideVirtualTextureThrottle);
-
+		FVirtualTextureUpdateSettings Settings(ViewFamily);
 		VirtualTextureUpdater = FVirtualTextureSystem::Get().BeginUpdate(GraphBuilder, FeatureLevel, this, Settings);
 		VirtualTextureFeedbackBegin(GraphBuilder, Views, SceneTexturesConfig.Extent);
 	}
diff --git a/Engine/Source/Runtime/Renderer/Private/RendererModule.h b/Engine/Source/Runtime/Renderer/Private/RendererModule.h
--- a/Engine/Source/Runtime/Renderer/Private/RendererModule.h
+++ b/Engine/Source/Runtime/Renderer/Private/RendererModule.h
@@ -110,7 +110,7 @@ public:
 	virtual void AddVirtualTextureProducerDestroyedCallback(const FVirtualTextureProducerHandle& Handle, FVTProducerDestroyedFunction* Function, void* Baton) override;
 	virtual uint32 RemoveAllVirtualTextureProducerDestroyedCallbacks(const void* Baton) override;
 	virtual void ReleaseVirtualTexturePendingResources() override;
-	virtual void RequestVirtualTextureTiles(TArray<uint64>&& InPageRequests) override;
+	virtual void RequestVirtualTextureTiles(TArray<uint64>&& InPageRequests, FIntPoint const& InCaptureResolution) override;
 	virtual void RequestVirtualTextureTiles(const FVector2D& InScreenSpaceSize, int32 InMipLevel) override;
 	virtual void RequestVirtualTextureTiles(const FMaterialRenderProxy* InMaterialRenderProxy, const FVector2D& InScreenSpaceSize, ERHIFeatureLevel::Type InFeatureLevel) override;
 	virtual void RequestVirtualTextureTiles(IAllocatedVirtualTexture* AllocatedVT, const FVector2D& InScreenSpaceSize, const FVector2D& InViewportPosition, const FVector2D& InViewportSize, const FVector2D& InUV0, const FVector2D& InUV1, int32 InMipLevel) override;
diff --git a/Engine/Source/Runtime/Renderer/Private/SceneRendering.cpp b/Engine/Source/Runtime/Renderer/Private/SceneRendering.cpp
--- a/Engine/Source/Runtime/Renderer/Private/SceneRendering.cpp
+++ b/Engine/Source/Runtime/Renderer/Private/SceneRendering.cpp
@@ -5690,9 +5690,9 @@ uint64 FRendererModule::GetVirtualTextureRequestRecordBuffer(TSet<uint64>& OutPa
 #endif
 }
 
-void FRendererModule::RequestVirtualTextureTiles(TArray<uint64>&& InPageRequests)
+void FRendererModule::RequestVirtualTextureTiles(TArray<uint64>&& InPageRequests, FIntPoint const& InCaptureResolution)
 {
-	FVirtualTextureSystem::Get().RequestRecordedTiles(MoveTemp(InPageRequests));
+	FVirtualTextureSystem::Get().RequestRecordedTiles(MoveTemp(InPageRequests), InCaptureResolution);
 }
 
 void FRendererModule::FlushVirtualTextureCache()
diff --git a/Engine/Source/Runtime/Renderer/Private/VT/VirtualTextureSystem.cpp b/Engine/Source/Runtime/Renderer/Private/VT/VirtualTextureSystem.cpp
--- a/Engine/Source/Runtime/Renderer/Private/VT/VirtualTextureSystem.cpp
+++ b/Engine/Source/Runtime/Renderer/Private/VT/VirtualTextureSystem.cpp
@@ -117,7 +117,8 @@ static TAutoConsoleVariable<int32> CVarVTEnablePlayback(
 static TAutoConsoleVariable<float> CVarVTPlaybackMipBias(
 	TEXT("r.VT.PlaybackMipBias"),
 	0,
-	TEXT("Mip bias to apply during playback of recorded feedback requests."),
+	TEXT("Mip bias to apply during playback of recorded feedback requests.\n")
+	TEXT("This is in addition to the automatic bias applied for the difference between capture and playback resolutions."),
 	ECVF_RenderThreadSafe
 );
 static TAutoConsoleVariable<int32> CVarVTPageUpdateFlushCount(
@@ -206,6 +207,16 @@ FVirtualTextureUpdateSettings::FVirtualTextureUpdateSettings()
 	MaxContinuousUpdates = VirtualTextureScalability::GetMaxContinuousUpdatesPerFrame();
 }
 
+FVirtualTextureUpdateSettings::FVirtualTextureUpdateSettings(FViewFamilyInfo const& InViewFamily)
+	: FVirtualTextureUpdateSettings()
+{
+	if (InViewFamily.bOverrideVirtualTextureThrottle)
+	{
+		EnableThrottling(false);
+	}
+	PlaybackViewSize = InViewFamily.SceneTexturesConfig.Extent;
+}
+
 static bool ShouldSyncProduceLockedTiles()
 {
 #if WITH_EDITOR
@@ -281,7 +292,7 @@ public:
 struct FAddRequestedTilesParameters
 {
 	FVirtualTextureSystem* System = nullptr;
-	uint32 LevelBias = 0u;
+	uint32 ViewLevelBias = 0u;
 	const uint64* RequestBuffer = nullptr;
 	uint32 NumRequests = 0u;
 	FUniquePageList* UniquePageList = nullptr;
@@ -2875,10 +2886,17 @@ void FVirtualTextureSystem::GatherFeedbackRequests(FConcurrentLinearBulkObjectAl
 	{
 		if (Settings.bEnablePlayback)
 		{
+			float LevelBias = CVarVTPlaybackMipBias.GetValueOnRenderThread();
+			LevelBias += UTexture2D::GetGlobalMipMapLODBias();
+			if (Settings.PlaybackViewSize.X > 0 && Settings.PlaybackViewSize.Y > 0)
+			{
+				LevelBias += FMath::Log2(FMath::Max3((float)PageRequestCaptureResolution.X / (float)Settings.PlaybackViewSize.X, (float)PageRequestCaptureResolution.Y / (float)Settings.PlaybackViewSize.Y, 1.f));
+			}
+
 			// todo: We can split this into concurrent tasks. 
 			FAddRequestedTilesParameters Parameters;
 			Parameters.System = this;
-			Parameters.LevelBias = FMath::FloorToInt(CVarVTPlaybackMipBias.GetValueOnRenderThread() + UTexture2D::GetGlobalMipMapLODBias() + 0.5f);
+			Parameters.ViewLevelBias = FMath::RoundToInt(FMath::Max(LevelBias, 0.f));
 			Parameters.RequestBuffer = PageRequestPlaybackBuffer.GetData();
 			Parameters.NumRequests = PageRequestPlaybackBuffer.Num();
 			Parameters.UniquePageList = Allocator.Create<FUniquePageList>();
@@ -2888,6 +2906,7 @@ void FVirtualTextureSystem::GatherFeedbackRequests(FConcurrentLinearBulkObjectAl
 		}
 
 		PageRequestPlaybackBuffer.Reset(0);
+		PageRequestCaptureResolution = FIntPoint::ZeroValue;
 	}
 
 	// Pages from feedback buffer were generated several frames ago, so they may no longer be valid for newly allocated VTs
@@ -3384,10 +3403,10 @@ void FVirtualTextureSystem::RecordPageRequests(FUniquePageList const* UniquePage
 
 		// PersistentHash is used to find associated AllocatedVT at playback.
 		const uint32 PersistentHash = AllocatedVT->GetPersistentHash();
-		// Log2TileCountX is used at playback to detect and adjust for any TextureLODGroup bias that has been applied.
-		const uint32 Log2TileCountX = FMath::CeilLogTwo(AllocatedVT->GetBlockWidthInTiles()) & 0xf;
+		// MaxLevel is used to detect and adjust for any TextureLODGroup bias that has been applied at playback.
+		const uint32 MaxLevel = AllocatedVT->GetMaxLevel() & 0xf;
 		const uint32 LocalAddress = (vAddress - AllocatedVT->GetVirtualAddress()) & 0xffffff;
-		const uint64 ExportPageId = ((uint64)PersistentHash << 32) | ((uint64)vLevelPlusOne << 28) | ((uint64)Log2TileCountX << 24) | (uint64)LocalAddress;
+		const uint64 ExportPageId = ((uint64)PersistentHash << 32) | ((uint64)vLevelPlusOne << 28) | ((uint64)MaxLevel << 24) | (uint64)LocalAddress;
 		
 		OutPages.Add(ExportPageId);
 	}
@@ -3411,7 +3430,7 @@ uint64 FVirtualTextureSystem::GetVirtualTextureRequestRecordBuffer(TSet<uint64>&
 
 #endif // WITH_EDITOR
 
-void FVirtualTextureSystem::RequestRecordedTiles(TArray<uint64>&& InPageRequests)
+void FVirtualTextureSystem::RequestRecordedTiles(TArray<uint64>&& InPageRequests, FIntPoint const& InCaptureResolution)
 {
 	UE::TScopeLock Lock(Mutex);
 
@@ -3423,6 +3442,11 @@ void FVirtualTextureSystem::RequestRecordedTiles(TArray<uint64>&& InPageRequests
 	{
 		PageRequestPlaybackBuffer.Append(InPageRequests);
 	}
+
+	// We expect all recorded requests to have been captured at the same resolution.
+	// If that turns out not to be the case we will need to store an array of recorded requests per resolution.
+	ensure(PageRequestCaptureResolution == FIntPoint::ZeroValue || PageRequestCaptureResolution == InCaptureResolution);
+	PageRequestCaptureResolution = InCaptureResolution;
 }
 
 void FVirtualTextureSystem::AddRequestedTilesTask(const FAddRequestedTilesParameters& Parameters)
@@ -3430,7 +3454,7 @@ void FVirtualTextureSystem::AddRequestedTilesTask(const FAddRequestedTilesParame
 	FUniquePageList* RESTRICT RequestedPageList = Parameters.UniquePageList;
 	const uint64* RESTRICT Buffer = Parameters.RequestBuffer;
 	const uint32 BufferSize = Parameters.NumRequests;
-	const uint32 ViewLevelBias = Parameters.LevelBias;
+	const uint32 ViewLevelBias = Parameters.ViewLevelBias;
 
 	uint32 CurrentPersistentHash = ~0u;
 	IAllocatedVirtualTexture* AllocatedVT = nullptr;
@@ -3453,25 +3477,23 @@ void FVirtualTextureSystem::AddRequestedTilesTask(const FAddRequestedTilesParame
 		{
 			const uint32 SourceLevelPlusOne = (uint32)(PageRequest >> 28) & 0xf;
 			const uint32 SourceLevel = FMath::Max(SourceLevelPlusOne, 1u) - 1u;
-			const uint32 SourceLog2TileCountX = (uint32)(PageRequest >> 24) & 0xf;
+			const uint32 SourceMaxLevel = (uint32)(PageRequest >> 24) & 0xf;
 			const uint32 SourceLocalAddress = (uint32)(PageRequest & 0xffffff);
 
 			// Deal with the case of texture size being reduced from recording to playback due to LOD bias in cook.
-			const uint32 Log2TileCountX = FMath::CeilLogTwo(AllocatedVT->GetBlockWidthInTiles());
-			const uint32 TextureSizeBias = SourceLog2TileCountX > Log2TileCountX ? SourceLog2TileCountX - Log2TileCountX : 0;
+			const uint32 MaxLevel = AllocatedVT->GetMaxLevel();
+			const uint32 TextureSizeBias = SourceMaxLevel > MaxLevel ? SourceMaxLevel - MaxLevel: 0;
+			const uint32 Level = FMath::Clamp((int32)SourceLevel - (int32)TextureSizeBias + (int32)ViewLevelBias, 0, MaxLevel);
 
 			const uint32 LocalAddress = SourceLocalAddress >> (TextureSizeBias * 2u);
 			const uint32 BaseAddress = AllocatedVT->GetVirtualAddress();
 			const uint32 Address = BaseAddress + LocalAddress;
-
-			const uint32 MaxLevel = AllocatedVT->GetMaxLevel();
-			const uint32 Level = FMath::Clamp((int32)SourceLevel - (int32)TextureSizeBias + (int32)ViewLevelBias, 0, MaxLevel);
-
 			const uint32 PageX = FMath::ReverseMortonCode2(Address) >> Level;
 			const uint32 PageY = FMath::ReverseMortonCode2(Address >> 1) >> Level;
+
 			const uint32 SpaceId = AllocatedVT->GetSpaceID();
 
-			const uint32 PageId = PageX | (PageY << 12) | ((Level+1) << 24) | (SpaceId << 28);
+			const uint32 PageId = PageX | (PageY << 12) | ((Level + 1) << 24) | (SpaceId << 28);
 
 			RequestedPageList->Add(PageId, 0xffff);
 		}
diff --git a/Engine/Source/Runtime/Renderer/Private/VT/VirtualTextureSystem.h b/Engine/Source/Runtime/Renderer/Private/VT/VirtualTextureSystem.h
--- a/Engine/Source/Runtime/Renderer/Private/VT/VirtualTextureSystem.h
+++ b/Engine/Source/Runtime/Renderer/Private/VT/VirtualTextureSystem.h
@@ -33,10 +33,12 @@ struct FAddRequestedTilesParameters;
 struct FGatherRequestsParameters;
 struct FPageUpdateBuffer;
 class ISceneRenderer;
+class FViewFamilyInfo;
 
 struct FVirtualTextureUpdateSettings
 {
 	FVirtualTextureUpdateSettings();
+	FVirtualTextureUpdateSettings(FViewFamilyInfo const& ViewFamily);
 
 	/** Force settings so that throttling page uploads is effectively disabled. */
 	FVirtualTextureUpdateSettings& EnableThrottling(bool bEnable)
@@ -80,6 +82,7 @@ struct FVirtualTextureUpdateSettings
 	int32 MaxSVTPageUploads;
 	int32 MaxPagesProduced;
 	int32 MaxContinuousUpdates;
+	FIntPoint PlaybackViewSize = FIntPoint(0, 0);
 };
 
 class FVirtualTextureUpdater
@@ -162,7 +165,7 @@ public:
 	void SetVirtualTextureRequestRecordBuffer(uint64 Handle);
 	uint64 GetVirtualTextureRequestRecordBuffer(TSet<uint64>& OutPageRequests);
 #endif
-	void RequestRecordedTiles(TArray<uint64>&& InPageRequests);
+	void RequestRecordedTiles(TArray<uint64>&& InPageRequests, FIntPoint const& InCaptureResolution);
 
 	void FlushCache();
 	void FlushCache(FVirtualTextureProducerHandle const& ProducerHandle, int32 SpaceID, FIntRect const& TextureRegion, uint32 MaxLevelToEvict, uint32 MaxAgeToKeepMapped, EVTInvalidatePriority InvalidatePriority);
@@ -300,6 +303,7 @@ private:
 	TSet<uint64> PageRequestRecordBuffer;
 #endif
 	TArray<uint64> PageRequestPlaybackBuffer;
+	FIntPoint PageRequestCaptureResolution = FIntPoint::ZeroValue;
 
 #if !UE_BUILD_SHIPPING
 	void GetOnScreenMessages(FCoreDelegates::FSeverityMessageMap& OutMessages);
