Workaround for incorrect view distance scale being captured in the GPU scene primitive data for InstanceDrawDistance, InstanceWorldPositionOffsetDisableDistance and PixelProgrammableDistance. The GPU scene now tracks the value of r.ViewDistanceScale and triggers a full upload if it changes. This behavior can be disabled by r.GPUScene.ViewDistanceScaleWorkaround.

#Synced-CL 50488888

diff --git a/Engine/Source/Runtime/Renderer/Private/GPUScene.cpp b/Engine/Source/Runtime/Renderer/Private/GPUScene.cpp
--- a/Engine/Source/Runtime/Renderer/Private/GPUScene.cpp
+++ b/Engine/Source/Runtime/Renderer/Private/GPUScene.cpp
@@ -31,6 +31,7 @@
 #include "PrimitiveUniformShaderParametersBuilder.h"
 #include "InstanceDataSceneProxy.h"
 #include "SceneRendererInterface.h"
+#include "UnrealEngine.h"
 
 // Useful for debugging
 #define FORCEINLINE_GPUSCENE FORCEINLINE
@@ -152,6 +153,12 @@ static TAutoConsoleVariable<bool> CVarGPUSceneUseReservedResources(
 	TEXT("Turning this on makes all GPU-Scene buffers default to using reserved resource allocations."),
 	ECVF_ReadOnly);
 
+static TAutoConsoleVariable<bool> CVarGPUSceneViewDistanceScaleWorkaround(
+	TEXT("r.GPUScene.ViewDistanceScaleWorkaround"),
+	true,
+	TEXT("If true (default) the whole GPU scene is reuploaded whenever the view distance scale is changed to work around the fact that the scale is baked into some primitive data."),
+	ECVF_RenderThreadSafe
+);
 LLM_DECLARE_TAG_API(GPUScene, RENDERER_API);
 DECLARE_LLM_MEMORY_STAT(TEXT("GPUScene"), STAT_GPUSceneLLM, STATGROUP_LLMFULL);
 DECLARE_LLM_MEMORY_STAT(TEXT("GPUScene"), STAT_GPUSceneSummaryLLM, STATGROUP_LLM);
@@ -783,6 +790,19 @@ void FGPUScene::UpdateInternal(FRDGBuilder& GraphBuilder, FSceneUniformBuffer& S
 	MaxInstancesDuringPrevUpdate = uint32(InstanceSceneDataAllocator.GetMaxSize());
 #endif // UE_BUILD_SHIPPING
 
+	if (CVarGPUSceneViewDistanceScaleWorkaround.GetValueOnRenderThread())
+	{
+		if (GetCachedScalabilityCVars().ViewDistanceScale != TrackedViewDistanceScle)
+		{
+			bUpdateAllPrimitives = true;
+		}
+		TrackedViewDistanceScle = GetCachedScalabilityCVars().ViewDistanceScale;
+	}
+	else
+	{
+		TrackedViewDistanceScle = -1.0f;
+	}
+
 	bool bPrevUseTiledInstanceDataLayout = bUseTiledInstanceDataLayout;
 	bUseTiledInstanceDataLayout = CVarGPUSceneInstanceDataTileSizeLog2.GetValueOnRenderThread() >= 0;
 	int32 NewTileSizeLog2 = FMath::Max(0, CVarGPUSceneInstanceDataTileSizeLog2.GetValueOnRenderThread());
diff --git a/Engine/Source/Runtime/Renderer/Private/GPUScene.h b/Engine/Source/Runtime/Renderer/Private/GPUScene.h
--- a/Engine/Source/Runtime/Renderer/Private/GPUScene.h
+++ b/Engine/Source/Runtime/Renderer/Private/GPUScene.h
@@ -509,6 +509,8 @@ private:
 	bool bLoggedPayloadBufferOverflow = false;
 #endif // UE_BUILD_SHIPPING
 	uint64 InstancePayloadDataSizeReserveUnclamped = 0;
+	// Used to track the GetCachedScalabilityCVars().ViewDistanceScale such that we can trigger a full upload if it changes, this is required as the value is baked into some values such as in FPrimitiveUniformShaderParametersBuilder::InstanceWorldPositionOffsetDisableDistance
+	float TrackedViewDistanceScle = -1.0f;
 };
 
 class FGPUSceneScopeBeginEndHelper
