Merged the mobile SSAO implementation into the SM5 path.
Removed the mobile-specific GTAO.
Updated mobile platform CVars to use full-resolution SSAO by default; SSAO is now enabled only for High and Epic scalability tiers.
Improved mobile SSAO quality by using GBuffer normals when available.
Added r.Mobile.SSAOShadingModelsExcludeMask to exclude specific shading models from SSAO (deferred shading only).
Added a blur pass for full-resolution SSAO when TAA is disabled.

#Synced-CL 50477682

#Synced-CL 50530667

diff --git a/Engine/Config/Android/AndroidScalability.ini b/Engine/Config/Android/AndroidScalability.ini
--- a/Engine/Config/Android/AndroidScalability.ini
+++ b/Engine/Config/Android/AndroidScalability.ini
@@ -116,7 +116,7 @@ r.BloomQuality=0
 r.DepthOfFieldQuality=0
 r.LightShaftQuality=0
 r.EyeAdaptationQuality=1
-r.Mobile.AmbientOcclusionQuality=0
+r.AmbientOcclusionLevels=0
 r.LensFlareQuality=0
 
 [PostProcessQuality@1]
@@ -125,7 +125,7 @@ r.BloomQuality=1
 r.DepthOfFieldQuality=1
 r.LightShaftQuality=1
 r.EyeAdaptationQuality=1
-r.Mobile.AmbientOcclusionQuality=0
+r.AmbientOcclusionLevels=0
 r.LensFlareQuality=1
 
 [PostProcessQuality@2]
@@ -134,7 +134,7 @@ r.BloomQuality=2
 r.DepthOfFieldQuality=1
 r.LightShaftQuality=1
 r.EyeAdaptationQuality=1
-r.Mobile.AmbientOcclusionQuality=2
+r.AmbientOcclusionLevels=1
 r.LensFlareQuality=2
 
 [PostProcessQuality@3]
@@ -143,7 +143,7 @@ r.BloomQuality=3
 r.DepthOfFieldQuality=1
 r.LightShaftQuality=1
 r.EyeAdaptationQuality=1
-r.Mobile.AmbientOcclusionQuality=3
+r.AmbientOcclusionLevels=1
 r.LensFlareQuality=2
 
 [AntiAliasingQuality@0]
diff --git a/Engine/Config/IOS/IOSScalability.ini b/Engine/Config/IOS/IOSScalability.ini
--- a/Engine/Config/IOS/IOSScalability.ini
+++ b/Engine/Config/IOS/IOSScalability.ini
@@ -116,7 +116,7 @@ r.BloomQuality=0
 r.DepthOfFieldQuality=0
 r.LightShaftQuality=0
 r.EyeAdaptationQuality=1
-r.Mobile.AmbientOcclusionQuality=0
+r.AmbientOcclusionLevels=0
 r.LensFlareQuality=0
 
 [PostProcessQuality@1]
@@ -125,7 +125,7 @@ r.BloomQuality=1
 r.DepthOfFieldQuality=0
 r.LightShaftQuality=0
 r.EyeAdaptationQuality=1
-r.Mobile.AmbientOcclusionQuality=0
+r.AmbientOcclusionLevels=0
 r.LensFlareQuality=1
 
 [PostProcessQuality@2]
@@ -134,7 +134,7 @@ r.BloomQuality=2
 r.DepthOfFieldQuality=0
 r.LightShaftQuality=0
 r.EyeAdaptationQuality=1
-r.Mobile.AmbientOcclusionQuality=2
+r.AmbientOcclusionLevels=1
 r.LensFlareQuality=2
 
 [PostProcessQuality@3]
@@ -143,7 +143,7 @@ r.BloomQuality=3
 r.DepthOfFieldQuality=1
 r.LightShaftQuality=1
 r.EyeAdaptationQuality=1
-r.Mobile.AmbientOcclusionQuality=3
+r.AmbientOcclusionLevels=1
 r.LensFlareQuality=2
 
 [AntiAliasingQuality@0]
diff --git a/Engine/Shaders/Private/DeferredShadingCommon.ush b/Engine/Shaders/Private/DeferredShadingCommon.ush
--- a/Engine/Shaders/Private/DeferredShadingCommon.ush
+++ b/Engine/Shaders/Private/DeferredShadingCommon.ush
@@ -507,7 +507,7 @@ float4 ComputeIndirectLightingSampleE(uint2 TracingPixelCoord, uint TracingRayIn
 #if SHADING_PATH_MOBILE
 
 #ifndef ENABLE_SHADINGMODEL_SUPPORT_MOBILE_DEFERRED
-#define ENABLE_SHADINGMODEL_SUPPORT_MOBILE_DEFERRED 0
+#define ENABLE_SHADINGMODEL_SUPPORT_MOBILE_DEFERRED 1
 #endif
 
 #define MOBILE_SHADINGMODEL_SUPPORT (!MOBILE_DEFERRED_SHADING || ENABLE_SHADINGMODEL_SUPPORT_MOBILE_DEFERRED)
diff --git a/Engine/Shaders/Private/PostProcessAmbientOcclusion.usf b/Engine/Shaders/Private/PostProcessAmbientOcclusion.usf
--- a/Engine/Shaders/Private/PostProcessAmbientOcclusion.usf
+++ b/Engine/Shaders/Private/PostProcessAmbientOcclusion.usf
@@ -150,6 +150,16 @@
 	};
 #endif // USE_SAMPLESET
 
+float3 PackSceneDepth(float InSceneDepth)
+{
+	return UnpackRGBA8(PackR24F(InSceneDepth)).rgb;
+}
+
+float UnpackSceneDepth(float3 InPackedSceneDepth)
+{
+	return UnpackR24F(PackRGBA8(float4(InPackedSceneDepth, 0)));
+}
+
 // -----------------------------------------------------------------------------------------------------------------------------
 // To be included after defines
 #include "PostProcessAmbientOcclusionCommon.ush"
@@ -260,9 +270,10 @@ void MainPSandCS(in float4 UVAndScreenPos, float4 SvPosition, out float4 OutColo
 	float InvTanHalfFov = ScreenSpaceAOParams[3].w;
 	float3 FovFix = float3(InvTanHalfFov, Ratio * InvTanHalfFov, 1);
 	float3 InvFovFix = 1.0f / FovFix;
-
+	float4 ModifiedSvPosition = float4(SvPosition.xy * SSAO_SvPositionScaleBias.xx + SSAO_SvPositionScaleBias.yy, SvPosition.zw);
+		
 	float SceneDepth = GetDepthFromAOInput(UV);
-	float3 WorldNormal = GetWorldSpaceNormalFromAOInput(UV, SvPosition);
+	float3 WorldNormal = GetWorldSpaceNormalFromAOInput(UV, ModifiedSvPosition);
 
 	// can be NaN if WorldNormal=0,0,0 which happens when !USE_NORMALS
 	float3 ViewSpaceNormal = normalize(mul(WorldNormal, (float3x3)View.TranslatedWorldToView));
@@ -438,7 +449,7 @@ void MainPSandCS(in float4 UVAndScreenPos, float4 SvPosition, out float4 OutColo
 #endif
 
 	// we don't support ddx_fine() for SM4
-#if !COMPUTE_SHADER && QUAD_MESSAGE_PASSING_BLUR > 0 && FEATURE_LEVEL >= FEATURE_LEVEL_SM5
+#if !COMPUTE_SHADER && QUAD_MESSAGE_PASSING_BLUR > 0
 	{
 		// .x: AO output, .y:SceneDepth .zw:view space normal
 		float4 CenterPixel = float4(OutColor.r, SceneDepth, normalize(ViewSpaceNormal).xy); 
@@ -484,6 +495,10 @@ void MainPSandCS(in float4 UVAndScreenPos, float4 SvPosition, out float4 OutColo
 //		OutColor = (WeightA - 0.333f) / 0.666f;
 	}
 #endif
+
+#if OUTPUT_DEPTH
+	OutColor.gba = PackSceneDepth(SceneDepth);
+#endif
 }
 
 void MainPS(in noperspective float4 UVAndScreenPos : TEXCOORD0, in FStereoPSInput StereoInput, float4 SvPosition : SV_POSITION, out float4 OutColor : SV_Target0)
@@ -507,7 +522,9 @@ void MainCS(
 	
 	// todo: move to a function
 	float4 SvPosition = float4(PixelCenter, 0, 0) * ScaleFactor;	
-	float2 BufferUV = SvPositionToBufferUV(SvPosition);
+	float2 BufferUV = PixelCenter * AOViewport_ExtentInverse;
+	float2 ScreenPos = UVToScreenPos(BufferUV, AOViewport_Extent);
+
 	SvPosition.z = LookupDeviceZ(BufferUV);
 	// todo: investigate
 //  SvPosition.w = ConvertFromDeviceZ(SvPosition.z);
@@ -522,7 +539,7 @@ void MainCS(
 	BRANCH
 	if (SceneDepth * FadeMul + FadeAdd < 1)
 	{
-		MainPSandCS(float4(BufferUV, SvPositionToScreenPosition(SvPosition).xy), SvPosition, OutColor);
+		MainPSandCS(float4(BufferUV, ScreenPos), SvPosition, OutColor);
 	}
 
 	// Here we could optimized for coalessing writes but that might not be the performance bottleneck.
@@ -530,6 +547,7 @@ void MainCS(
 	// http://on-demand.gputechconf.com/gtc/2010/presentations/S12312-DirectCompute-Pre-Conference-Tutorial.pdf
 	OutTexture[PixelPos] = OutColor.r;
 }
+#endif
 
 SCREEN_PASS_TEXTURE_VIEWPORT(SSAOSmoothOutputViewport)
 FScreenTransform SSAOSmoothOutputToInput;
@@ -537,6 +555,27 @@ FScreenTransform SSAOSmoothOutputToInput;
 Texture2D SSAOSmoothInputTexture;
 SamplerState SSAOSmoothInputSampler;
 
+float SSAOSmoothPSandCS(in float2 UV)
+{
+	float2 SampleUV = ApplyScreenTransform(UV, SSAOSmoothOutputToInput);
+
+	// Use a 4x4 box filter because the random texture is tiled 4x4
+	float Result;
+	Result = SSAOSmoothInputTexture.SampleLevel(SSAOSmoothInputSampler, SampleUV, 0).r;
+	Result += SSAOSmoothInputTexture.SampleLevel(SSAOSmoothInputSampler, SampleUV, 0, int2(2, 0)).r;
+	Result += SSAOSmoothInputTexture.SampleLevel(SSAOSmoothInputSampler, SampleUV, 0, int2(0, 2)).r;
+	Result += SSAOSmoothInputTexture.SampleLevel(SSAOSmoothInputSampler, SampleUV, 0, int2(2, 2)).r;
+
+	return Result * 0.25;
+}
+
+void MainSSAOSmoothPS(in noperspective float4 UVAndScreenPos : TEXCOORD0, in FStereoPSInput StereoInput, float4 SvPosition : SV_POSITION, out float4 OutColor : SV_Target0)
+{
+	float SmoothColor = SSAOSmoothPSandCS(UVAndScreenPos.xy);
+	OutColor = SmoothColor;
+}
+
+#if COMPUTE_SHADER
 RWTexture2D<float> SSAOSmoothOutputTexture;
 
 [numthreads(THREADGROUP_SIZEX, THREADGROUP_SIZEY, 1)]
@@ -554,18 +593,111 @@ void MainSSAOSmoothCS(
 	uint2 DestPixelPos = SSAOSmoothOutputViewport_ViewportMin + DispatchThreadId;
 	float2 DestUV = DestPixelPos * SSAOSmoothOutputViewport_ExtentInverse;
 
-	float2 SampleUV = ApplyScreenTransform(DestUV, SSAOSmoothOutputToInput);
+	SSAOSmoothOutputTexture[DestPixelPos] = SSAOSmoothPSandCS(DestUV);
+}
+#endif
 
-	// Use a 4x4 box filter because the random texture is tiled 4x4
-	float Result;
-	Result  = SSAOSmoothInputTexture.SampleLevel(SSAOSmoothInputSampler, SampleUV, 0).r;
-	Result += SSAOSmoothInputTexture.SampleLevel(SSAOSmoothInputSampler, SampleUV, 0, int2(2, 0)).r;
-	Result += SSAOSmoothInputTexture.SampleLevel(SSAOSmoothInputSampler, SampleUV, 0, int2(0, 2)).r;
-	Result += SSAOSmoothInputTexture.SampleLevel(SSAOSmoothInputSampler, SampleUV, 0, int2(2, 2)).r;
 
-	SSAOSmoothOutputTexture[DestPixelPos] = Result * 0.25;
-}
+#ifdef UPSAMPLE_PASS
+
+Texture2D AOTexture;
+SamplerState AOSampler;
+
+float UpsampleSSAO(float2 BufferUV)
+{
+#if UPSAMPLE_QUALITY == 0
+	// Bilinear sample.
+	return Texture2DSampleLevel(AOTexture, AOSampler, BufferUV, 0).x;
+
+#elif UPSAMPLE_QUALITY == 1
+	// 4 sample bilinear with depth weighting.
+	float2 LowResBufferSize = floor(View.BufferSizeAndInvSize.xy / 2);
+	float2 LowResTexelSize = 1.0f / LowResBufferSize;
+
+	float2 Corner00UV = floor(BufferUV * LowResBufferSize - 0.5f) / LowResBufferSize + 0.5f * LowResTexelSize;
+	float2 BilinearWeights = (BufferUV - Corner00UV) * LowResBufferSize;
+
+	float4 AOValues00 = Texture2DSampleLevel(AOTexture, AOSampler, Corner00UV, 0);
+	float4 AOValues10 = Texture2DSampleLevel(AOTexture, AOSampler, Corner00UV + float2(LowResTexelSize.x, 0), 0);
+	float4 AOValues01 = Texture2DSampleLevel(AOTexture, AOSampler, Corner00UV + float2(0, LowResTexelSize.y), 0);
+	float4 AOValues11 = Texture2DSampleLevel(AOTexture, AOSampler, Corner00UV + LowResTexelSize, 0);
+
+	float4 CornerWeights = float4(
+		(1 - BilinearWeights.y) * (1 - BilinearWeights.x),
+		(1 - BilinearWeights.y) * BilinearWeights.x,
+		BilinearWeights.y * (1 - BilinearWeights.x),
+		BilinearWeights.y * BilinearWeights.x);
+
+	float SceneDepth00 = UnpackSceneDepth(AOValues00.gba);
+	float SceneDepth10 = UnpackSceneDepth(AOValues10.gba);
+	float SceneDepth01 = UnpackSceneDepth(AOValues01.gba);
+	float SceneDepth11 = UnpackSceneDepth(AOValues11.gba);
+
+	float MaxDepth = max(max(max(SceneDepth00, SceneDepth10), SceneDepth01), SceneDepth11);
+	float MinDepth = min(min(min(SceneDepth00, SceneDepth10), SceneDepth01), SceneDepth11);
+
+	if (MaxDepth / MinDepth > 1.02f)
+	{
+		float Epsilon = 0.0001f;
+		float SceneDepth = GetDepthFromAOInput(BufferUV);
+		float4 CornerDepths = abs(float4(SceneDepth00, SceneDepth10, SceneDepth01, SceneDepth11));
+		float4 DepthWeights = 1.0f / (abs(CornerDepths - SceneDepth.xxxx) + Epsilon);
+		CornerWeights *= DepthWeights;
+	}
+
+	float InterpolatedResult =
+		(CornerWeights.x * AOValues00.r
+			+ CornerWeights.y * AOValues10.r
+			+ CornerWeights.z * AOValues01.r
+			+ CornerWeights.w * AOValues11.r)
+		/ dot(CornerWeights, 1);
+
+	return InterpolatedResult;
+
+#elif UPSAMPLE_QUALITY == 2
+	// 9 sample with depth weighting.
+	float2 LowResBufferSize = floor(View.BufferSizeAndInvSize.xy * 0.5);
+	float2 LowResTexelSize = 1.0f / LowResBufferSize;
+
+	float4 AOValues[9];
+	AOValues[0] = Texture2DSampleLevel(AOTexture, AOSampler, BufferUV + float2(-1, -1) * LowResTexelSize, 0);
+	AOValues[1] = Texture2DSampleLevel(AOTexture, AOSampler, BufferUV + float2(0, -1) * LowResTexelSize, 0);
+	AOValues[2] = Texture2DSampleLevel(AOTexture, AOSampler, BufferUV + float2(1, -1) * LowResTexelSize, 0);
+	AOValues[3] = Texture2DSampleLevel(AOTexture, AOSampler, BufferUV + float2(-1, 0) * LowResTexelSize, 0);
+	AOValues[4] = Texture2DSampleLevel(AOTexture, AOSampler, BufferUV + float2(0, 0) * LowResTexelSize, 0);
+	AOValues[5] = Texture2DSampleLevel(AOTexture, AOSampler, BufferUV + float2(1, 0) * LowResTexelSize, 0);
+	AOValues[6] = Texture2DSampleLevel(AOTexture, AOSampler, BufferUV + float2(-1, 1) * LowResTexelSize, 0);
+	AOValues[7] = Texture2DSampleLevel(AOTexture, AOSampler, BufferUV + float2(0, 1) * LowResTexelSize, 0);
+	AOValues[8] = Texture2DSampleLevel(AOTexture, AOSampler, BufferUV + float2(1, 1) * LowResTexelSize, 0);
+
+	float SceneDepth = GetDepthFromAOInput(BufferUV);
+	float Epsilon = 0.0001f;
+
+	float TotalSum = 0;
+	float TotalWeight = 0;
+	UNROLL for (int i = 0; i < 9; ++i)
+	{
+		float SampleValue = AOValues[i].r;
+		float SampleDepth = UnpackSceneDepth(AOValues[i].gba);
+
+		float Weight = 1.0f / (abs(SceneDepth - SampleDepth) + Epsilon);
+		TotalWeight += Weight;
+		TotalSum += SampleValue * Weight;
+	}
+
+	return TotalSum / TotalWeight;
 #endif
+}
+
+void AmbientOcclusionUpsamplePS(
+	in float4 SvPosition : SV_POSITION,
+	out float4 OutColor : SV_Target0)
+{
+	float2 BufferUV = SvPositionToBufferUV(SvPosition);
+	OutColor = UpsampleSSAO(BufferUV).xxxx;
+}
+
+#endif // UPSAMPLE_PASS
 
 #if SHADER_QUALITY == 0
 	// very low
@@ -648,12 +780,20 @@ float3 GetNormal(float2 UV, float3 ViewSpacePosMid)
 #else
 	// Get the normal derived from the depth buffer
 	float2 DeltaUV = View.BufferSizeAndInvSize.zw;
+
+#if SHADING_PATH_MOBILE
+	Texture2D DepthTexture = MobileSceneTextures.SceneDepthTexture;
+	SamplerState DepthTextureSampler = MobileSceneTextures.SceneDepthTextureSampler;
+#else
+	Texture2D DepthTexture = SceneTexturesStruct.SceneDepthTexture;
+	SamplerState DepthTextureSampler = SceneTexturesStruct_SceneDepthTextureSampler;
+#endif
 	
-	float DeviceZ		= Texture2DSampleLevel(SceneTexturesStruct.SceneDepthTexture, SceneTexturesStruct_SceneDepthTextureSampler, UV,0).r;									
-	float DeviceZLeft	= Texture2DSampleLevel(SceneTexturesStruct.SceneDepthTexture, SceneTexturesStruct_SceneDepthTextureSampler, UV + float2(-DeltaUV.x,  0.0f),0).r;		
-	float DeviceZTop	= Texture2DSampleLevel(SceneTexturesStruct.SceneDepthTexture, SceneTexturesStruct_SceneDepthTextureSampler, UV + float2( 0.0f     , -DeltaUV.y),0).r;	
-	float DeviceZRight	= Texture2DSampleLevel(SceneTexturesStruct.SceneDepthTexture, SceneTexturesStruct_SceneDepthTextureSampler, UV + float2( DeltaUV.x,  0.0f),0).r;		
-	float DeviceZBottom = Texture2DSampleLevel(SceneTexturesStruct.SceneDepthTexture, SceneTexturesStruct_SceneDepthTextureSampler, UV + float2( 0.0f     ,  DeltaUV.y),0).r;	
+	float DeviceZ		= Texture2DSampleLevel(DepthTexture, DepthTextureSampler, UV,0).r;
+	float DeviceZLeft	= Texture2DSampleLevel(DepthTexture, DepthTextureSampler, UV + float2(-DeltaUV.x,  0.0f),0).r;
+	float DeviceZTop	= Texture2DSampleLevel(DepthTexture, DepthTextureSampler, UV + float2( 0.0f     , -DeltaUV.y),0).r;
+	float DeviceZRight	= Texture2DSampleLevel(DepthTexture, DepthTextureSampler, UV + float2( DeltaUV.x,  0.0f),0).r;
+	float DeviceZBottom = Texture2DSampleLevel(DepthTexture, DepthTextureSampler, UV + float2( 0.0f     ,  DeltaUV.y),0).r;
 
 	float DeviceZDdx	= TakeSmallerAbsDelta(DeviceZLeft, DeviceZ, DeviceZRight);
 	float DeviceZDdy	= TakeSmallerAbsDelta(DeviceZTop, DeviceZ, DeviceZBottom);
@@ -673,7 +813,15 @@ float3 GetNormal(float2 UV, float3 ViewSpacePosMid)
 
 float GetLinearDepthProj(float2 ScreenUV)
 {
-	float DeviceZ = Texture2DSampleLevel(SceneTexturesStruct.SceneDepthTexture, SceneTexturesStruct_SceneDepthTextureSampler, ScreenUV, 0).r;
+#if SHADING_PATH_MOBILE
+	Texture2D DepthTexture = MobileSceneTextures.SceneDepthTexture;
+	SamplerState DepthTextureSampler = MobileSceneTextures.SceneDepthTextureSampler;
+#else
+	Texture2D DepthTexture = SceneTexturesStruct.SceneDepthTexture;
+	SamplerState DepthTextureSampler = SceneTexturesStruct_SceneDepthTextureSampler;
+#endif
+	
+	float DeviceZ = Texture2DSampleLevel(DepthTexture, DepthTextureSampler, ScreenUV, 0).r;
 	return  1.0f / (DeviceZ * View.InvDeviceZToWorldZTransform[2] - View.InvDeviceZToWorldZTransform[3]);
 }
 
diff --git a/Engine/Shaders/Private/PostProcessAmbientOcclusionCommon.ush b/Engine/Shaders/Private/PostProcessAmbientOcclusionCommon.ush
--- a/Engine/Shaders/Private/PostProcessAmbientOcclusionCommon.ush
+++ b/Engine/Shaders/Private/PostProcessAmbientOcclusionCommon.ush
@@ -48,6 +48,7 @@ float2 SSAO_DownsampledAOInverseSize;
 float2 SSAO_DownsampledAOUVViewportMin;
 float2 SSAO_DownsampledAOUVViewportMax;
 float2 SSAO_SvPositionScaleBias;
+uint SSAO_ShadingModelsExcludeMask;
 
 // could be moved to a more central spot
 // @param ScreenPos -1 .. 1
@@ -103,8 +104,9 @@ float3 GetWorldSpaceNormalFromAOInput(float2 UV, float4 SvPosition)
 	#if USE_AO_SETUP_AS_INPUT
 		// Low resolution normal computed in the setup (downscaled) pass.
 		WorldNormal = Texture2DSample(SSAO_SetupTexture, SSAO_Sampler, frac(UV)).xyz * 2 - 1;
-	#elif COMPUTE_SHADER || FORWARD_SHADING || SHADING_PATH_MOBILE
+	#elif COMPUTE_SHADER || FORWARD_SHADING || (SHADING_PATH_MOBILE && (!MOBILE_DEFERRED_SHADING || ALLOW_FRAMEBUFFER_FETCH))
 		// Async compute and forward shading don't have access to the gbuffer.
+		// for mobile case only multi-pass defferred can access gbuffer
 		WorldNormal = ReconstructNormalFromDepthBuffer(SvPosition);
 	#else
 #if SUBSTRATE_GBUFFER_FORMAT==1
@@ -118,9 +120,13 @@ float3 GetWorldSpaceNormalFromAOInput(float2 UV, float4 SvPosition)
 #else
 		// Otherwise sample Gbuffer normals if the shader model has normals.
 		FGBufferData GBuffer = GetGBufferData(UV, false);
-		if( GBuffer.ShadingModelID != SHADINGMODELID_HAIR )
+		bool bUseNormal = (GBuffer.ShadingModelID != SHADINGMODELID_HAIR);
+		#if	USE_SM_EXCLUDE_MASK
+			bUseNormal = bUseNormal & (((1u << GBuffer.ShadingModelID) & SSAO_ShadingModelsExcludeMask) == 0u);
+		#endif
+		if (bUseNormal)
 		{
-			WorldNormal = GetGBufferData(UV, false).WorldNormal;
+			WorldNormal = GBuffer.WorldNormal;
 		}
 #endif
 	#endif
@@ -388,3 +394,8 @@ float GetDepthFromAOInput(float2 UV)
 	return CalcSceneDepth(UV);
 #endif
 }
+
+float2 UVToScreenPos(float2 UV, float2 Extent)
+{
+	return (UV * Extent - AOViewport_ScreenPosToViewportBias) / AOViewport_ScreenPosToViewportScale;
+}
diff --git a/Engine/Shaders/Private/PostProcessAmbientOcclusionMobile.usf b/Engine/Shaders/Private/PostProcessAmbientOcclusionMobile.usf
deleted file mode 100644
diff --git a/Engine/Source/Runtime/Renderer/Private/CompositionLighting/CompositionLighting.cpp b/Engine/Source/Runtime/Renderer/Private/CompositionLighting/CompositionLighting.cpp
--- a/Engine/Source/Runtime/Renderer/Private/CompositionLighting/CompositionLighting.cpp
+++ b/Engine/Source/Runtime/Renderer/Private/CompositionLighting/CompositionLighting.cpp
@@ -94,12 +94,12 @@ bool ShouldRenderScreenSpaceAmbientOcclusion(const FViewInfo& View, bool bLumenW
 
 static ESSAOType GetDownscaleSSAOType(const FViewInfo& View)
 {
-	return FSSAOHelper::IsAmbientOcclusionCompute(View.GetFeatureLevel()) ? ESSAOType::ECS : ESSAOType::EPS;
+	return FSSAOHelper::IsAmbientOcclusionCompute() ? ESSAOType::ECS : ESSAOType::EPS;
 }
 
 static ESSAOType GetFullscreenSSAOType(const FViewInfo& View, uint32 Levels)
 {
-	if (FSSAOHelper::IsAmbientOcclusionCompute(View.GetFeatureLevel()))
+	if (FSSAOHelper::IsAmbientOcclusionCompute())
 	{
 		if (FSSAOHelper::IsAmbientOcclusionAsyncCompute(View, Levels))
 		{
@@ -112,17 +112,16 @@ static ESSAOType GetFullscreenSSAOType(const FViewInfo& View, uint32 Levels)
 	return ESSAOType::EPS;
 }
 
-static FSSAOCommonParameters GetSSAOCommonParameters(
+FSSAOCommonParameters GetSSAOCommonParameters(
 	FRDGBuilder& GraphBuilder,
 	const FViewInfo& View,
-	TRDGUniformBufferRef<FSceneTextureUniformParameters> SceneTexturesUniformBuffer,
 	uint32 Levels,
 	bool bAllowGBufferRead)
 {
-	const FSceneTextureParameters SceneTextureParameters = GetSceneTextureParameters(GraphBuilder, SceneTexturesUniformBuffer);
+	const FSceneTextureParameters SceneTextureParameters = GetSceneTextureParameters(GraphBuilder, View);
 
 	FSSAOCommonParameters CommonParameters;
-	CommonParameters.SceneTexturesUniformBuffer = SceneTexturesUniformBuffer;
+	CommonParameters.SceneTextures = GetSceneTextureShaderParameters(View);
 	CommonParameters.SceneTexturesViewport = FScreenPassTextureViewport(SceneTextureParameters.SceneDepthTexture, View.ViewRect);
 
 	CommonParameters.GBufferA = bAllowGBufferRead ? FScreenPassTexture(SceneTextureParameters.GBufferATexture, View.ViewRect) : FScreenPassTexture();
@@ -143,14 +142,13 @@ static FSSAOCommonParameters GetSSAOCommonParameters(
 FGTAOCommonParameters GetGTAOCommonParameters(
 	FRDGBuilder& GraphBuilder,
 	const FViewInfo& View,
-	TRDGUniformBufferRef<FSceneTextureUniformParameters> SceneTexturesUniformBuffer,
 	EGTAOType GTAOType
 	)
 {
-	const FSceneTextureParameters SceneTextureParameters = GetSceneTextureParameters(GraphBuilder, SceneTexturesUniformBuffer);
+	const FSceneTextureParameters SceneTextureParameters = GetSceneTextureParameters(GraphBuilder, View);
 
 	FGTAOCommonParameters CommonParameters;
-	CommonParameters.SceneTexturesUniformBuffer = SceneTexturesUniformBuffer;
+	CommonParameters.SceneTextures = GetSceneTextureShaderParameters(View);
 	CommonParameters.SceneTexturesViewport = FScreenPassTextureViewport(SceneTextureParameters.SceneDepthTexture, View.ViewRect);
 
 	CommonParameters.SceneDepth = FScreenPassTexture(SceneTextureParameters.SceneDepthTexture, View.ViewRect);
@@ -585,7 +583,7 @@ void FCompositionLighting::ProcessBeforeBasePass(FRDGBuilder& GraphBuilder, FDBu
 
 		if (bEnableSSAO)
 		{
-			FSSAOCommonParameters Parameters = GetSSAOCommonParameters(GraphBuilder, View, SceneTextures.UniformBuffer, ViewConfig.Levels, false);
+			FSSAOCommonParameters Parameters = GetSSAOCommonParameters(GraphBuilder, View, ViewConfig.Levels, false);
 			FScreenPassRenderTarget FinalTarget = FScreenPassRenderTarget(SceneTextures.ScreenSpaceAO, View.ViewRect, ERenderTargetLoadAction::ENoAction);
 
 			AddPostProcessingAmbientOcclusion(
@@ -658,7 +656,7 @@ void FCompositionLighting::ProcessAfterBasePass(FRDGBuilder& GraphBuilder, FInst
 				{
 					check(HorizonsTexture);
 
-					FGTAOCommonParameters Parameters = GetGTAOCommonParameters(GraphBuilder, View, SceneTextures.UniformBuffer, ViewConfig.GTAOType);
+					FGTAOCommonParameters Parameters = GetGTAOCommonParameters(GraphBuilder, View, ViewConfig.GTAOType);
 
 					FScreenPassTexture GTAOHorizons(HorizonsTexture, Parameters.DownsampledViewRect);
 					AmbientOcclusion = AddPostProcessingGTAOPostAsync(GraphBuilder, View, Parameters, GTAOHorizons, FinalTarget);
@@ -671,12 +669,12 @@ void FCompositionLighting::ProcessAfterBasePass(FRDGBuilder& GraphBuilder, FInst
 				{
 					if (ViewConfig.GTAOType == EGTAOType::ENonAsync)
 					{
-						FGTAOCommonParameters Parameters = GetGTAOCommonParameters(GraphBuilder, View, SceneTextures.UniformBuffer, ViewConfig.GTAOType);
+						FGTAOCommonParameters Parameters = GetGTAOCommonParameters(GraphBuilder, View, ViewConfig.GTAOType);
 						AmbientOcclusion = AddPostProcessingGTAOAllPasses(GraphBuilder, View, Parameters, FinalTarget);
 					}
 					else if (ViewConfig.SSAOLocation == ESSAOLocation::AfterBasePass)
 					{
-						FSSAOCommonParameters Parameters = GetSSAOCommonParameters(GraphBuilder, View, SceneTextures.UniformBuffer, ViewConfig.Levels, true);
+						FSSAOCommonParameters Parameters = GetSSAOCommonParameters(GraphBuilder, View, ViewConfig.Levels, true);
 						AmbientOcclusion = AddPostProcessingAmbientOcclusion(GraphBuilder, View, Parameters, FinalTarget);
 					}
 
@@ -712,7 +710,7 @@ void FCompositionLighting::ProcessAfterOcclusion(FRDGBuilder& GraphBuilder)
 		{
 			RDG_GPU_MASK_SCOPE(GraphBuilder, View.GPUMask);
 
-			FGTAOCommonParameters CommonParameters = GetGTAOCommonParameters(GraphBuilder, View, SceneTextures.UniformBuffer, ViewConfig.GTAOType);
+			FGTAOCommonParameters CommonParameters = GetGTAOCommonParameters(GraphBuilder, View, ViewConfig.GTAOType);
 
 			const ERenderTargetLoadAction LoadAction = View.DecayLoadAction(ERenderTargetLoadAction::ENoAction);
 
diff --git a/Engine/Source/Runtime/Renderer/Private/CompositionLighting/PostProcessAmbientOcclusion.cpp b/Engine/Source/Runtime/Renderer/Private/CompositionLighting/PostProcessAmbientOcclusion.cpp
--- a/Engine/Source/Runtime/Renderer/Private/CompositionLighting/PostProcessAmbientOcclusion.cpp
+++ b/Engine/Source/Runtime/Renderer/Private/CompositionLighting/PostProcessAmbientOcclusion.cpp
@@ -86,6 +86,13 @@ static TAutoConsoleVariable<int32> CVarAmbientOcclusionDepthBoundsTest(
 	TEXT("Whether to use depth bounds test to cull distant pixels during AO pass. This option is only valid when pixel shader path is used (r.AmbientOcclusion.Compute=0), without upsampling."),
 	ECVF_RenderThreadSafe);
 
+static TAutoConsoleVariable<int32> CVarAmbientOcclusionFramebufferCompression(
+	TEXT("r.AmbientOcclusion.FramebufferCompression"),
+	0,
+	TEXT("Whether to use framebuffer compression for AO texture if platform supports it.\n"),
+	ECVF_RenderThreadSafe
+);
+
 static TAutoConsoleVariable<int32> CVarAmbientOcclusionMethod(
 	TEXT("r.AmbientOcclusion.Method"),
 	0,
@@ -178,9 +185,9 @@ int32 FSSAOHelper::GetAmbientOcclusionShaderLevel(const FSceneView& View)
 		(QualityPercent > 5.0f);
 }
 
-bool FSSAOHelper::IsAmbientOcclusionCompute(const ERHIFeatureLevel::Type FeatureLevel)
+bool FSSAOHelper::IsAmbientOcclusionCompute()
 {
-	return FeatureLevel >= ERHIFeatureLevel::SM5 && CVarAmbientOcclusionCompute.GetValueOnAnyThread() >= 1;
+	return CVarAmbientOcclusionCompute.GetValueOnAnyThread() >= 1;
 }
 
 int32 FSSAOHelper::GetNumAmbientOcclusionLevels()
@@ -210,7 +217,7 @@ bool FSSAOHelper::IsAmbientOcclusionAsyncCompute(const FViewInfo& View, uint32 A
 {
 	// if AsyncCompute is feasible
 	// only single level is allowed.  more levels end up reading from gbuffer normals atm which is not allowed.
-	if(IsAmbientOcclusionCompute(View.GetFeatureLevel()) && (AOPassCount == 1))
+	if(IsAmbientOcclusionCompute() && (AOPassCount == 1))
 	{
 		int32 ComputeCVar = CVarAmbientOcclusionCompute.GetValueOnRenderThread();
 
@@ -241,7 +248,7 @@ uint32 FSSAOHelper::ComputeAmbientOcclusionPassCount(const FViewInfo& View)
 	{
 		int32 CVarLevel = GetNumAmbientOcclusionLevels();
 
-		if (IsAmbientOcclusionCompute(View.GetFeatureLevel()) || IsForwardShadingEnabled(View.GetShaderPlatform()))
+		if (IsAmbientOcclusionCompute() || IsForwardShadingEnabled(View.GetShaderPlatform()))
 		{	
 			if (CVarLevel<0)
 			{
@@ -307,10 +314,15 @@ EGTAOType FSSAOHelper::GetGTAOPassType(const FViewInfo& View, uint32 Levels)
 	return EGTAOType::EOff;
 }
 
+bool FSSAOHelper::IsDepthBoundsTestEnabled()
+{
+	return GSupportsDepthBoundsTest && CVarAmbientOcclusionDepthBoundsTest.GetValueOnRenderThread() != 0;
+}
+
 FRDGTextureDesc GetScreenSpaceAOTextureDesc(ERHIFeatureLevel::Type FeatureLevel, FIntPoint Extent)
 {
 	ETextureCreateFlags TextureCreateFlags = TexCreate_UAV | TexCreate_RenderTargetable | TexCreate_ShaderResource | GFastVRamConfig.ScreenSpaceAO;
-	if (FSSAOHelper::IsAmbientOcclusionCompute(FeatureLevel))
+	if (FSSAOHelper::IsAmbientOcclusionCompute())
 	{
 		TextureCreateFlags |= TexCreate_NoFastClear;
 	}
@@ -347,7 +359,7 @@ static FSSAOShaderParameters GetSSAOShaderParameters(
 {
 	const FFinalPostProcessSettings& Settings = View.FinalPostProcessSettings;
 
-	const FIntPoint RandomizationSize = GSystemTextures.SSAORandomization->GetDesc().Extent;
+	const FIntPoint RandomizationSize = GSystemTextures.SSAORandomization->GetDesc().Extent * InputViewport.Extent / OutputViewport.Extent;
 	FVector2D ViewportUVToRandomUV(InputViewport.Extent.X / (float)RandomizationSize.X, InputViewport.Extent.Y / (float)RandomizationSize.Y);
 
 	// e.g. 4 means the input texture is 4x smaller than the buffer size
@@ -516,14 +528,9 @@ public:
 	DECLARE_GLOBAL_SHADER(FAmbientOcclusionSetupPS);
 	SHADER_USE_PARAMETER_STRUCT(FAmbientOcclusionSetupPS, FGlobalShader);
 
-	static bool ShouldCompilePermutation(const FGlobalShaderPermutationParameters& Parameters)
-	{
-		return IsFeatureLevelSupported(Parameters.Platform, ERHIFeatureLevel::SM5);
-	}
-
 	BEGIN_SHADER_PARAMETER_STRUCT(FParameters, )
 		SHADER_PARAMETER_STRUCT_INCLUDE(FViewShaderParameters, View)
-		SHADER_PARAMETER_RDG_UNIFORM_BUFFER(FSceneTextureUniformParameters, SceneTextures)
+		SHADER_PARAMETER_STRUCT_INCLUDE(FSceneTextureShaderParameters, SceneTextures)
 
 		SHADER_PARAMETER_STRUCT_INCLUDE(FSSAOShaderParameters, SSAOParameters)
 
@@ -571,7 +578,7 @@ FScreenPassTexture AddAmbientOcclusionSetupPass(
 
 	FAmbientOcclusionSetupPS::FParameters* PassParameters = GraphBuilder.AllocParameters<FAmbientOcclusionSetupPS::FParameters>();
 	PassParameters->View = View.GetShaderParameters();
-	PassParameters->SceneTextures = CommonParameters.SceneTexturesUniformBuffer;
+	PassParameters->SceneTextures = CommonParameters.SceneTextures;
 	PassParameters->SSAOParameters = GetSSAOShaderParameters(View, InputViewport, OutputViewport, CommonParameters.SceneTexturesViewport, EAOTechnique::SSAO);
 	PassParameters->ThresholdInverse = ThresholdInverseValue;
 	PassParameters->InputExtentInverse = FVector2f(1.0f) / FVector2f(InputViewport.Extent);
@@ -609,11 +616,6 @@ public:
 		OutEnvironment.SetDefine(TEXT("THREADGROUP_SIZEY"), 8);
 	}
 
-	static bool ShouldCompilePermutation(const FGlobalShaderPermutationParameters& Parameters)
-	{
-		return IsFeatureLevelSupported(Parameters.Platform, ERHIFeatureLevel::SM5);
-	}
-
 	BEGIN_SHADER_PARAMETER_STRUCT(FParameters, )
 		SHADER_PARAMETER_STRUCT(FScreenPassTextureViewportParameters, SSAOSmoothOutputViewport)
 		SHADER_PARAMETER(FScreenTransform, SSAOSmoothOutputToInput)
@@ -626,6 +628,31 @@ public:
 };
 IMPLEMENT_GLOBAL_SHADER(FAmbientOcclusionSmoothCS, "/Engine/Private/PostProcessAmbientOcclusion.usf", "MainSSAOSmoothCS", SF_Compute);
 
+
+class FAmbientOcclusionSmoothPS : public FGlobalShader
+{
+public:
+	DECLARE_GLOBAL_SHADER(FAmbientOcclusionSmoothPS);
+	SHADER_USE_PARAMETER_STRUCT(FAmbientOcclusionSmoothPS, FGlobalShader);
+
+	static void ModifyCompilationEnvironment(const FGlobalShaderPermutationParameters& Parameters, FShaderCompilerEnvironment& OutEnvironment)
+	{
+		FGlobalShader::ModifyCompilationEnvironment(Parameters, OutEnvironment);
+		OutEnvironment.SetDefine(TEXT("COMPUTE_SHADER"), 0);
+	}
+
+	BEGIN_SHADER_PARAMETER_STRUCT(FParameters, )
+		SHADER_PARAMETER_STRUCT(FScreenPassTextureViewportParameters, SSAOSmoothOutputViewport)
+		SHADER_PARAMETER(FScreenTransform, SSAOSmoothOutputToInput)
+
+		SHADER_PARAMETER_RDG_TEXTURE(Texture2D, SSAOSmoothInputTexture)
+		SHADER_PARAMETER_SAMPLER(SamplerState, SSAOSmoothInputSampler)
+		
+		RENDER_TARGET_BINDING_SLOTS()
+	END_SHADER_PARAMETER_STRUCT();
+};
+IMPLEMENT_GLOBAL_SHADER(FAmbientOcclusionSmoothPS, "/Engine/Private/PostProcessAmbientOcclusion.usf", "MainSSAOSmoothPS", SF_Pixel);
+
 void AddAmbientOcclusionSmoothPass(
 	FRDGBuilder& GraphBuilder,
 	const FViewInfo& View,
@@ -641,22 +668,45 @@ void AddAmbientOcclusionSmoothPass(
 	const FScreenPassTextureViewportParameters InputViewportParameters = GetScreenPassTextureViewportParameters(InputViewport);
 	const FScreenPassTextureViewportParameters OutputViewportParameters = GetScreenPassTextureViewportParameters(OutputViewport);
 
- 	FAmbientOcclusionSmoothCS::FParameters* PassParameters = GraphBuilder.AllocParameters<FAmbientOcclusionSmoothCS::FParameters>();
-	PassParameters->SSAOSmoothOutputViewport = OutputViewportParameters;
-	PassParameters->SSAOSmoothOutputToInput = FScreenTransform::ChangeTextureUVCoordinateFromTo(OutputViewport, InputViewport);
-	PassParameters->SSAOSmoothInputTexture = Input.Texture;
-	PassParameters->SSAOSmoothInputSampler = TStaticSamplerState<SF_Bilinear, AM_Clamp, AM_Clamp, AM_Clamp>::GetRHI();
-	PassParameters->SSAOSmoothOutputTexture = GraphBuilder.CreateUAV(Output.Texture);
+	if (SSAOType == ESSAOType::EPS)
+	{
+		FAmbientOcclusionSmoothPS::FParameters* PassParameters = GraphBuilder.AllocParameters<FAmbientOcclusionSmoothPS::FParameters>();
+		PassParameters->SSAOSmoothOutputViewport = OutputViewportParameters;
+		PassParameters->SSAOSmoothOutputToInput = FScreenTransform::ChangeTextureUVCoordinateFromTo(OutputViewport, InputViewport);
+		PassParameters->SSAOSmoothInputTexture = Input.Texture;
+		PassParameters->SSAOSmoothInputSampler = TStaticSamplerState<SF_Bilinear, AM_Clamp, AM_Clamp, AM_Clamp>::GetRHI();
+		PassParameters->RenderTargets[0] = Output.GetRenderTargetBinding();
+
+		TShaderMapRef<FAmbientOcclusionSmoothPS> PixelShader(View.ShaderMap);
+
+		AddDrawScreenPass(
+			GraphBuilder,
+			RDG_EVENT_NAME("AmbientOcclusionSmooth %dx%d", OutputViewport.Rect.Width(), OutputViewport.Rect.Height()),
+			View,
+			OutputViewport,
+			InputViewport,
+			PixelShader,
+			PassParameters);
+	}
+	else
+	{
+		FAmbientOcclusionSmoothCS::FParameters* PassParameters = GraphBuilder.AllocParameters<FAmbientOcclusionSmoothCS::FParameters>();
+		PassParameters->SSAOSmoothOutputViewport = OutputViewportParameters;
+		PassParameters->SSAOSmoothOutputToInput = FScreenTransform::ChangeTextureUVCoordinateFromTo(OutputViewport, InputViewport);
+		PassParameters->SSAOSmoothInputTexture = Input.Texture;
+		PassParameters->SSAOSmoothInputSampler = TStaticSamplerState<SF_Bilinear, AM_Clamp, AM_Clamp, AM_Clamp>::GetRHI();
+		PassParameters->SSAOSmoothOutputTexture = GraphBuilder.CreateUAV(Output.Texture);
 
-	TShaderMapRef<FAmbientOcclusionSmoothCS> ComputeShader(View.ShaderMap);
+		TShaderMapRef<FAmbientOcclusionSmoothCS> ComputeShader(View.ShaderMap);
 
-	FComputeShaderUtils::AddPass(
-		GraphBuilder,
-		RDG_EVENT_NAME("AmbientOcclusionSmooth %dx%d", OutputViewport.Rect.Width(), OutputViewport.Rect.Height()),
-		SSAOType == ESSAOType::EAsyncCS ? ERDGPassFlags::AsyncCompute : ERDGPassFlags::Compute,
-		ComputeShader,
-		PassParameters,
-		FComputeShaderUtils::GetGroupCount(OutputViewport.Rect.Size(), 8));
+		FComputeShaderUtils::AddPass(
+			GraphBuilder,
+			RDG_EVENT_NAME("AmbientOcclusionSmooth %dx%d", OutputViewport.Rect.Width(), OutputViewport.Rect.Height()),
+			SSAOType == ESSAOType::EAsyncCS ? ERDGPassFlags::AsyncCompute : ERDGPassFlags::Compute,
+			ComputeShader,
+			PassParameters,
+			FComputeShaderUtils::GetGroupCount(OutputViewport.Rect.Size(), 8));
+	}
 }
 
 //----------------------------------------------------------------------------------------------------------------------
@@ -670,7 +720,7 @@ void AddAmbientOcclusionSmoothPass(
 
 BEGIN_SHADER_PARAMETER_STRUCT(FAmbientOcclusionParameters, )
 	SHADER_PARAMETER_STRUCT_INCLUDE(FViewShaderParameters, View)
-	SHADER_PARAMETER_RDG_UNIFORM_BUFFER(FSceneTextureUniformParameters, SceneTextures)
+	SHADER_PARAMETER_STRUCT_INCLUDE(FSceneTextureShaderParameters, SceneTextures)
 
 	SHADER_PARAMETER_STRUCT_INCLUDE(FHZBParameters, HZBParameters)
 	SHADER_PARAMETER_STRUCT_INCLUDE(FSSAOShaderParameters, SSAOParameters)
@@ -679,6 +729,7 @@ BEGIN_SHADER_PARAMETER_STRUCT(FAmbientOcclusionParameters, )
 	SHADER_PARAMETER(FVector2f, SSAO_DownsampledAOUVViewportMin)
 	SHADER_PARAMETER(FVector2f, SSAO_DownsampledAOUVViewportMax)
 	SHADER_PARAMETER(FVector2f, SSAO_SvPositionScaleBias)
+	SHADER_PARAMETER(uint32, SSAO_ShadingModelsExcludeMask)
 
 	SHADER_PARAMETER_RDG_TEXTURE(Texture2D, SSAO_SetupTexture)
 	SHADER_PARAMETER_RDG_TEXTURE(Texture2D, SSAO_NormalsTexture)
@@ -690,17 +741,65 @@ BEGIN_SHADER_PARAMETER_STRUCT(FAmbientOcclusionParameters, )
 	SHADER_PARAMETER_SAMPLER(SamplerState, RandomNormalTextureSampler)
 END_SHADER_PARAMETER_STRUCT();
 
+struct FAmbientOcclusionPermutations
+{
+	class FUseUpsampleDim		: SHADER_PERMUTATION_BOOL("USE_UPSAMPLE");
+	class FUseAoSetupAsInputDim : SHADER_PERMUTATION_BOOL("USE_AO_SETUP_AS_INPUT");
+	class FOutputDepthDim		: SHADER_PERMUTATION_BOOL("OUTPUT_DEPTH");
+	class FUseSMExcludeMaskDim	: SHADER_PERMUTATION_BOOL("USE_SM_EXCLUDE_MASK");
+	class FShaderQualityDim		: SHADER_PERMUTATION_INT("SHADER_QUALITY", 5);
+
+	using FPermutationDomain = TShaderPermutationDomain<FUseUpsampleDim, FUseAoSetupAsInputDim, FUseSMExcludeMaskDim, FOutputDepthDim, FShaderQualityDim>;
+
+	static bool ShouldCompilePermutation(const FGlobalShaderPermutationParameters& Parameters, EShaderFrequency Frequency)
+	{
+		FPermutationDomain PermutationVector(Parameters.PermutationId);
+
+		if (!IsMobilePlatform(Parameters.Platform))
+		{
+			// Mobile specific features
+			if (PermutationVector.Get<FUseSMExcludeMaskDim>() || PermutationVector.Get<FOutputDepthDim>())
+			{
+				return false;
+			}
+		}
+
+		// SM mask requires access to GBuffer, only pixel shader version can do it
+		const bool bSMAccess = !IsMobilePlatform(Parameters.Platform) || (IsMobileDeferredShadingEnabled(Parameters.Platform) && MobileUsesGBufferCustomData(Parameters.Platform));
+		if (PermutationVector.Get<FUseSMExcludeMaskDim>() && (Frequency != SF_Pixel || !bSMAccess))
+		{
+			return false;
+		}
+
+		// OUTPUT_DEPTH is only used in a specific setup with a bilateral upsample
+		if (PermutationVector.Get<FOutputDepthDim>() &&
+			(PermutationVector.Get<FUseUpsampleDim>() || PermutationVector.Get<FUseAoSetupAsInputDim>()))
+		{
+			return false;
+		}
+		
+		return true;
+	}
+
+	static FPermutationDomain BuildPermutationVector(bool bDoUpsample, bool bAOSetupAsInput, int32 ShaderQuality, bool bOutputDepth, uint32 SMExcludeMask)
+	{
+		FPermutationDomain PermutationVector;
+		PermutationVector.Set<FUseUpsampleDim>(bDoUpsample);
+		PermutationVector.Set<FUseAoSetupAsInputDim>(bAOSetupAsInput);
+		PermutationVector.Set<FShaderQualityDim>(ShaderQuality);
+		PermutationVector.Set<FOutputDepthDim>(bOutputDepth);
+		PermutationVector.Set<FUseSMExcludeMaskDim>(SMExcludeMask != 0);
+		return PermutationVector;
+	}
+};
+
 class FAmbientOcclusionPS : public FGlobalShader
 {
 public:
 	DECLARE_GLOBAL_SHADER(FAmbientOcclusionPS);
 	SHADER_USE_PARAMETER_STRUCT(FAmbientOcclusionPS, FGlobalShader);
 
-	class FUseUpsampleDim       : SHADER_PERMUTATION_BOOL("USE_UPSAMPLE");
-	class FUseAoSetupAsInputDim : SHADER_PERMUTATION_BOOL("USE_AO_SETUP_AS_INPUT");
-	class FShaderQualityDim     : SHADER_PERMUTATION_INT("SHADER_QUALITY", 5);
-
-	using FPermutationDomain = TShaderPermutationDomain<FUseUpsampleDim, FUseAoSetupAsInputDim, FShaderQualityDim>;
+	using FPermutationDomain = FAmbientOcclusionPermutations::FPermutationDomain;
 
 	static void ModifyCompilationEnvironment(const FGlobalShaderPermutationParameters& Parameters, FShaderCompilerEnvironment& OutEnvironment)
 	{
@@ -710,7 +809,7 @@ public:
 
 	static bool ShouldCompilePermutation(const FGlobalShaderPermutationParameters& Parameters)
 	{
-		return IsFeatureLevelSupported(Parameters.Platform, ERHIFeatureLevel::SM5);
+		return FAmbientOcclusionPermutations::ShouldCompilePermutation(Parameters, SF_Pixel);
 	}
 
 	BEGIN_SHADER_PARAMETER_STRUCT(FParameters, )
@@ -727,11 +826,7 @@ public:
 	DECLARE_GLOBAL_SHADER(FAmbientOcclusionCS);
 	SHADER_USE_PARAMETER_STRUCT(FAmbientOcclusionCS, FGlobalShader);
 
-	class FUseUpsampleDim       : SHADER_PERMUTATION_BOOL("USE_UPSAMPLE");
-	class FUseAoSetupAsInputDim : SHADER_PERMUTATION_BOOL("USE_AO_SETUP_AS_INPUT");
-	class FShaderQualityDim     : SHADER_PERMUTATION_INT("SHADER_QUALITY", 5);
-
-	using FPermutationDomain = TShaderPermutationDomain<FUseUpsampleDim, FUseAoSetupAsInputDim, FShaderQualityDim>;
+	using FPermutationDomain = FAmbientOcclusionPermutations::FPermutationDomain;
 
 	static void ModifyCompilationEnvironment(const FGlobalShaderPermutationParameters& Parameters, FShaderCompilerEnvironment& OutEnvironment)
 	{
@@ -743,7 +838,7 @@ public:
 
 	static bool ShouldCompilePermutation(const FGlobalShaderPermutationParameters& Parameters)
 	{
-		return IsFeatureLevelSupported(Parameters.Platform, ERHIFeatureLevel::SM5);
+		return FAmbientOcclusionPermutations::ShouldCompilePermutation(Parameters, SF_Compute);
 	}
 
 	BEGIN_SHADER_PARAMETER_STRUCT(FParameters, )
@@ -785,6 +880,12 @@ FScreenPassRenderTarget CreateAmbientOcclusionOutputTarget(
 		OutputDesc.Format = IntermediateFormatOverride;
 	}
 
+	if (UE::PixelFormat::HasCapabilities(OutputDesc.Format, EPixelFormatCapabilities::LossyCompressible) 
+		&& CVarAmbientOcclusionFramebufferCompression.GetValueOnRenderThread() != 0)
+	{
+		OutputDesc.Flags |= TexCreate_LossyCompressionLowBitrate;
+	}
+
 	return FScreenPassRenderTarget(GraphBuilder.CreateTexture(OutputDesc, TEXT("AmbientOcclusion")), ViewRect, ERenderTargetLoadAction::ENoAction);
 }
 
@@ -808,11 +909,17 @@ void AddAmbientOcclusionPass(
 
 	const FScreenPassTextureViewport OutputViewport(Output);
 
-	const bool bDoUpsample = DownsampledAO.IsValid();
-
+	bool bDoUpsample = DownsampledAO.IsValid();
+	if (CommonParameters.MobileHalfResolutionSetting > 0)
+	{
+		// these are not supported with half-res SSAO
+		bDoUpsample = false;
+		bAOSetupAsInput = false;
+	}
+	
 	FAmbientOcclusionParameters SharedParameters;
 	SharedParameters.View = View.GetShaderParameters();
-	SharedParameters.SceneTextures = CommonParameters.SceneTexturesUniformBuffer;
+	SharedParameters.SceneTextures = CommonParameters.SceneTextures;
 	SharedParameters.HZBParameters = GetHZBParametersForAO(GraphBuilder, View, CommonParameters.SceneTexturesViewport.Extent, EAOTechnique::SSAO);
 	SharedParameters.SSAOParameters = GetSSAOShaderParameters(View, InputViewport, OutputViewport, CommonParameters.SceneTexturesViewport, EAOTechnique::SSAO);
 
@@ -848,8 +955,18 @@ void AddAmbientOcclusionPass(
 		SharedParameters.SSAO_DownsampledAOUVViewportMax = FVector2f(1.0f, 1.0f);
 	}
 
-	SharedParameters.SSAO_SvPositionScaleBias = FVector2f(1, 0);
+	if (CommonParameters.MobileHalfResolutionSetting > 0)
+	{
+		// Downscaled 
+		SharedParameters.SSAO_SvPositionScaleBias = FVector2f(2, -0.5);
+	}
+	else
+	{
+		SharedParameters.SSAO_SvPositionScaleBias = FVector2f(1, 0);
+	}
 	
+	SharedParameters.SSAO_ShadingModelsExcludeMask = CommonParameters.MobileShadingModelsExcludeMask;
+			
 	SharedParameters.SSAO_Sampler = TStaticSamplerState<SF_Point, AM_Clamp, AM_Clamp, AM_Clamp>::GetRHI();
 
 	SharedParameters.RandomNormalTexture = GSystemTextures.SSAORandomization->GetRHI();
@@ -866,10 +983,13 @@ void AddAmbientOcclusionPass(
 		PassParameters->Substrate = Substrate::BindSubstrateGlobalUniformParameters(View);
 		PassParameters->OutTexture = GraphBuilder.CreateUAV(Output.Texture);
 
-		FAmbientOcclusionCS::FPermutationDomain PermutationVector;
-		PermutationVector.Set<FAmbientOcclusionCS::FUseUpsampleDim>(bDoUpsample);
-		PermutationVector.Set<FAmbientOcclusionCS::FUseAoSetupAsInputDim>(bAOSetupAsInput);
-		PermutationVector.Set<FAmbientOcclusionCS::FShaderQualityDim>(CommonParameters.ShaderQuality);
+		FAmbientOcclusionCS::FPermutationDomain PermutationVector = FAmbientOcclusionPermutations::BuildPermutationVector(
+			bDoUpsample, 
+			bAOSetupAsInput, 
+			CommonParameters.ShaderQuality, 
+			CommonParameters.MobileHalfResolutionSetting > 1,
+			0// SMExcludeMask
+		);
 
 		TShaderMapRef<FAmbientOcclusionCS> ComputeShader(View.ShaderMap, PermutationVector);
 		FComputeShaderUtils::AddPass(
@@ -882,21 +1002,21 @@ void AddAmbientOcclusionPass(
 	}
 	else
 	{
-		const uint32 ScaleToFullRes = CommonParameters.SceneTexturesViewport.Extent.X / InputViewport.Extent.X;
+		const bool bFullRes = (CommonParameters.SceneTexturesViewport == OutputViewport);
 		const bool bDepthBoundsTestEnabled =
 			!bDoUpsample
-			&& CVarAmbientOcclusionDepthBoundsTest.GetValueOnRenderThread()
-			&& ScaleToFullRes == 1
-			&& GSupportsDepthBoundsTest
+			&& bFullRes
+			&& FSSAOHelper::IsDepthBoundsTestEnabled()
 			&& CommonParameters.SceneDepth.IsValid()
 			&& CommonParameters.SceneDepth.Texture->Desc.NumSamples == 1;
 
-		FDepthStencilBinding DepthStencilBinding(CommonParameters.SceneDepth.Texture, ERenderTargetLoadAction::ELoad, ERenderTargetLoadAction::ELoad, FExclusiveDepthStencil::DepthRead_StencilWrite);
-
+		FDepthStencilBinding DepthStencilBinding{};
 		float DepthFar = 0.0f;
 
 		if (bDepthBoundsTestEnabled)
 		{
+			DepthStencilBinding = FDepthStencilBinding(CommonParameters.SceneDepth.Texture, ERenderTargetLoadAction::ELoad, ERenderTargetLoadAction::ELoad, FExclusiveDepthStencil::DepthRead_StencilWrite);
+			
 			const FFinalPostProcessSettings& Settings = View.FinalPostProcessSettings;
 			const FMatrix& ProjectionMatrix = View.ViewMatrices.GetViewToClip();
 			const FVector4f Far = (FVector4f)ProjectionMatrix.TransformFVector4(FVector4(0, 0, Settings.AmbientOcclusionFadeDistance));
@@ -949,10 +1069,13 @@ void AddAmbientOcclusionPass(
 			PassParameters->RenderTargets.DepthStencil = DepthStencilBinding;
 		}
 		
-		FAmbientOcclusionPS::FPermutationDomain PermutationVector;
-		PermutationVector.Set<FAmbientOcclusionPS::FUseUpsampleDim>(bDoUpsample);
-		PermutationVector.Set<FAmbientOcclusionPS::FUseAoSetupAsInputDim>(bAOSetupAsInput);
-		PermutationVector.Set<FAmbientOcclusionPS::FShaderQualityDim>(CommonParameters.ShaderQuality);
+		FAmbientOcclusionPS::FPermutationDomain PermutationVector = FAmbientOcclusionPermutations::BuildPermutationVector(
+			bDoUpsample, 
+			bAOSetupAsInput, 
+			CommonParameters.ShaderQuality, 
+			CommonParameters.MobileHalfResolutionSetting > 1,
+			CommonParameters.MobileShadingModelsExcludeMask
+		);
 
 		TShaderMapRef<FAmbientOcclusionPS> PixelShader(View.ShaderMap, PermutationVector);
 		TShaderMapRef<FScreenPassVS> VertexShader(View.ShaderMap);
@@ -1102,7 +1225,7 @@ public:
 
 	BEGIN_SHADER_PARAMETER_STRUCT(FParameters, )
 		SHADER_PARAMETER_STRUCT_REF(FViewUniformShaderParameters, View)
-		SHADER_PARAMETER_RDG_UNIFORM_BUFFER(FSceneTextureUniformParameters, SceneTextures)
+		SHADER_PARAMETER_STRUCT_INCLUDE(FSceneTextureShaderParameters, SceneTextures)
 		SHADER_PARAMETER_STRUCT_INCLUDE(FHZBParameters, HZBParameters)
 		SHADER_PARAMETER_RDG_UNIFORM_BUFFER(FSubstrateGlobalUniformParameters, Substrate)
 
@@ -1146,7 +1269,7 @@ FGTAOHorizonSearchOutputs AddGTAOHorizonSearchIntegratePass(
 
 	FGTAOHorizonSearchAndIntegrateCS::FParameters* PassParameters = GraphBuilder.AllocParameters<FGTAOHorizonSearchAndIntegrateCS::FParameters>();
 	PassParameters->View = View.ViewUniformBuffer;
-	PassParameters->SceneTextures = CommonParameters.SceneTexturesUniformBuffer;
+	PassParameters->SceneTextures = CommonParameters.SceneTextures;
 	PassParameters->HZBParameters = GetHZBParametersForAO(GraphBuilder, View, SceneViewport.Extent, EAOTechnique::GTAO);
 	PassParameters->Substrate = Substrate::BindSubstrateGlobalUniformParameters(View);
 	PassParameters->SSAOParameters = GetSSAOShaderParameters(View, SceneViewport, OutputViewport, CommonParameters.SceneTexturesViewport, EAOTechnique::GTAO);
@@ -1193,7 +1316,7 @@ public:
 
 	BEGIN_SHADER_PARAMETER_STRUCT(FParameters, )
 		SHADER_PARAMETER_STRUCT_REF(FViewUniformShaderParameters, View)
-		SHADER_PARAMETER_RDG_UNIFORM_BUFFER(FSceneTextureUniformParameters, SceneTextures)
+		SHADER_PARAMETER_STRUCT_INCLUDE(FSceneTextureShaderParameters, SceneTextures)
 		SHADER_PARAMETER_RDG_UNIFORM_BUFFER(FSubstrateGlobalUniformParameters, Substrate)
 
 		SHADER_PARAMETER_STRUCT_INCLUDE(FSSAOShaderParameters, SSAOParameters)
@@ -1239,7 +1362,7 @@ FScreenPassTexture AddGTAOInnerIntegratePass(
 	FGTAOInnerIntegratePS::FParameters* PassParameters = GraphBuilder.AllocParameters<FGTAOInnerIntegratePS::FParameters>();
 
 	PassParameters->View = View.ViewUniformBuffer;
-	PassParameters->SceneTextures = CommonParameters.SceneTexturesUniformBuffer;
+	PassParameters->SceneTextures = CommonParameters.SceneTextures;
 	PassParameters->Substrate = Substrate::BindSubstrateGlobalUniformParameters(View);
 	PassParameters->SSAOParameters = GetSSAOShaderParameters(View, InputViewport, OutputViewport, CommonParameters.SceneTexturesViewport, EAOTechnique::GTAO);
 	PassParameters->GTAOParameters = GetGTAOShaderParameters(View, OutputViewport.Extent);
@@ -1290,7 +1413,7 @@ public:
 
 	BEGIN_SHADER_PARAMETER_STRUCT(FParameters, )
 		SHADER_PARAMETER_STRUCT_REF(FViewUniformShaderParameters, View)
-		SHADER_PARAMETER_RDG_UNIFORM_BUFFER(FSceneTextureUniformParameters, SceneTextures)
+		SHADER_PARAMETER_STRUCT_INCLUDE(FSceneTextureShaderParameters, SceneTextures)
 		SHADER_PARAMETER_STRUCT_INCLUDE(FHZBParameters, HZBParameters)
 		SHADER_PARAMETER_RDG_UNIFORM_BUFFER(FSubstrateGlobalUniformParameters, Substrate)
 
@@ -1318,7 +1441,7 @@ FGTAOHorizonSearchOutputs AddGTAOHorizonSearchPass(
 	FGTAOHorizonSearchCS::FParameters* PassParameters = GraphBuilder.AllocParameters<FGTAOHorizonSearchCS::FParameters>();
 
 	PassParameters->View = View.ViewUniformBuffer;
-	PassParameters->SceneTextures = CommonParameters.SceneTexturesUniformBuffer;
+	PassParameters->SceneTextures = CommonParameters.SceneTextures;
 	PassParameters->HZBParameters = GetHZBParametersForAO(GraphBuilder, View, SceneViewport.Extent, EAOTechnique::GTAO);
 	PassParameters->Substrate = Substrate::BindSubstrateGlobalUniformParameters(View);
 	PassParameters->SSAOParameters = GetSSAOShaderParameters(View, SceneViewport, OutputViewport, CommonParameters.SceneTexturesViewport, EAOTechnique::GTAO);
@@ -1367,7 +1490,7 @@ public:
 
 	BEGIN_SHADER_PARAMETER_STRUCT(FParameters, )
 		SHADER_PARAMETER_STRUCT_REF(FViewUniformShaderParameters, View)
-		SHADER_PARAMETER_RDG_UNIFORM_BUFFER(FSceneTextureUniformParameters, SceneTextures)
+		SHADER_PARAMETER_STRUCT_INCLUDE(FSceneTextureShaderParameters, SceneTextures)
 
 		SHADER_PARAMETER_STRUCT_INCLUDE(FSSAOShaderParameters, SSAOParameters)
 		SHADER_PARAMETER_STRUCT_INCLUDE(FGTAOShaderParameters, GTAOParameters)
@@ -1442,7 +1565,7 @@ FGTAOTemporalOutputs AddGTAOTemporalPass(
 	FGTAOTemporalFilterCS::FParameters* PassParameters = GraphBuilder.AllocParameters<FGTAOTemporalFilterCS::FParameters>();
 
 	PassParameters->View = View.ViewUniformBuffer;
-	PassParameters->SceneTextures = CommonParameters.SceneTexturesUniformBuffer;
+	PassParameters->SceneTextures = CommonParameters.SceneTextures;
 	PassParameters->SSAOParameters = GetSSAOShaderParameters(View, InputViewport, OutputViewport, CommonParameters.SceneTexturesViewport, EAOTechnique::GTAO);
 	PassParameters->GTAOParameters = GetGTAOShaderParameters(View, OutputViewport.Extent);
 
@@ -1505,7 +1628,7 @@ public:
 
 	BEGIN_SHADER_PARAMETER_STRUCT(FParameters, )
 		SHADER_PARAMETER_STRUCT_REF(FViewUniformShaderParameters, View)
-		SHADER_PARAMETER_RDG_UNIFORM_BUFFER(FSceneTextureUniformParameters, SceneTextures)
+		SHADER_PARAMETER_STRUCT_INCLUDE(FSceneTextureShaderParameters, SceneTextures)
 
 		SHADER_PARAMETER_STRUCT_INCLUDE(FSSAOShaderParameters, SSAOParameters)
 
@@ -1553,7 +1676,7 @@ FScreenPassTexture AddGTAOSpatialFilter(
 
 	FGTAOSpatialFilterCS::FParameters* PassParameters = GraphBuilder.AllocParameters<FGTAOSpatialFilterCS::FParameters>();
 	PassParameters->View = View.ViewUniformBuffer;
-	PassParameters->SceneTextures = CommonParameters.SceneTexturesUniformBuffer;
+	PassParameters->SceneTextures = CommonParameters.SceneTextures;
 	PassParameters->SSAOParameters = GetSSAOShaderParameters(View, InputViewport, OutputViewport, CommonParameters.SceneTexturesViewport, EAOTechnique::GTAO);
 
 	PassParameters->GTAOSpatialFilterExtents = OutputViewport.Rect.Size();
diff --git a/Engine/Source/Runtime/Renderer/Private/CompositionLighting/PostProcessAmbientOcclusion.h b/Engine/Source/Runtime/Renderer/Private/CompositionLighting/PostProcessAmbientOcclusion.h
--- a/Engine/Source/Runtime/Renderer/Private/CompositionLighting/PostProcessAmbientOcclusion.h
+++ b/Engine/Source/Runtime/Renderer/Private/CompositionLighting/PostProcessAmbientOcclusion.h
@@ -10,6 +10,7 @@
 #include "UniformBuffer.h"
 #include "RendererInterface.h"
 #include "HZB.h"
+#include "SceneTexturesConfig.h"
 
 class FViewInfo;
 
@@ -91,7 +92,7 @@ public:
 	static int32 GetAmbientOcclusionShaderLevel(const FSceneView& View);
 
 	// @return whether AmbientOcclusion should run a compute shader.
-	static bool IsAmbientOcclusionCompute(const ERHIFeatureLevel::Type FeatureLevel);
+	static bool IsAmbientOcclusionCompute();
 
 	static int32 GetNumAmbientOcclusionLevels();
 	static float GetAmbientOcclusionStepMipLevelFactor();
@@ -104,13 +105,16 @@ public:
 	static uint32 ComputeAmbientOcclusionPassCount(const FViewInfo& View);
 
 	static EGTAOType GetGTAOPassType(const FViewInfo& View, uint32 Levels);
+
+	// @return Whether depth bounds test for AO is enabled
+	static bool IsDepthBoundsTestEnabled();
 };
 
 // SSAO
 
 struct FSSAOCommonParameters
 {
-	TRDGUniformBufferRef<FSceneTextureUniformParameters> SceneTexturesUniformBuffer = nullptr;
+	FSceneTextureShaderParameters SceneTextures;
 	FScreenPassTextureViewport SceneTexturesViewport;
 
 	FScreenPassTexture GBufferA;
@@ -118,6 +122,8 @@ struct FSSAOCommonParameters
 
 	uint32 Levels = 1;
 	int32 ShaderQuality = 4;
+	int32 MobileHalfResolutionSetting = 0;
+	uint32 MobileShadingModelsExcludeMask = 0;
 	ESSAOType DownscaleType = ESSAOType::EPS;
 	ESSAOType FullscreenType = ESSAOType::EPS;
 	bool bNeedSmoothingPass = true;
@@ -150,7 +156,7 @@ FScreenPassTexture AddAmbientOcclusionFinalPass(
 
 struct FGTAOCommonParameters
 {
-	TRDGUniformBufferRef<FSceneTextureUniformParameters> SceneTexturesUniformBuffer = nullptr;
+	FSceneTextureShaderParameters SceneTextures;
 	FScreenPassTextureViewport SceneTexturesViewport;
 	
 	FScreenPassTexture SceneDepth;
diff --git a/Engine/Source/Runtime/Renderer/Private/MobileShadingRenderer.cpp b/Engine/Source/Runtime/Renderer/Private/MobileShadingRenderer.cpp
--- a/Engine/Source/Runtime/Renderer/Private/MobileShadingRenderer.cpp
+++ b/Engine/Source/Runtime/Renderer/Private/MobileShadingRenderer.cpp
@@ -1617,7 +1617,8 @@ void FMobileSceneRenderer::Render(FRDGBuilder& GraphBuilder, const FSceneRenderU
 		}
 
 		// When renderer is in ERendererOutput::DepthPrepassOnly mode, bRequiresAmbientOcclusionPass is set to false in InitViews()
-		if (bRendererOutputFinalSceneColor)
+		// Render SSAO after base pass when possible, for GBuffer access. AsyncCompute is suspposed to run with a BasePass
+		if (bRendererOutputFinalSceneColor && ((!bDeferredShading || !bRequiresMultiPass) || MobileSSAOUsesCompute()))
 		{
 			BasePassTextures.AmbientOcclusionTexture = RenderAmbientOcclusion(GraphBuilder, SceneTextures);
 		}
@@ -2804,7 +2805,7 @@ void FMobileSceneRenderer::RenderDeferredMultiPass(FRDGBuilder& GraphBuilder, FS
 		RenderDeferredMultiPassDecals(GraphBuilder, View, ViewContext.DecalViewOffset, ViewContext.ViewIndex, SceneTextures, InstanceCullingManager);
 	}
 
-	if (!bIsFullDepthPrepassEnabled)
+	if (!HasBeenProduced(BasePassTextures.AmbientOcclusionTexture))
 	{
 		BasePassTextures.AmbientOcclusionTexture = RenderAmbientOcclusion(GraphBuilder, SceneTextures);
 	}
@@ -3253,10 +3254,8 @@ void FMobileSceneRenderer::PreTonemapMSAA(FRHICommandList& RHICmdList, const FMi
 
 bool FMobileSceneRenderer::ShouldRenderHZB(TArrayView<FViewInfo> InViews)
 {
-	static const auto MobileAmbientOcclusionTechniqueCVar = IConsoleManager::Get().FindTConsoleVariableDataInt(TEXT("r.Mobile.AmbientOcclusionTechnique"));
-
 	// Mobile SSAO requests HZB
-	bool bIsFeatureRequested =  RequiresMobileAmbientOcclusionPass(Scene, InViews[0]) && MobileAmbientOcclusionTechniqueCVar->GetValueOnRenderThread() == 1;
+	bool bIsFeatureRequested =  RequiresMobileAmbientOcclusionPass(Scene, InViews[0]);
 
 	// Instance occlusion culling requires HZB
 	if (FInstanceCullingContext::IsOcclusionCullingEnabled())
diff --git a/Engine/Source/Runtime/Renderer/Private/PostProcess/PostProcessAmbientOcclusionMobile.cpp b/Engine/Source/Runtime/Renderer/Private/PostProcess/PostProcessAmbientOcclusionMobile.cpp
--- a/Engine/Source/Runtime/Renderer/Private/PostProcess/PostProcessAmbientOcclusionMobile.cpp
+++ b/Engine/Source/Runtime/Renderer/Private/PostProcess/PostProcessAmbientOcclusionMobile.cpp
@@ -20,6 +20,12 @@
 #include "ClearQuad.h"
 #include "PixelShaderUtils.h"
 
+extern FSSAOCommonParameters GetSSAOCommonParameters(
+	FRDGBuilder& GraphBuilder,
+	const FViewInfo& View,
+	uint32 Levels,
+	bool bAllowGBufferRead);
+
 static TAutoConsoleVariable<int32> CVarMobileAmbientOcclusion(
 	TEXT("r.Mobile.AmbientOcclusion"),
 	0,
@@ -29,760 +35,24 @@ static TAutoConsoleVariable<int32> CVarMobileAmbientOcclusion(
 	ECVF_ReadOnly | ECVF_RenderThreadSafe
 );
 
-static TAutoConsoleVariable<int32> CVarMobileAmbientOcclusionTechnique(
-	TEXT("r.Mobile.AmbientOcclusionTechnique"),
-	1,
-	TEXT("0: GTAO.\n")
-	TEXT("1: SSAO. (default)\n"),
-	ECVF_RenderThreadSafe
-);
-
-static TAutoConsoleVariable<int32> CVarMobileGTAOPreIntegratedTextureType(
-	TEXT("r.Mobile.GTAOPreIntegratedTextureType"),
-	2,
-	TEXT("0: No Texture.\n")
-	TEXT("1: Texture2D LUT.\n")
-	TEXT("2: Volume LUT(Default)."),
-	ECVF_ReadOnly);
-
-static TAutoConsoleVariable<int32> CVarMobileAmbientOcclusionQuality(
-	TEXT("r.Mobile.AmbientOcclusionQuality"),
-	2,
-	TEXT("The quality of screen space ambient occlusion on mobile platform.\n")
-	TEXT("0: Disabled.\n")
-	TEXT("1: Low\n")
-	TEXT("2: Medium. (Default)\n")
-	TEXT("3: High.\n"),
-	ECVF_Scalability | ECVF_RenderThreadSafe
-);
-
-// NOTE: this only applies to GTAO
-static TAutoConsoleVariable<int32> CVarMobileAmbientOcclusionShaderType(
-	TEXT("r.Mobile.AmbientOcclusionShaderType"),
-	2,
-	TEXT("0: ComputeShader.\n")
-	TEXT("1: Separate ComputeShader.\n")
-	TEXT("2: PixelShader.\n"),
-	ECVF_RenderThreadSafe
-);
-
-static TAutoConsoleVariable<int32> CVarMobileAmbientOcclusionSSAOCompute(
-	TEXT("r.Mobile.AmbientOcclusionSSAOCompute"),
-	0,
-	TEXT("Whether to use the compute shader variant of mobile SSAO\n")
-	TEXT("0: Disabled (PS).\n")
-	TEXT("1: Compute (graphics pipe)\n")
-	TEXT("2: Compute (async pipe)\n"),
-	ECVF_Scalability | ECVF_RenderThreadSafe
-);
-
-
-static TAutoConsoleVariable<int32> CVarMobileAmbientOcclusionCompression(
-	TEXT("r.Mobile.AmbientOcclusionCompression"),
-	1,
-	TEXT("0: Disable framebuffer compression for AO texture.\n")
-	TEXT("1: Enable framebuffer compression for AO texture if supported (Default).\n"),
-	ECVF_RenderThreadSafe
-);
-
-enum class EMobileAmbientOcclusionShaderType
-{
-	ComputeShader,
-	SeparateComputeShader,
-	PixelShader
-};
-
-static EMobileAmbientOcclusionShaderType GetCVarMobileAmbientOcclusionShaderTypeOnRenderThread()
-{
-	int32 ShaderType = FMath::Clamp(CVarMobileAmbientOcclusionShaderType.GetValueOnRenderThread(), 0, 2);
-	if (ShaderType == 0 && GetMaxWorkGroupInvocations() < 1024)
-	{
-		ShaderType = 1;
-	}
-	return static_cast<EMobileAmbientOcclusionShaderType>(ShaderType);
-}
-
-static bool IsComputeShaderType(EMobileAmbientOcclusionShaderType ShaderType)
-{
-	return ShaderType != EMobileAmbientOcclusionShaderType::PixelShader;
-}
-
-static bool IsTwoPassShaderType(EMobileAmbientOcclusionShaderType ShaderType)
-{
-	return ShaderType != EMobileAmbientOcclusionShaderType::ComputeShader;
-}
-
-static TAutoConsoleVariable<int32> CVarMobileAmbientOcclusionDepthBoundsTest(
-	TEXT("r.Mobile.AmbientOcclusionDepthBoundsTest"),
-	1,
-	TEXT("Whether to use depth bounds test to cull distant pixels during AO pass. This option is only valid when pixel shader path is used"),
-	ECVF_RenderThreadSafe);
-
 static TAutoConsoleVariable<int32> CVarMobileSSAOHalfResolution(
 	TEXT("r.Mobile.SSAOHalfResolution"),
-	2,
+	0,
 	TEXT("Whether to calculate SSAO at half resolution.\n")
 	TEXT("0: Disabled.\n")
 	TEXT("1: Half Resolution with bilinear upsample\n")
 	TEXT("2: Half Resolution with 4 tap bilateral upsample\n")
 	TEXT("3: Half Resolution with 9 tap bilateral upsample\n"),
-	ECVF_RenderThreadSafe);
-
-// --------------------------------------------------------------------------------------------------------------------
-DECLARE_GPU_STAT_NAMED(MobileSSAO, TEXT("SSAO"));
-DECLARE_GPU_STAT_NAMED(MobileGTAO, TEXT("GTAO"));
-
-// --------------------------------------------------------------------------------------------------------------------
-class FGTAOMobile_HorizonSearchIntegral : public FGlobalShader
-{
-public:
-	class FLUTTextureTypeDim : SHADER_PERMUTATION_INT("PREINTEGRATED_LUT_TYPE", 3);
-	class FShaderQualityDim : SHADER_PERMUTATION_INT("SHADER_QUALITY", 3);
-
-	using FCommonPermutationDomain = TShaderPermutationDomain<
-		FLUTTextureTypeDim,
-		FShaderQualityDim>;
-
-	BEGIN_SHADER_PARAMETER_STRUCT(FParameters, )
-		SHADER_PARAMETER_STRUCT_REF(FViewUniformShaderParameters, View)
-		SHADER_PARAMETER_EX(FVector4f, ViewRectMin, EShaderPrecisionModifier::Half)
-		SHADER_PARAMETER_EX(FVector4f, DepthBufferSizeAndInvSize, EShaderPrecisionModifier::Half)
-		SHADER_PARAMETER_EX(FVector4f, BufferSizeAndInvSize, EShaderPrecisionModifier::Half)
-		SHADER_PARAMETER_EX(FVector4f, ViewSizeAndInvSize, EShaderPrecisionModifier::Half)
-		SHADER_PARAMETER(FVector4f, FadeRadiusMulAdd_FadeDistance_AttenFactor)
-		SHADER_PARAMETER(FVector4f, WorldRadiusAdj_SinDeltaAngle_CosDeltaAngle_Thickness)
-
-		SHADER_PARAMETER_RDG_TEXTURE(Texture2D, SceneDepthTexture)
-		SHADER_PARAMETER_SAMPLER(SamplerState, SceneDepthSampler)
-
-		SHADER_PARAMETER_RDG_TEXTURE(Texture2D, NormalTexture)
-		SHADER_PARAMETER_SAMPLER(SamplerState, NormalSampler)
-
-		SHADER_PARAMETER_TEXTURE(Texture2D, GTAOPreIntegrated2D)
-		SHADER_PARAMETER_TEXTURE(Texture3D, GTAOPreIntegrated3D)
-		SHADER_PARAMETER_SAMPLER(SamplerState, GTAOPreIntegratedSampler)
-	END_SHADER_PARAMETER_STRUCT()
-
-	static bool ShouldCompilePermutation(const FGlobalShaderPermutationParameters& Parameters, const FCommonPermutationDomain& CommonPermutationVector)
-	{
-		auto LUTTextureType = CommonPermutationVector.Get<FLUTTextureTypeDim>();
-
-		int32 MobileGTAOPreIntegratedTextureType = CVarMobileGTAOPreIntegratedTextureType.GetValueOnAnyThread();
-		return IsMobilePlatform(Parameters.Platform) && (MobileGTAOPreIntegratedTextureType == LUTTextureType);
-	}
-
-	static void ModifyCompilationEnvironment(const FGlobalShaderPermutationParameters& Parameters, FShaderCompilerEnvironment& OutEnvironment)
-	{
-		FGlobalShader::ModifyCompilationEnvironment(Parameters, OutEnvironment);
-		OutEnvironment.SetDefine(TEXT("USE_NORMALBUFFER"), 0);
-	}
-
-	static FCommonPermutationDomain BuildPermutationVector(int32 LUTTextureType, int32 ShaderQuality)
-	{
-		FCommonPermutationDomain PermutationVector;
-		PermutationVector.Set<FLUTTextureTypeDim>(LUTTextureType);
-		PermutationVector.Set<FShaderQualityDim>(ShaderQuality);
-		return PermutationVector;
-	}
-
-	static void SetupShaderParameters(FParameters& ShaderParameters, FRDGBuilder& GraphBuilder, const FViewInfo& View, const FIntRect& ViewRect, const FIntPoint& DepthBufferSize, const FIntPoint& BufferSize, const FVector4f& FallOffStartEndScaleBias, const FVector4f& WorldRadiusAdjSinCosDeltaAngleThickness, FRDGTextureRef SceneDepthTexture)
-	{
-		const FFinalPostProcessSettings& Settings = View.FinalPostProcessSettings;
-
-		float FadeRadius = FMath::Max(1.0f, Settings.AmbientOcclusionFadeRadius);
-		float InvFadeRadius = 1.0f / FadeRadius;
-
-		ShaderParameters.View = View.ViewUniformBuffer;
-		ShaderParameters.ViewRectMin = FVector4f(ViewRect.Min.X, ViewRect.Min.Y, 0.0f, 0.0f);
-		ShaderParameters.DepthBufferSizeAndInvSize = FVector4f(DepthBufferSize.X, DepthBufferSize.Y, 1.0f / DepthBufferSize.X, 1.0f / DepthBufferSize.Y);
-		ShaderParameters.BufferSizeAndInvSize = FVector4f(BufferSize.X, BufferSize.Y, 1.0f / BufferSize.X, 1.0f / BufferSize.Y);
-		ShaderParameters.ViewSizeAndInvSize = FVector4f(ViewRect.Width(), ViewRect.Height(), 1.0f / ViewRect.Width(), 1.0f / ViewRect.Height());
-		ShaderParameters.FadeRadiusMulAdd_FadeDistance_AttenFactor = FVector4f(InvFadeRadius, -(Settings.AmbientOcclusionFadeDistance - FadeRadius) * InvFadeRadius, Settings.AmbientOcclusionFadeDistance, 2.0f / (FallOffStartEndScaleBias.Y * FallOffStartEndScaleBias.Y));
-		ShaderParameters.WorldRadiusAdj_SinDeltaAngle_CosDeltaAngle_Thickness = WorldRadiusAdjSinCosDeltaAngleThickness;
-
-		ShaderParameters.SceneDepthTexture = SceneDepthTexture;
-		ShaderParameters.SceneDepthSampler = TStaticSamplerState<SF_Point, AM_Clamp, AM_Clamp, AM_Clamp>::GetRHI();
-
-		if (GSystemTextures.GTAOPreIntegrated.IsValid())
-		{
-			ShaderParameters.GTAOPreIntegrated2D = GSystemTextures.GTAOPreIntegrated->GetRHI();
-			ShaderParameters.GTAOPreIntegrated3D = GSystemTextures.GTAOPreIntegrated->GetRHI();
-			ShaderParameters.GTAOPreIntegratedSampler = TStaticSamplerState<SF_Bilinear, AM_Clamp, AM_Clamp, AM_Clamp>::GetRHI();
-		}
-	}
-
-	FGTAOMobile_HorizonSearchIntegral() = default;
-	FGTAOMobile_HorizonSearchIntegral(const CompiledShaderInitializerType& Initializer)
-		: FGlobalShader(Initializer)
-	{}
-};
-
-class FGTAOMobile_HorizonSearchIntegralSpatialFilterCS : public FGTAOMobile_HorizonSearchIntegral
-{
-	using Super = FGTAOMobile_HorizonSearchIntegral;
-
-public:
-	// Changing these numbers requires PostProcessAmbientOcclusionMobile.usf to be recompiled.
-	// The maximum thread group is 512 on IOS A9 and A10 and the shared memory is 16K
-	static const uint32 ThreadGroupSizeX = 32;
-	static const uint32 ThreadGroupSizeY = 32;
-
-	// The number of texels on each axis processed by a single thread group.
-	static const FIntPoint TexelsPerThreadGroup;
-
-	DECLARE_GLOBAL_SHADER(FGTAOMobile_HorizonSearchIntegralSpatialFilterCS);
-	SHADER_USE_PARAMETER_STRUCT(FGTAOMobile_HorizonSearchIntegralSpatialFilterCS, Super);
-
-	BEGIN_SHADER_PARAMETER_STRUCT(FParameters, )
-		SHADER_PARAMETER_STRUCT_INCLUDE(FGTAOMobile_HorizonSearchIntegral::FParameters, Common)
-		SHADER_PARAMETER_EX(FVector4f, Power_Intensity_ScreenPixelsToSearch, EShaderPrecisionModifier::Half)
-		SHADER_PARAMETER_RDG_TEXTURE_UAV(RWTexture2D<half4>, OutTexture)
-	END_SHADER_PARAMETER_STRUCT()
-
-	using FPermutationDomain = TShaderPermutationDomain<
-		Super::FCommonPermutationDomain>;
-
-	static bool ShouldCompilePermutation(const FGlobalShaderPermutationParameters& Parameters)
-	{
-		FPermutationDomain PermutationVector(Parameters.PermutationId);
-
-		return Super::ShouldCompilePermutation(Parameters, PermutationVector.Get<Super::FCommonPermutationDomain>());
-	}
-
-	static void ModifyCompilationEnvironment(const FGlobalShaderPermutationParameters& Parameters, FShaderCompilerEnvironment& OutEnvironment)
-	{
-		Super::ModifyCompilationEnvironment(Parameters, OutEnvironment);
-
-		OutEnvironment.SetDefine(TEXT("HORIZONSEARCH_INTEGRAL_SPATIALFILTER_COMPUTE_SHADER"), 1u);
-
-		OutEnvironment.SetDefine(TEXT("THREADGROUP_SIZEX"), ThreadGroupSizeX);
-		OutEnvironment.SetDefine(TEXT("THREADGROUP_SIZEY"), ThreadGroupSizeY);
-	}
-
-	static FPermutationDomain BuildPermutationVector(int32 LUTTextureType, int32 ShaderQuality)
-	{
-		FPermutationDomain PermutationVector;
-		PermutationVector.Set<Super::FCommonPermutationDomain>(Super::BuildPermutationVector(LUTTextureType, ShaderQuality));
-		return PermutationVector;
-	}
-};
-
-const FIntPoint FGTAOMobile_HorizonSearchIntegralSpatialFilterCS::TexelsPerThreadGroup(ThreadGroupSizeX, ThreadGroupSizeY);
-
-IMPLEMENT_GLOBAL_SHADER(FGTAOMobile_HorizonSearchIntegralSpatialFilterCS, "/Engine/Private/PostProcessAmbientOcclusionMobile.usf", "GTAOHorizonSearchIntegralSpatialFilterCS", SF_Compute);
-
-class FGTAOMobile_HorizonSearchIntegralCS : public FGTAOMobile_HorizonSearchIntegral
-{
-	using Super = FGTAOMobile_HorizonSearchIntegral;
-public:
-	// Changing these numbers requires PostProcessAmbientOcclusionMobile.usf to be recompiled.
-	// Use smaller thread group for low end devices
-	static const uint32 ThreadGroupSizeX = 16;
-	static const uint32 ThreadGroupSizeY = 8;
-
-	// The number of texels on each axis processed by a single thread group.
-	static const FIntPoint TexelsPerThreadGroup;
-
-	DECLARE_GLOBAL_SHADER(FGTAOMobile_HorizonSearchIntegralCS);
-	SHADER_USE_PARAMETER_STRUCT(FGTAOMobile_HorizonSearchIntegralCS, Super);
-
-	BEGIN_SHADER_PARAMETER_STRUCT(FParameters, )
-		SHADER_PARAMETER_STRUCT_INCLUDE(FGTAOMobile_HorizonSearchIntegral::FParameters, Common)
-		SHADER_PARAMETER_RDG_TEXTURE_UAV(RWTexture2D<half4>, OutTexture)
-	END_SHADER_PARAMETER_STRUCT()
-
-	using FPermutationDomain = TShaderPermutationDomain<
-		Super::FCommonPermutationDomain>;
-
-	static bool ShouldCompilePermutation(const FGlobalShaderPermutationParameters& Parameters)
-	{
-		FPermutationDomain PermutationVector(Parameters.PermutationId);
-
-		return Super::ShouldCompilePermutation(Parameters, PermutationVector.Get<Super::FCommonPermutationDomain>());
-	}
-
-	static void ModifyCompilationEnvironment(const FGlobalShaderPermutationParameters& Parameters, FShaderCompilerEnvironment& OutEnvironment)
-	{
-		Super::ModifyCompilationEnvironment(Parameters, OutEnvironment);
-
-		OutEnvironment.SetDefine(TEXT("HORIZONSEARCH_INTEGRAL_COMPUTE_SHADER"), 1u);
-
-		OutEnvironment.SetDefine(TEXT("THREADGROUP_SIZEX"), ThreadGroupSizeX);
-		OutEnvironment.SetDefine(TEXT("THREADGROUP_SIZEY"), ThreadGroupSizeY);
-	}
-
-	static FPermutationDomain BuildPermutationVector(int32 LUTTextureType, int32 ShaderQuality)
-	{
-		FPermutationDomain PermutationVector;
-		PermutationVector.Set<Super::FCommonPermutationDomain>(Super::BuildPermutationVector(LUTTextureType, ShaderQuality));
-		return PermutationVector;
-	}
-};
-
-const FIntPoint FGTAOMobile_HorizonSearchIntegralCS::TexelsPerThreadGroup(ThreadGroupSizeX, ThreadGroupSizeY);
-
-IMPLEMENT_GLOBAL_SHADER(FGTAOMobile_HorizonSearchIntegralCS, "/Engine/Private/PostProcessAmbientOcclusionMobile.usf", "GTAOHorizonSearchIntegralCS", SF_Compute);
-
-class FGTAOMobile_SpatialFilter : public FGlobalShader
-{
-public:
-
-	BEGIN_SHADER_PARAMETER_STRUCT(FParameters, )
-		SHADER_PARAMETER_EX(FVector4f, ViewRectMin, EShaderPrecisionModifier::Half)
-		SHADER_PARAMETER_EX(FVector4f, BufferSizeAndInvSize, EShaderPrecisionModifier::Half)
-		SHADER_PARAMETER_EX(FVector4f, ViewSizeAndInvSize, EShaderPrecisionModifier::Half)
-		SHADER_PARAMETER_EX(FVector4f, Power_Intensity_ScreenPixelsToSearch, EShaderPrecisionModifier::Half)
-		SHADER_PARAMETER_RDG_TEXTURE(Texture2D, AOInputTexture)
-		SHADER_PARAMETER_SAMPLER(SamplerState, AOInputSampler)
-	END_SHADER_PARAMETER_STRUCT()
-
-	static void SetupShaderParameters(FParameters& ShaderParameters, const FViewInfo& View, const FIntRect& ViewRect, const FIntPoint& BufferSize, FRDGTextureRef HorizonSearchIntegralTexture)
-	{
-		const FFinalPostProcessSettings& Settings = View.FinalPostProcessSettings;
-
-		ShaderParameters.ViewRectMin = FVector4f(ViewRect.Min.X, ViewRect.Min.Y, 0.0f, 0.0f);
-		ShaderParameters.BufferSizeAndInvSize = FVector4f(BufferSize.X, BufferSize.Y, 1.0f / BufferSize.X, 1.0f / BufferSize.Y);
-		ShaderParameters.ViewSizeAndInvSize = FVector4f(ViewRect.Width(), ViewRect.Height(), 1.0f / ViewRect.Width(), 1.0f / ViewRect.Height());
-		ShaderParameters.Power_Intensity_ScreenPixelsToSearch = FVector4f(Settings.AmbientOcclusionPower * 0.5f, Settings.AmbientOcclusionIntensity, 0.0f, 0.0f);
-
-		ShaderParameters.AOInputTexture = HorizonSearchIntegralTexture;
-		ShaderParameters.AOInputSampler = TStaticSamplerState<SF_Point, AM_Clamp, AM_Clamp, AM_Clamp>::GetRHI();
-	}
-
-	static bool ShouldCompilePermutation(const FGlobalShaderPermutationParameters& Parameters)
-	{
-		return IsMobilePlatform(Parameters.Platform);
-	}
-	
-	FGTAOMobile_SpatialFilter() = default;
-	FGTAOMobile_SpatialFilter(const CompiledShaderInitializerType& Initializer)
-		: FGlobalShader(Initializer)
-	{}
-};
-
-class FGTAOMobile_SpatialFilterCS : public FGTAOMobile_SpatialFilter
-{
-	using Super = FGTAOMobile_SpatialFilter;
-public:
-	// Changing these numbers requires PostProcessAmbientOcclusionMobile.usf to be recompiled.
-	// Use smaller thread group for low end devices
-	static const uint32 ThreadGroupSizeX = 16;
-	static const uint32 ThreadGroupSizeY = 8;
-
-	// The number of texels on each axis processed by a single thread group.
-	static const FIntPoint TexelsPerThreadGroup;
-
-	DECLARE_GLOBAL_SHADER(FGTAOMobile_SpatialFilterCS);
-	SHADER_USE_PARAMETER_STRUCT(FGTAOMobile_SpatialFilterCS, Super);
-
-	BEGIN_SHADER_PARAMETER_STRUCT(FParameters, )
-		SHADER_PARAMETER_STRUCT_INCLUDE(FGTAOMobile_SpatialFilter::FParameters, Common)
-		SHADER_PARAMETER_RDG_TEXTURE_UAV(RWTexture2D<half4>, OutTexture)
-	END_SHADER_PARAMETER_STRUCT()
-
-	static bool ShouldCompilePermutation(const FGlobalShaderPermutationParameters& Parameters)
-	{
-		return Super::ShouldCompilePermutation(Parameters);
-	}
-
-	static void ModifyCompilationEnvironment(const FGlobalShaderPermutationParameters& Parameters, FShaderCompilerEnvironment& OutEnvironment)
-	{
-		Super::ModifyCompilationEnvironment(Parameters, OutEnvironment);
-		OutEnvironment.SetDefine(TEXT("SPATIALFILTER_COMPUTE_SHADER"), 1u);
-
-		OutEnvironment.SetDefine(TEXT("THREADGROUP_SIZEX"), ThreadGroupSizeX);
-		OutEnvironment.SetDefine(TEXT("THREADGROUP_SIZEY"), ThreadGroupSizeY);
-	}
-};
-
-const FIntPoint FGTAOMobile_SpatialFilterCS::TexelsPerThreadGroup(ThreadGroupSizeX, ThreadGroupSizeY);
-
-IMPLEMENT_GLOBAL_SHADER(FGTAOMobile_SpatialFilterCS, "/Engine/Private/PostProcessAmbientOcclusionMobile.usf", "GTAOSpatialFilterCS", SF_Compute);
-
-class FGTAOMobile_HorizonSearchIntegralPS : public FGTAOMobile_HorizonSearchIntegral
-{
-	using Super = FGTAOMobile_HorizonSearchIntegral;
-public:
-
-	DECLARE_GLOBAL_SHADER(FGTAOMobile_HorizonSearchIntegralPS);
-	SHADER_USE_PARAMETER_STRUCT(FGTAOMobile_HorizonSearchIntegralPS, Super);
-
-	BEGIN_SHADER_PARAMETER_STRUCT(FParameters, )
-		SHADER_PARAMETER_STRUCT_INCLUDE(FGTAOMobile_HorizonSearchIntegral::FParameters, Common)
-		RENDER_TARGET_BINDING_SLOTS()
-	END_SHADER_PARAMETER_STRUCT()
-
-	using FPermutationDomain = TShaderPermutationDomain<
-		Super::FCommonPermutationDomain>;
-
-	static bool ShouldCompilePermutation(const FGlobalShaderPermutationParameters& Parameters)
-	{
-		FPermutationDomain PermutationVector(Parameters.PermutationId);
-
-		return Super::ShouldCompilePermutation(Parameters, PermutationVector.Get<Super::FCommonPermutationDomain>());
-	}
-
-	static void ModifyCompilationEnvironment(const FGlobalShaderPermutationParameters& Parameters, FShaderCompilerEnvironment& OutEnvironment)
-	{
-		Super::ModifyCompilationEnvironment(Parameters, OutEnvironment);
-
-		OutEnvironment.SetDefine(TEXT("HORIZONSEARCH_INTEGRAL_PIXEL_SHADER"), 1u);
-		OutEnvironment.SetDefine(TEXT("FORCE_DEPTH_TEXTURE_READS"), 1);
-	}
-
-	static FPermutationDomain BuildPermutationVector(int32 LUTTextureType, int32 ShaderQuality)
-	{
-		FPermutationDomain PermutationVector;
-		PermutationVector.Set<Super::FCommonPermutationDomain>(Super::BuildPermutationVector(LUTTextureType, ShaderQuality));
-		return PermutationVector;
-	}
-};
-
-IMPLEMENT_GLOBAL_SHADER(FGTAOMobile_HorizonSearchIntegralPS, "/Engine/Private/PostProcessAmbientOcclusionMobile.usf", "GTAOHorizonSearchIntegralPS", SF_Pixel);
-
-class FGTAOMobile_SpatialFilterPS : public FGTAOMobile_SpatialFilter
-{
-	using Super = FGTAOMobile_SpatialFilter;
-public:
-	DECLARE_GLOBAL_SHADER(FGTAOMobile_SpatialFilterPS);
-	SHADER_USE_PARAMETER_STRUCT(FGTAOMobile_SpatialFilterPS, Super);
-
-	BEGIN_SHADER_PARAMETER_STRUCT(FParameters, )
-		SHADER_PARAMETER_STRUCT_INCLUDE(FGTAOMobile_SpatialFilter::FParameters, Common)
-		RENDER_TARGET_BINDING_SLOTS()
-	END_SHADER_PARAMETER_STRUCT()
-
-	static bool ShouldCompilePermutation(const FGlobalShaderPermutationParameters& Parameters)
-	{
-		return Super::ShouldCompilePermutation(Parameters);
-	}
-
-	static void ModifyCompilationEnvironment(const FGlobalShaderPermutationParameters& Parameters, FShaderCompilerEnvironment& OutEnvironment)
-	{
-		Super::ModifyCompilationEnvironment(Parameters, OutEnvironment);
-		OutEnvironment.SetDefine(TEXT("SPATIALFILTER_PIXEL_SHADER"), 1u);
-		OutEnvironment.SetDefine(TEXT("FORCE_DEPTH_TEXTURE_READS"), 1);
-	}
-};
-
-IMPLEMENT_GLOBAL_SHADER(FGTAOMobile_SpatialFilterPS, "/Engine/Private/PostProcessAmbientOcclusionMobile.usf", "GTAOSpatialFilterPS", SF_Pixel);
-
-static void RenderGTAO(FRDGBuilder& GraphBuilder, const FSceneTextures& SceneTextures, FRDGTextureRef AmbientOcclusionTexture, const TArray<FViewInfo>& Views)
-{
-	RDG_EVENT_SCOPE_STAT(GraphBuilder, MobileGTAO, "MobileGTAO");
-		
-	static const auto GTAOThicknessBlendCVar = IConsoleManager::Get().FindTConsoleVariableDataFloat(TEXT("r.GTAO.ThicknessBlend"));
-	static const auto GTAOFalloffStartRatioCVar = IConsoleManager::Get().FindTConsoleVariableDataFloat(TEXT("r.GTAO.FalloffStartRatio"));
-	static const auto GTAOFalloffEndCVar = IConsoleManager::Get().FindTConsoleVariableDataFloat(TEXT("r.GTAO.FalloffEnd"));
-	static const auto GTAONumAnglesCVar = IConsoleManager::Get().FindTConsoleVariableDataFloat(TEXT("r.GTAO.NumAngles"));
-	const uint32 DownsampleFactor = 2;
-
-	const int32 MobileGTAOPreIntegratedTextureType = FMath::Min(CVarMobileGTAOPreIntegratedTextureType.GetValueOnRenderThread(), 2);
-	const int32 MobileAmbientOcclusionQuality = FMath::Min(CVarMobileAmbientOcclusionQuality.GetValueOnRenderThread(), 3);
-
-	FRDGTextureRef SceneDepthTexture = SceneTextures.Depth.Resolve;
-	const FIntPoint DepthBufferSize = SceneDepthTexture->Desc.Extent;
-	const FIntPoint BufferSize = AmbientOcclusionTexture->Desc.Extent;
-
-	float FallOffEnd = GTAOFalloffEndCVar ? GTAOFalloffEndCVar->GetValueOnRenderThread() : 200.0f;
-	float FallOffStartRatio = GTAOFalloffStartRatioCVar ? FMath::Clamp(GTAOFalloffStartRatioCVar->GetValueOnRenderThread(), 0.0f, 0.999f) : 0.5f;
-	float FallOffStart = FallOffEnd * FallOffStartRatio;
-	float FallOffStartSq = FallOffStart * FallOffStart;
-	float FallOffEndSq = FallOffEnd * FallOffEnd;
-
-	float FallOffScale = 1.0f / (FallOffEndSq - FallOffStartSq);
-	float FallOffBias = -FallOffStartSq * FallOffScale;
-
-	FVector4f FallOffStartEndScaleBias(FallOffStart, FallOffEnd, FallOffScale, FallOffBias);
-
-	float ThicknessBlend = GTAOThicknessBlendCVar ? GTAOThicknessBlendCVar->GetValueOnRenderThread() : 0.5f;
-	ThicknessBlend = FMath::Clamp(1.0f - (ThicknessBlend*ThicknessBlend), 0.0f, 0.99f);
-
-	float NumAngles = GTAONumAnglesCVar ? FMath::Clamp(GTAONumAnglesCVar->GetValueOnRenderThread(), 1.0f, 16.0f) : 2;
-	float SinDeltaAngle, CosDeltaAngle;
-	FMath::SinCos(&SinDeltaAngle, &CosDeltaAngle, PI / NumAngles);
-
-	const EMobileAmbientOcclusionShaderType ShaderType = GetCVarMobileAmbientOcclusionShaderTypeOnRenderThread();
-
-	FRDGTextureUAVRef AmbientOcclusionTextureUAV{};
-	if (IsComputeShaderType(ShaderType))
-	{
-		AmbientOcclusionTextureUAV = GraphBuilder.CreateUAV(AmbientOcclusionTexture);
-	}
-
-	FRDGTextureRef HorizonSearchIntegralTexture{};
-	FRDGTextureUAVRef HorizonSearchIntegralTextureUAV{};
-	if (IsTwoPassShaderType(ShaderType))
-	{
-		ETextureCreateFlags TextureCreateFlags = TexCreate_ShaderResource | TexCreate_RenderTargetable;
-		const EPixelFormat PixelFormat = PF_R8G8B8A8;
-		if (UE::PixelFormat::HasCapabilities(PixelFormat, EPixelFormatCapabilities::LossyCompressible) && CVarMobileAmbientOcclusionCompression.GetValueOnRenderThread() != 0)
-		{
-			TextureCreateFlags |= TexCreate_LossyCompressionLowBitrate;
-		}
-
-		HorizonSearchIntegralTexture = GraphBuilder.CreateTexture(FRDGTextureDesc::Create2D(BufferSize, PixelFormat, FClearValueBinding::Black, TextureCreateFlags | (IsComputeShaderType(ShaderType) ? TexCreate_UAV : TexCreate_None)), TEXT("HorizonSearchIntegralTexture"));
-		if (IsComputeShaderType(ShaderType))
-		{
-			HorizonSearchIntegralTextureUAV = GraphBuilder.CreateUAV(HorizonSearchIntegralTexture);
-		}
-	}
-
-	for (int32 ViewIndex = 0; ViewIndex < Views.Num(); ViewIndex++)
-	{
-		const FViewInfo& View = Views[ViewIndex];
-		const FFinalPostProcessSettings& Settings = View.FinalPostProcessSettings;
-		const FIntRect& ViewRect = FIntRect::DivideAndRoundUp(View.ViewRect, DownsampleFactor);
-		FVector4f WorldRadiusAdjSinCosDeltaAngleThickness(FallOffStartEndScaleBias.Y * DepthBufferSize.Y * View.ViewMatrices.GetViewToClip().M[0][0], SinDeltaAngle, CosDeltaAngle, ThicknessBlend);
-
-		if (ShaderType == EMobileAmbientOcclusionShaderType::ComputeShader)
-		{
-			FGTAOMobile_HorizonSearchIntegralSpatialFilterCS::FParameters* HorizonSearchIntegralSpatialFilterParameters = GraphBuilder.AllocParameters<FGTAOMobile_HorizonSearchIntegralSpatialFilterCS::FParameters>();
-			FGTAOMobile_HorizonSearchIntegral::SetupShaderParameters(HorizonSearchIntegralSpatialFilterParameters->Common, GraphBuilder, View, ViewRect, DepthBufferSize, BufferSize, FallOffStartEndScaleBias, WorldRadiusAdjSinCosDeltaAngleThickness, SceneDepthTexture);
-			
-			HorizonSearchIntegralSpatialFilterParameters->Power_Intensity_ScreenPixelsToSearch = FVector4f(Settings.AmbientOcclusionPower * 0.5f, Settings.AmbientOcclusionIntensity, 0.0f, 0.0f);
-			HorizonSearchIntegralSpatialFilterParameters->OutTexture = AmbientOcclusionTextureUAV;
-
-			auto ComputeShaderPermutationVector = FGTAOMobile_HorizonSearchIntegralSpatialFilterCS::BuildPermutationVector(MobileGTAOPreIntegratedTextureType, MobileAmbientOcclusionQuality - 1);
-			TShaderMapRef<FGTAOMobile_HorizonSearchIntegralSpatialFilterCS> ComputeShader(View.ShaderMap, ComputeShaderPermutationVector);
-
-			FComputeShaderUtils::AddPass(
-				GraphBuilder,
-				RDG_EVENT_NAME("AmbientOcclusion_HorizonSearchIntegralSpatialFilter %dx%d (CS)", ViewRect.Width(), ViewRect.Height()),
-				ERDGPassFlags::Compute | ERDGPassFlags::NeverCull,
-				ComputeShader,
-				HorizonSearchIntegralSpatialFilterParameters,
-				FComputeShaderUtils::GetGroupCount(ViewRect.Size(), FGTAOMobile_HorizonSearchIntegralSpatialFilterCS::TexelsPerThreadGroup));
-		}
-		else if (ShaderType == EMobileAmbientOcclusionShaderType::PixelShader)
-		{
-			TShaderMapRef<FScreenPassVS> VertexShader(View.ShaderMap);
-
-			FScreenPassRenderTarget HorizonSearchIntegralRT(HorizonSearchIntegralTexture, ViewRect, ViewIndex > 0 ? ERenderTargetLoadAction::ELoad : ERenderTargetLoadAction::EClear);
-
-			FGTAOMobile_HorizonSearchIntegralPS::FParameters* HorizonSearchIntegralParameters = GraphBuilder.AllocParameters<FGTAOMobile_HorizonSearchIntegralPS::FParameters>();
-			FGTAOMobile_HorizonSearchIntegral::SetupShaderParameters(HorizonSearchIntegralParameters->Common, GraphBuilder, View, ViewRect, DepthBufferSize, BufferSize, FallOffStartEndScaleBias, WorldRadiusAdjSinCosDeltaAngleThickness, SceneDepthTexture);
-
-			HorizonSearchIntegralParameters->RenderTargets[0] = HorizonSearchIntegralRT.GetRenderTargetBinding();
-
-			auto HorizonSearchIntegralShaderPermutationVector = FGTAOMobile_HorizonSearchIntegralPS::BuildPermutationVector(MobileGTAOPreIntegratedTextureType, MobileAmbientOcclusionQuality - 1);
-			TShaderMapRef<FGTAOMobile_HorizonSearchIntegralPS> HorizonSearchIntegralShader(View.ShaderMap, HorizonSearchIntegralShaderPermutationVector);
-
-			ClearUnusedGraphResources(HorizonSearchIntegralShader, HorizonSearchIntegralParameters);
-
-			GraphBuilder.AddPass(
-				RDG_EVENT_NAME("AmbientOcclusion_HorizonSearchIntegral %dx%d (PS)", ViewRect.Width(), ViewRect.Height()),
-				HorizonSearchIntegralParameters,
-				ERDGPassFlags::Raster | ERDGPassFlags::NeverCull,
-				[VertexShader, HorizonSearchIntegralShader, HorizonSearchIntegralParameters, ViewRect, BufferSize](FRHICommandList& RHICmdList)
-			{
-				RHICmdList.SetViewport(ViewRect.Min.X, ViewRect.Min.Y, 0.0f, ViewRect.Max.X, ViewRect.Max.Y, 1.0f);
-
-				FGraphicsPipelineStateInitializer GraphicsPSOInit;
-				RHICmdList.ApplyCachedRenderTargets(GraphicsPSOInit);
-				GraphicsPSOInit.BlendState = TStaticBlendState<>::GetRHI();
-				GraphicsPSOInit.RasterizerState = TStaticRasterizerState<>::GetRHI();
-				GraphicsPSOInit.DepthStencilState = TStaticDepthStencilState<false, CF_Always>::GetRHI();
-
-				GraphicsPSOInit.BoundShaderState.VertexDeclarationRHI = GFilterVertexDeclaration.VertexDeclarationRHI;
-				GraphicsPSOInit.BoundShaderState.VertexShaderRHI = VertexShader.GetVertexShader();
-				GraphicsPSOInit.BoundShaderState.PixelShaderRHI = HorizonSearchIntegralShader.GetPixelShader();
-				GraphicsPSOInit.PrimitiveType = PT_TriangleList;
-
-				SetGraphicsPipelineState(RHICmdList, GraphicsPSOInit, 0);
-
-				SetShaderParameters(RHICmdList, HorizonSearchIntegralShader, HorizonSearchIntegralShader.GetPixelShader(), *HorizonSearchIntegralParameters);
-
-				DrawRectangle(
-					RHICmdList,
-					0, 0,
-					BufferSize.X, BufferSize.Y,
-					ViewRect.Min.X, ViewRect.Min.Y,
-					ViewRect.Width(), ViewRect.Height(),
-					BufferSize,
-					BufferSize,
-					VertexShader,
-					EDRF_UseTriangleOptimization);
-			});
-
-			FScreenPassRenderTarget AmbientOcclusionRT(AmbientOcclusionTexture, ViewRect, ViewIndex > 0 ? ERenderTargetLoadAction::ELoad : ERenderTargetLoadAction::EClear);
+	ECVF_Scalability | ECVF_RenderThreadSafe);
 
-			FGTAOMobile_SpatialFilterPS::FParameters* SpatialFilterParameters = GraphBuilder.AllocParameters<FGTAOMobile_SpatialFilterPS::FParameters>();
-			FGTAOMobile_SpatialFilter::SetupShaderParameters(SpatialFilterParameters->Common, View, ViewRect, BufferSize, HorizonSearchIntegralTexture);
-			SpatialFilterParameters->RenderTargets[0] = AmbientOcclusionRT.GetRenderTargetBinding();
-
-			TShaderMapRef<FGTAOMobile_SpatialFilterPS> SpatialFilterShader(View.ShaderMap);
-
-			ClearUnusedGraphResources(SpatialFilterShader, SpatialFilterParameters);
-
-			GraphBuilder.AddPass(
-				RDG_EVENT_NAME("AmbientOcclusion_SpatialFilter %dx%d (PS)", ViewRect.Width(), ViewRect.Height()),
-				SpatialFilterParameters,
-				ERDGPassFlags::Raster | ERDGPassFlags::NeverCull,
-				[VertexShader, SpatialFilterShader, SpatialFilterParameters, ViewRect, BufferSize](FRHICommandList& RHICmdList)
-			{
-				RHICmdList.SetViewport(ViewRect.Min.X, ViewRect.Min.Y, 0.0f, ViewRect.Max.X, ViewRect.Max.Y, 1.0f);
-
-				FGraphicsPipelineStateInitializer GraphicsPSOInit;
-				RHICmdList.ApplyCachedRenderTargets(GraphicsPSOInit);
-
-				GraphicsPSOInit.BlendState = TStaticBlendState<>::GetRHI();
-				GraphicsPSOInit.RasterizerState = TStaticRasterizerState<>::GetRHI();
-				GraphicsPSOInit.DepthStencilState = TStaticDepthStencilState<false, CF_Always>::GetRHI();
-
-				GraphicsPSOInit.BoundShaderState.VertexDeclarationRHI = GFilterVertexDeclaration.VertexDeclarationRHI;
-				GraphicsPSOInit.BoundShaderState.VertexShaderRHI = VertexShader.GetVertexShader();
-				GraphicsPSOInit.BoundShaderState.PixelShaderRHI = SpatialFilterShader.GetPixelShader();
-				GraphicsPSOInit.PrimitiveType = PT_TriangleList;
-
-				SetGraphicsPipelineState(RHICmdList, GraphicsPSOInit, 0);
-
-				SetShaderParameters(RHICmdList, SpatialFilterShader, SpatialFilterShader.GetPixelShader(), *SpatialFilterParameters);
-
-				DrawRectangle(
-					RHICmdList,
-					0, 0,
-					BufferSize.X, BufferSize.Y,
-					ViewRect.Min.X, ViewRect.Min.Y,
-					ViewRect.Width(), ViewRect.Height(),
-					BufferSize,
-					BufferSize,
-					VertexShader,
-					EDRF_UseTriangleOptimization);
-			});
-		}
-		else
-		{
-			FGTAOMobile_HorizonSearchIntegralCS::FParameters* HorizonSearchIntegralParameters = GraphBuilder.AllocParameters<FGTAOMobile_HorizonSearchIntegralCS::FParameters>();
-			FGTAOMobile_HorizonSearchIntegral::SetupShaderParameters(HorizonSearchIntegralParameters->Common, GraphBuilder, View, ViewRect, DepthBufferSize, BufferSize, FallOffStartEndScaleBias, WorldRadiusAdjSinCosDeltaAngleThickness, SceneDepthTexture);
-
-			HorizonSearchIntegralParameters->OutTexture = HorizonSearchIntegralTextureUAV;
-
-			auto HorizonSearchIntegralShaderPermutationVector = FGTAOMobile_HorizonSearchIntegralCS::BuildPermutationVector(MobileGTAOPreIntegratedTextureType, MobileAmbientOcclusionQuality - 1);
-			TShaderMapRef<FGTAOMobile_HorizonSearchIntegralCS> HorizonSearchIntegralShader(View.ShaderMap, HorizonSearchIntegralShaderPermutationVector);
-
-			FComputeShaderUtils::AddPass(
-				GraphBuilder,
-				RDG_EVENT_NAME("AmbientOcclusion_HorizonSearchIntegral %dx%d (CS)", ViewRect.Width(), ViewRect.Height()),
-				ERDGPassFlags::Compute | ERDGPassFlags::NeverCull,
-				HorizonSearchIntegralShader,
-				HorizonSearchIntegralParameters,
-				FComputeShaderUtils::GetGroupCount(ViewRect.Size(), FGTAOMobile_HorizonSearchIntegralCS::TexelsPerThreadGroup));
-
-			FGTAOMobile_SpatialFilterCS::FParameters* SpatialFilterParameters = GraphBuilder.AllocParameters<FGTAOMobile_SpatialFilterCS::FParameters>();
-			FGTAOMobile_SpatialFilter::SetupShaderParameters(SpatialFilterParameters->Common, View, ViewRect, BufferSize, HorizonSearchIntegralTexture);
-			SpatialFilterParameters->OutTexture = AmbientOcclusionTextureUAV;
-
-			TShaderMapRef<FGTAOMobile_SpatialFilterCS> SpatialFilterShader(View.ShaderMap);
-
-			FComputeShaderUtils::AddPass(
-				GraphBuilder,
-				RDG_EVENT_NAME("AmbientOcclusion_SpatialFilter %dx%d (CS)", ViewRect.Width(), ViewRect.Height()),
-				ERDGPassFlags::Compute | ERDGPassFlags::NeverCull,
-				SpatialFilterShader,
-				SpatialFilterParameters,
-				FComputeShaderUtils::GetGroupCount(ViewRect.Size(), FGTAOMobile_SpatialFilterCS::TexelsPerThreadGroup));
-		}
-	}
-}
+static TAutoConsoleVariable<int32> CVarMobileSSAOShadingModelExcludeMask(
+	TEXT("r.Mobile.SSAOShadingModelsExcludeMask"),
+	0,
+	TEXT("ShadingModels mask that determines which pixels are excluded from SSAO computation.\n"),
+	ECVF_Scalability | ECVF_RenderThreadSafe);
 
 // --------------------------------------------------------------------------------------------------------------------
-struct FMobileSSAOCommonParameters
-{
-	TRDGUniformBufferRef<FMobileSceneTextureUniformParameters> SceneTexturesUniformBufferRDG;
-	FScreenPassTextureViewport SceneTexturesViewport;
-
-	FScreenPassTexture SceneDepth;
-};
-
-static const uint32 kMobileSSAOParametersArraySize = 5; 
-
-BEGIN_SHADER_PARAMETER_STRUCT(FMobileSSAOShaderParameters, )
-	SHADER_PARAMETER_ARRAY(FVector4f, ScreenSpaceAOParams, [kMobileSSAOParametersArraySize])
-
-	SHADER_PARAMETER_STRUCT(FScreenPassTextureViewportParameters, AOViewport)
-	SHADER_PARAMETER_STRUCT(FScreenPassTextureViewportParameters, AOSceneViewport)
-END_SHADER_PARAMETER_STRUCT();
-
-
-BEGIN_SHADER_PARAMETER_STRUCT(FMobileAmbientOcclusionParameters, RENDERER_API)
-	SHADER_PARAMETER_STRUCT_REF(FViewUniformShaderParameters, View)
-	SHADER_PARAMETER_RDG_UNIFORM_BUFFER(FMobileSceneTextureUniformParameters, SceneTextures)
-
-	SHADER_PARAMETER_STRUCT_INCLUDE(FHZBParameters, HZBParameters)
-	SHADER_PARAMETER_STRUCT_INCLUDE(FMobileSSAOShaderParameters, SSAOParameters)
-
-	SHADER_PARAMETER_SAMPLER(SamplerState, SSAO_Sampler)
-
-	SHADER_PARAMETER(FVector2f, SSAO_DownsampledAOInverseSize)
-	SHADER_PARAMETER(FVector2f, SSAO_SvPositionScaleBias)
-
-	SHADER_PARAMETER_RDG_TEXTURE(Texture2D, RandomNormalTexture)
-	SHADER_PARAMETER_SAMPLER(SamplerState, RandomNormalTextureSampler)
-END_SHADER_PARAMETER_STRUCT();
-
-class FMobileAmbientOcclusionPS : public FGlobalShader
-{
-public:
-	DECLARE_GLOBAL_SHADER(FMobileAmbientOcclusionPS);
-	SHADER_USE_PARAMETER_STRUCT(FMobileAmbientOcclusionPS, FGlobalShader);
-
-	class FShaderQualityDim : SHADER_PERMUTATION_INT("SHADER_QUALITY", 3);
-	class FOutputDepth : SHADER_PERMUTATION_BOOL("OUTPUT_DEPTH");
-	using FPermutationDomain = TShaderPermutationDomain<FShaderQualityDim, FOutputDepth>;
-
-	static void ModifyCompilationEnvironment(const FGlobalShaderPermutationParameters& Parameters, FShaderCompilerEnvironment& OutEnvironment)
-	{
-		FGlobalShader::ModifyCompilationEnvironment(Parameters, OutEnvironment);
-		OutEnvironment.SetDefine(TEXT("SSAO"), 1);
-		OutEnvironment.SetDefine(TEXT("FORCE_DEPTH_TEXTURE_READS"), 1);
-	}
-
-	static bool ShouldCompilePermutation(const FGlobalShaderPermutationParameters& Parameters)
-	{
-		return IsMobilePlatform(Parameters.Platform);
-	}
-
-	BEGIN_SHADER_PARAMETER_STRUCT(FParameters, )
-		SHADER_PARAMETER_STRUCT_INCLUDE(FMobileAmbientOcclusionParameters, SharedParameters)
-		RENDER_TARGET_BINDING_SLOTS()
-	END_SHADER_PARAMETER_STRUCT();
-};
-IMPLEMENT_GLOBAL_SHADER(FMobileAmbientOcclusionPS, "/Engine/Private/PostProcessAmbientOcclusionMobile.usf", "MainPS", SF_Pixel);
-
-class FMobileAmbientOcclusionCS : public FGlobalShader
-{
-public:
-	DECLARE_GLOBAL_SHADER(FMobileAmbientOcclusionCS);
-	SHADER_USE_PARAMETER_STRUCT(FMobileAmbientOcclusionCS, FGlobalShader);
-
-	class FUseUpsampleDim : SHADER_PERMUTATION_BOOL("USE_UPSAMPLE");
-	class FShaderQualityDim : SHADER_PERMUTATION_INT("SHADER_QUALITY", 5);
-
-	using FPermutationDomain = TShaderPermutationDomain<FUseUpsampleDim, FShaderQualityDim>;
-
-	static void ModifyCompilationEnvironment(const FGlobalShaderPermutationParameters& Parameters, FShaderCompilerEnvironment& OutEnvironment)
-	{
-		FGlobalShader::ModifyCompilationEnvironment(Parameters, OutEnvironment);
-		OutEnvironment.SetDefine(TEXT("COMPUTE_SHADER"), 1);
-		OutEnvironment.SetDefine(TEXT("THREADGROUP_SIZEX"), 16);
-		OutEnvironment.SetDefine(TEXT("THREADGROUP_SIZEY"), 16);
-		OutEnvironment.SetDefine(TEXT("SSAO"), 1);
-		OutEnvironment.SetDefine(TEXT("USE_AO_SETUP_AS_INPUT"), 0);
-	}
-
-	static bool ShouldCompilePermutation(const FGlobalShaderPermutationParameters& Parameters)
-	{
-		return IsMobileAmbientOcclusionEnabled(Parameters.Platform);
-	}
-
-	BEGIN_SHADER_PARAMETER_STRUCT(FParameters, )
-		SHADER_PARAMETER_STRUCT_INCLUDE(FMobileAmbientOcclusionParameters, SharedParameters)
-		SHADER_PARAMETER_RDG_TEXTURE_UAV(RWTexture2D, OutTexture)
-	END_SHADER_PARAMETER_STRUCT();
-};
-IMPLEMENT_GLOBAL_SHADER(FMobileAmbientOcclusionCS, "/Engine/Private/PostProcessAmbientOcclusionMobile.usf", "MainCS", SF_Compute);
-
+DECLARE_GPU_STAT_NAMED(MobileSSAO, TEXT("SSAO"));
 
 class FMobileAmbientOcclusionUpsamplePS : public FGlobalShader
 {
@@ -793,8 +63,8 @@ class FMobileAmbientOcclusionUpsamplePS : public FGlobalShader
 	using FPermutationDomain = TShaderPermutationDomain <FUpsampleQualityDim>;
 
 	BEGIN_SHADER_PARAMETER_STRUCT(FParameters, )
-		SHADER_PARAMETER_RDG_UNIFORM_BUFFER(FMobileSceneTextureUniformParameters, SceneTextures)
 		SHADER_PARAMETER_STRUCT_REF(FViewUniformShaderParameters, View)
+		SHADER_PARAMETER_STRUCT_INCLUDE(FSceneTextureShaderParameters, SceneTextures)
 		SHADER_PARAMETER_RDG_TEXTURE(Texture2D, AOTexture)
 		SHADER_PARAMETER_SAMPLER(SamplerState, AOSampler)
 		RENDER_TARGET_BINDING_SLOTS()
@@ -808,116 +78,32 @@ class FMobileAmbientOcclusionUpsamplePS : public FGlobalShader
 	static void ModifyCompilationEnvironment(const FGlobalShaderPermutationParameters& Parameters, FShaderCompilerEnvironment& OutEnvironment)
 	{
 		OutEnvironment.SetDefine(TEXT("UPSAMPLE_PASS"), 1);
-	
 	}
 };
-IMPLEMENT_GLOBAL_SHADER(FMobileAmbientOcclusionUpsamplePS, "/Engine/Private/PostProcessAmbientOcclusionMobile.usf", "AmbientOcclusionUpsamplePS", SF_Pixel);
+IMPLEMENT_GLOBAL_SHADER(FMobileAmbientOcclusionUpsamplePS, "/Engine/Private/PostProcessAmbientOcclusion.usf", "AmbientOcclusionUpsamplePS", SF_Pixel);
 
 // --------------------------------------------------------------------------------------------------------------------
-static FMobileSSAOCommonParameters GetMobileSSAOCommonParameters(
+void AddAmbientOcclusionUpsamplePass(
 	FRDGBuilder& GraphBuilder,
 	const FViewInfo& View,
-	FRDGTextureRef SceneDepthTexture,
-	TRDGUniformBufferRef<FMobileSceneTextureUniformParameters> SceneTexturesUniformBufferRDG)
-{
-	FMobileSSAOCommonParameters CommonParameters;
-	CommonParameters.SceneTexturesUniformBufferRDG = SceneTexturesUniformBufferRDG;
-	CommonParameters.SceneTexturesViewport = FScreenPassTextureViewport(SceneDepthTexture, View.ViewRect);
-
-	CommonParameters.SceneDepth = FScreenPassTexture(SceneDepthTexture);
-	return CommonParameters;
-}
-
-static FMobileSSAOShaderParameters GetMobileSSAOShaderParameters(
-	const FViewInfo& View,
-	const FScreenPassTextureViewport& InputViewport,
-	const FScreenPassTextureViewport& OutputViewport,
-	const FScreenPassTextureViewport& SceneViewport)
-{
-	const FFinalPostProcessSettings& Settings = View.FinalPostProcessSettings;
-
-	FIntPoint RandomizationSize = GSystemTextures.SSAORandomization->GetDesc().Extent * InputViewport.Extent / OutputViewport.Extent;
-	FVector2D ViewportUVToRandomUV(InputViewport.Extent.X / (float)RandomizationSize.X, InputViewport.Extent.Y / (float)RandomizationSize.Y);
-
-	// e.g. 4 means the input texture is 4x smaller than the buffer size
-	uint32 ScaleToFullRes = SceneViewport.Extent.X / InputViewport.Extent.X;
-
-	FIntRect ViewRect = FIntRect::DivideAndRoundUp(View.ViewRect, ScaleToFullRes);
-
-	float AORadiusInShader = Settings.AmbientOcclusionRadius;
-	float ScaleRadiusInWorldSpace = 1.0f;
-
-	if (!Settings.AmbientOcclusionRadiusInWS)
-	{
-		// radius is defined in view space in 400 units
-		AORadiusInShader /= 400.0f;
-		ScaleRadiusInWorldSpace = 0.0f;
-	}
-
-	// /4 is an adjustment for usage with multiple mips
-	float f = FMath::Log2(float(ScaleToFullRes));
-	float g = pow(Settings.AmbientOcclusionMipScale, f);
-	AORadiusInShader *= pow(Settings.AmbientOcclusionMipScale, FMath::Log2(float(ScaleToFullRes))) / 4.0f;
-
-	float Ratio = View.UnscaledViewRect.Width() / (float)View.UnscaledViewRect.Height();
-		
-	// Grab this and pass into shader so we can negate the fov influence of projection on the screen pos.
-	float InvTanHalfFov = View.ViewMatrices.GetViewToClip().M[0][0];
-
-	float StaticFraction = FMath::Clamp(Settings.AmbientOcclusionStaticFraction, 0.0f, 1.0f);
-
-	// clamp to prevent user error
-	float FadeRadius = FMath::Max(1.0f, Settings.AmbientOcclusionFadeRadius);
-	float InvFadeRadius = 1.0f / FadeRadius;
-
-	FVector2D TemporalOffset(0.0f, 0.0f);
-
-	if (View.State)
-	{
-		TemporalOffset = (View.State->GetCurrentTemporalAASampleIndex() % 8) * FVector2D(2.48f, 7.52f) / (float)RandomizationSize.X;
-	}
-	const float HzbStepMipLevelFactorValue = FMath::Clamp(FSSAOHelper::GetAmbientOcclusionStepMipLevelFactor(), 0.0f, 100.0f);
-	const float InvAmbientOcclusionDistance = 1.0f / FMath::Max(Settings.AmbientOcclusionDistance_DEPRECATED, KINDA_SMALL_NUMBER);
-
-	FMobileSSAOShaderParameters Result{};
-
-	// /1000 to be able to define the value in that distance
-	Result.ScreenSpaceAOParams[0] = FVector4f(Settings.AmbientOcclusionPower, Settings.AmbientOcclusionBias / 1000.0f, InvAmbientOcclusionDistance, Settings.AmbientOcclusionIntensity);
-	Result.ScreenSpaceAOParams[1] = FVector4f(ViewportUVToRandomUV.X, ViewportUVToRandomUV.Y, AORadiusInShader, Ratio);
-	Result.ScreenSpaceAOParams[2] = FVector4f(ScaleToFullRes, Settings.AmbientOcclusionMipThreshold / ScaleToFullRes, ScaleRadiusInWorldSpace, Settings.AmbientOcclusionMipBlend);
-	Result.ScreenSpaceAOParams[3] = FVector4f(TemporalOffset.X, TemporalOffset.Y, StaticFraction, InvTanHalfFov);
-	Result.ScreenSpaceAOParams[4] = FVector4f(InvFadeRadius, -(Settings.AmbientOcclusionFadeDistance - FadeRadius) * InvFadeRadius, HzbStepMipLevelFactorValue, Settings.AmbientOcclusionFadeDistance);
-
-	Result.AOViewport = GetScreenPassTextureViewportParameters(OutputViewport);
-	Result.AOSceneViewport = GetScreenPassTextureViewportParameters(SceneViewport);
-
-	return Result;
-}
-
-static void AddMobileAmbientOcclusionPass(
-	FRDGBuilder& GraphBuilder,
-	const FViewInfo& View,
-	const FMobileSSAOCommonParameters& CommonParameters,
-	bool bHalfResolution,
-	bool bOutputDepth,
+	const FSSAOCommonParameters& CommonParameters,
+	int32 UpsampleQuality,
+	FScreenPassTexture Input,
 	FScreenPassRenderTarget Output)
 {
-	check(Output.IsValid());
-
-	const FScreenPassTextureViewport InputViewport = CommonParameters.SceneTexturesViewport;
-	const FScreenPassTextureViewport OutputViewport = GetDownscaledViewport(CommonParameters.SceneTexturesViewport, bHalfResolution ? 2 : 1);
-	
-	const bool bDepthBoundsTestEnabled = CVarMobileAmbientOcclusionDepthBoundsTest.GetValueOnRenderThread()
-		&& !bHalfResolution
-		&& GSupportsDepthBoundsTest
+	const bool bDepthBoundsTestEnabled = 
+		FSSAOHelper::IsDepthBoundsTestEnabled()
 		&& CommonParameters.SceneDepth.IsValid()
 		&& CommonParameters.SceneDepth.Texture->Desc.NumSamples == 1;
 
-	float DepthFar = 0.0f;
-	FDepthStencilBinding DepthStencilBinding(CommonParameters.SceneDepth.Texture, ERenderTargetLoadAction::ELoad, ERenderTargetLoadAction::ELoad, FExclusiveDepthStencil::DepthRead_StencilWrite);
+	const FScreenPassTextureViewport OutputViewport(Output);
 
+	FDepthStencilBinding DepthStencilBinding{}; 
+	float DepthFar = 0.0f;
 	if (bDepthBoundsTestEnabled)
 	{
+		DepthStencilBinding = FDepthStencilBinding(CommonParameters.SceneDepth.Texture, ERenderTargetLoadAction::ELoad, ERenderTargetLoadAction::ELoad, FExclusiveDepthStencil::DepthRead_StencilWrite);
+		
 		const FFinalPostProcessSettings& Settings = View.FinalPostProcessSettings;
 		const FMatrix& ProjectionMatrix = View.ViewMatrices.GetViewToClip();
 		const FVector4f Far = (FVector4f)ProjectionMatrix.TransformFVector4(FVector4(0, 0, Settings.AmbientOcclusionFadeDistance));
@@ -931,7 +117,7 @@ static void AddMobileAmbientOcclusionPass(
 			RDG_EVENT_NAME("DepthBounds ClearQuad(%s)", Output.Texture->Name),
 			ClearParameters,
 			ERDGPassFlags::Raster,
-			[OutputViewport, DepthFar](FRHICommandListImmediate& RHICmdList)
+			[OutputViewport, DepthFar](FRDGAsyncTask, FRHICommandList& RHICmdList)
 			{
 				// We must clear all pixels that won't be touched by AO shader.
 				FClearQuadCallbacks Callbacks;
@@ -958,182 +144,7 @@ static void AddMobileAmbientOcclusionPass(
 		// Make sure the following pass doesn't clear or ignore the data
 		Output.LoadAction = ERenderTargetLoadAction::ELoad;
 	}
-
-	FMobileAmbientOcclusionParameters SharedParameters;
-	SharedParameters.View = View.ViewUniformBuffer;
-	SharedParameters.SceneTextures = CommonParameters.SceneTexturesUniformBufferRDG;
-	SharedParameters.HZBParameters = GetHZBParametersForAO(GraphBuilder, View, CommonParameters.SceneTexturesViewport.Extent, EAOTechnique::SSAO);
-	SharedParameters.SSAOParameters = GetMobileSSAOShaderParameters(View, InputViewport, OutputViewport, CommonParameters.SceneTexturesViewport);
-	SharedParameters.SSAO_Sampler = TStaticSamplerState<SF_Point, AM_Clamp, AM_Clamp, AM_Clamp>::GetRHI();
-	SharedParameters.SSAO_SvPositionScaleBias = bHalfResolution ? FVector2f(2, -0.5) : FVector2f(1, 0);
-	SharedParameters.SSAO_DownsampledAOInverseSize = bHalfResolution ? FVector2f(1, 1) / FVector2f(OutputViewport.Extent) : FVector2f(1, 1);
-	SharedParameters.RandomNormalTexture = GraphBuilder.RegisterExternalTexture(GSystemTextures.SSAORandomization, TEXT("SSAORandomization"));
-	SharedParameters.RandomNormalTextureSampler = TStaticSamplerState<SF_Point, AM_Wrap, AM_Wrap, AM_Wrap>::GetRHI();
-
-	const int32 MobileAmbientOcclusionQuality = FMath::Min(CVarMobileAmbientOcclusionQuality.GetValueOnRenderThread(), 3) - 1;
-
-	int32 ComputeMethod = CVarMobileAmbientOcclusionSSAOCompute.GetValueOnRenderThread();
-	if (ComputeMethod>0)
-	{
-		FRDGEventName EventName(TEXT("AmbientOcclusionCS %dx%d"),
-		OutputViewport.Rect.Width(), OutputViewport.Rect.Height());
-
-		// Compute Shader Path
-		FMobileAmbientOcclusionCS::FParameters* PassParameters = GraphBuilder.AllocParameters<FMobileAmbientOcclusionCS::FParameters>();
-		PassParameters->SharedParameters = MoveTemp(SharedParameters);
-		PassParameters->OutTexture = GraphBuilder.CreateUAV(Output.Texture);
-
-		FMobileAmbientOcclusionCS::FPermutationDomain PermutationVector;
-		PermutationVector.Set<FMobileAmbientOcclusionCS::FUseUpsampleDim>(bHalfResolution);
-		PermutationVector.Set<FMobileAmbientOcclusionCS::FShaderQualityDim>(MobileAmbientOcclusionQuality);
-
-		TShaderMapRef<FMobileAmbientOcclusionCS> ComputeShader(View.ShaderMap, PermutationVector);
-		FComputeShaderUtils::AddPass(
-			GraphBuilder,
-			MoveTemp(EventName),
-			ComputeMethod == 1 ? ERDGPassFlags::Compute : ERDGPassFlags::AsyncCompute,
-			ComputeShader,
-			PassParameters,
-			FComputeShaderUtils::GetGroupCount(OutputViewport.Rect.Size(), 16));
-	}
-	else
-	{
-		FRDGEventName EventName(TEXT("AmbientOcclusionPS %dx%d"),
-		OutputViewport.Rect.Width(), OutputViewport.Rect.Height());
-    
-		FMobileAmbientOcclusionPS::FParameters* PassParameters = GraphBuilder.AllocParameters<FMobileAmbientOcclusionPS::FParameters>();
-		PassParameters->SharedParameters = MoveTemp(SharedParameters);
-		PassParameters->RenderTargets[0] = Output.GetRenderTargetBinding();
-		if (bDepthBoundsTestEnabled)
-		{
-			PassParameters->RenderTargets.DepthStencil = DepthStencilBinding;
-		}
-    
-		FMobileAmbientOcclusionPS::FPermutationDomain PermutationVector;
-		PermutationVector.Set<FMobileAmbientOcclusionPS::FShaderQualityDim>(MobileAmbientOcclusionQuality);
-		PermutationVector.Set<FMobileAmbientOcclusionPS::FOutputDepth>(bOutputDepth);
-    
-		TShaderMapRef<FMobileAmbientOcclusionPS> PixelShader(View.ShaderMap, PermutationVector);
-		TShaderMapRef<FScreenPassVS> VertexShader(View.ShaderMap);
-    
-		check(PassParameters);
-		ClearUnusedGraphResources(PixelShader, PassParameters);
-    
-		GraphBuilder.AddPass(
-			MoveTemp(EventName),
-			PassParameters,
-			ERDGPassFlags::Raster,
-			[&View, Output, OutputViewport, InputViewport, VertexShader, PixelShader, PassParameters, bHalfResolution, bDepthBoundsTestEnabled, DepthFar](FRHICommandListImmediate& RHICmdList)
-			{
-				const FIntRect InputRect = InputViewport.Rect;
-				const FIntPoint InputSize = InputViewport.Extent;
-				const FIntRect OutputRect = OutputViewport.Rect;
-				const FIntPoint OutputSize = OutputRect.Size();
-    
-				RHICmdList.SetViewport(OutputRect.Min.X, OutputRect.Min.Y, 0.0f, OutputRect.Max.X, OutputRect.Max.Y, 1.0f);
-    
-				FGraphicsPipelineStateInitializer GraphicsPSOInit;
-				RHICmdList.ApplyCachedRenderTargets(GraphicsPSOInit);
-    
-				GraphicsPSOInit.BlendState = TStaticBlendState<>::GetRHI();
-				GraphicsPSOInit.RasterizerState = TStaticRasterizerState<FM_Solid, CM_None>::GetRHI();
-				GraphicsPSOInit.DepthStencilState = TStaticDepthStencilState<false, CF_Always>::GetRHI();
-				GraphicsPSOInit.BoundShaderState.VertexDeclarationRHI = GFilterVertexDeclaration.VertexDeclarationRHI;
-				GraphicsPSOInit.BoundShaderState.VertexShaderRHI = VertexShader.GetVertexShader();
-				GraphicsPSOInit.BoundShaderState.PixelShaderRHI = PixelShader.GetPixelShader();
-				GraphicsPSOInit.PrimitiveType = PT_TriangleList;
-				GraphicsPSOInit.bDepthBounds = bDepthBoundsTestEnabled;
-    
-				SetGraphicsPipelineState(RHICmdList, GraphicsPSOInit, 0);
-    
-				SetShaderParameters(RHICmdList, PixelShader, PixelShader.GetPixelShader(), *PassParameters);
-    
-				if (bDepthBoundsTestEnabled)
-				{
-					RHICmdList.SetDepthBounds(DepthFar, 1.0f);
-				}
-    
-				DrawPostProcessPass(
-					RHICmdList,
-					0, 0, OutputSize.X, OutputSize.Y,
-					InputRect.Min.X, InputRect.Min.Y, InputRect.Width(), InputRect.Height(),
-					OutputSize,
-					InputSize,
-					VertexShader,
-					View.StereoViewIndex,
-					false,
-					EDRF_UseTriangleOptimization);
-    
-				if (bDepthBoundsTestEnabled)
-				{
-					RHICmdList.SetDepthBounds(0, 1.0f);
-				}
-			});	
-	}
-}
-
-static void AddMobileAmbientOcclusionUpsamplePass(
-	FRDGBuilder& GraphBuilder,
-	const FViewInfo& View,
-	const FMobileSSAOCommonParameters& CommonParameters,
-	int32 UpsampleQuality,
-	FRDGTextureRef SourceTexture,
-	FScreenPassRenderTarget Output)
-{
-	const bool bDepthBoundsTestEnabled =
-		CVarMobileAmbientOcclusionDepthBoundsTest.GetValueOnRenderThread()
-		&& GSupportsDepthBoundsTest
-		&& CommonParameters.SceneDepth.IsValid()
-		&& CommonParameters.SceneDepth.Texture->Desc.NumSamples == 1;
-
-	const FScreenPassTextureViewport OutputViewport = CommonParameters.SceneTexturesViewport;
-
-	float DepthFar = 0.0f;
-	FDepthStencilBinding DepthStencilBinding(CommonParameters.SceneDepth.Texture, ERenderTargetLoadAction::ELoad, ERenderTargetLoadAction::ELoad, FExclusiveDepthStencil::DepthRead_StencilWrite);
-
-	if (bDepthBoundsTestEnabled)
-	{
-		const FFinalPostProcessSettings& Settings = View.FinalPostProcessSettings;
-		const FMatrix& ProjectionMatrix = View.ViewMatrices.GetViewToClip();
-		const FVector4f Far = (FVector4f)ProjectionMatrix.TransformFVector4(FVector4(0, 0, Settings.AmbientOcclusionFadeDistance));
-		DepthFar = FMath::Min(1.0f, Far.Z / Far.W);
-
-		FRenderTargetParameters* ClearParameters = GraphBuilder.AllocParameters<FRenderTargetParameters>();
-		ClearParameters->RenderTargets[0] = Output.GetRenderTargetBinding();
-		ClearParameters->RenderTargets.DepthStencil = DepthStencilBinding;
-
-		GraphBuilder.AddPass(
-			RDG_EVENT_NAME("DepthBounds ClearQuad(%s)", Output.Texture->Name),
-			ClearParameters,
-			ERDGPassFlags::Raster,
-			[Output, OutputViewport, DepthFar](FRHICommandListImmediate& RHICmdList)
-			{
-				RHICmdList.SetViewport(OutputViewport.Rect.Min.X, OutputViewport.Rect.Min.Y, 0.0f, OutputViewport.Rect.Max.X, OutputViewport.Rect.Max.Y, 1.0f);
-
-				// We must clear all pixels that won't be touched by AO shader.
-				FClearQuadCallbacks Callbacks;
-				Callbacks.PSOModifier = [Output](FGraphicsPipelineStateInitializer& PSOInitializer)
-				{
-					PSOInitializer.bDepthBounds = true;
-				};
-				Callbacks.PreClear = [DepthFar](FRHICommandList& InRHICmdList)
-				{
-					// This is done by rendering a clear quad over a depth range from AmbientOcclusionFadeDistance to far plane.
-					InRHICmdList.SetDepthBounds(0, DepthFar);	// NOTE: Inverted depth
-				};
-				Callbacks.PostClear = [DepthFar](FRHICommandList& InRHICmdList)
-				{
-					// Set depth bounds test to cover everything from near plane to AmbientOcclusionFadeDistance and run AO pixel shader.
-					InRHICmdList.SetDepthBounds(DepthFar, 1.0f);
-				};
-
-				DrawClearQuad(RHICmdList, FLinearColor::White, Callbacks);
-			});
-
-		// Make sure the following pass doesn't clear or ignore the data
-		Output.LoadAction = ERenderTargetLoadAction::ELoad;
-	}
-
+	
 	FMobileAmbientOcclusionUpsamplePS::FParameters* PassParameters = GraphBuilder.AllocParameters<FMobileAmbientOcclusionUpsamplePS::FParameters>();
 	PassParameters->RenderTargets[0] = Output.GetRenderTargetBinding();
 	if (bDepthBoundsTestEnabled)
@@ -1142,8 +153,8 @@ static void AddMobileAmbientOcclusionUpsamplePass(
 	}
 
 	PassParameters->View = GetShaderBinding(View.ViewUniformBuffer);
-	PassParameters->SceneTextures = CommonParameters.SceneTexturesUniformBufferRDG;
-	PassParameters->AOTexture = SourceTexture;
+	PassParameters->SceneTextures = CommonParameters.SceneTextures;
+	PassParameters->AOTexture = Input.Texture;
 	PassParameters->AOSampler = TStaticSamplerState<SF_Bilinear>::GetRHI();
 
 	typename FMobileAmbientOcclusionUpsamplePS::FPermutationDomain PermutationVector;
@@ -1182,77 +193,132 @@ static void AddMobileAmbientOcclusionUpsamplePass(
 		});
 }
 
-static void RenderSSAO(FRDGBuilder& GraphBuilder, const FSceneTextures& SceneTextures, FRDGTextureRef AmbientOcclusionTexture, const TArray<FViewInfo>& Views)
+// @param Levels 0..3, how many different resolution levels we want to render
+FScreenPassTexture AddPostProcessingAmbientOcclusion(
+	FRDGBuilder& GraphBuilder,
+	const FViewInfo& View,
+	const FSSAOCommonParameters& CommonParameters,
+	FScreenPassRenderTarget FinalTarget)
 {
-	RDG_EVENT_SCOPE_STAT(GraphBuilder, MobileSSAO, "MobileSSAO");
-
-	const int32 HalfResolutionSetting = CVarMobileSSAOHalfResolution.GetValueOnRenderThread();
-	const bool bHalfResolution = HalfResolutionSetting > 0;
-	const int32 UpsampleQuality = FMath::Clamp(HalfResolutionSetting - 1, 0, 2);
-	const bool bBilateralUpsample = UpsampleQuality > 0;
-
-	FRDGTextureRef SceneDepthTexture = SceneTextures.Depth.Resolve;
-	FRDGTextureRef HalfResolutionTexture = nullptr;
-	if (bHalfResolution)
+	if (CommonParameters.MobileHalfResolutionSetting > 0)
 	{
-		// Bilateral requires 32bit format for AO + depth.
-		const EPixelFormat Format = bBilateralUpsample ? PF_R8G8B8A8 : PF_R8;
-		const FIntPoint Extent = GetDownscaledExtent(AmbientOcclusionTexture->Desc.Extent, 2);
-		ETextureCreateFlags TextureCreateFlags = TexCreate_ShaderResource | TexCreate_RenderTargetable;
-		if (CVarMobileAmbientOcclusionSSAOCompute.GetValueOnRenderThread() > 0)
+		const FScreenPassTextureViewport SceneViewport(CommonParameters.SceneTexturesViewport);
+		const FScreenPassTextureViewport HalfResOutputViewport(GetDownscaledViewport(SceneViewport, 2));
+
+		FScreenPassRenderTarget HalfResOutput;
 		{
-			TextureCreateFlags |= TexCreate_UAV;
+			bool bBilateralUpsample = (CommonParameters.MobileHalfResolutionSetting > 1);
+			
+			FRDGTextureDesc OutputDesc = FinalTarget.Texture->Desc;
+			OutputDesc.Reset();
+			// R:AmbientOcclusion, GBA:used for depth
+			OutputDesc.Format = bBilateralUpsample ? PF_B8G8R8A8 : PF_R8;
+			OutputDesc.ClearValue = FClearValueBinding::None;
+			OutputDesc.Flags &= ~TexCreate_DepthStencilTargetable;
+			OutputDesc.Flags |= TexCreate_RenderTargetable;
+			if (FSSAOHelper::IsAmbientOcclusionCompute())
+			{
+				OutputDesc.Flags |= TexCreate_UAV;
+			}
+			OutputDesc.Extent = HalfResOutputViewport.Extent;
+
+			HalfResOutput.Texture = GraphBuilder.CreateTexture(OutputDesc, TEXT("HalfResAmbientOcclusion"));
+			HalfResOutput.ViewRect = HalfResOutputViewport.Rect;
+			HalfResOutput.LoadAction = ERenderTargetLoadAction::ENoAction;
+			HalfResOutput.UpdateVisualizeTextureExtent();
 		}
-		HalfResolutionTexture = GraphBuilder.CreateTexture(
-			FRDGTextureDesc::Create2D(Extent, Format, FClearValueBinding::None, TextureCreateFlags),
-			TEXT("HalfResolutionScreenSpaceAO"));
-	}
 
-	for (FViewInfo const& View : Views)
+		FScreenPassTexture UpsampleInput = AddAmbientOcclusionFinalPass(
+			GraphBuilder,
+			View,
+			CommonParameters,
+			FScreenPassTexture(),
+			FScreenPassTexture(),
+			FScreenPassTexture(),
+			HalfResOutput);
+
+		const int32 UpsampleQuality = FMath::Clamp(CommonParameters.MobileHalfResolutionSetting - 1, 0, 2);
+		AddAmbientOcclusionUpsamplePass(
+			GraphBuilder,
+			View,
+			CommonParameters,
+			UpsampleQuality,
+			UpsampleInput,
+			FinalTarget);
+		
+		return FinalTarget;
+	}
+	else
 	{
-		FMobileSSAOCommonParameters Parameters = GetMobileSSAOCommonParameters(GraphBuilder, View, SceneDepthTexture, SceneTextures.MobileUniformBuffer);
-		FScreenPassRenderTarget FinalTarget = FScreenPassRenderTarget(AmbientOcclusionTexture, View.ViewRect, ERenderTargetLoadAction::ENoAction);
-
-		if (HalfResolutionTexture != nullptr)
+		FScreenPassTexture AmbientOcclusionInMip1;
+		FScreenPassTexture AmbientOcclusionPassMip1;
+		if (CommonParameters.Levels >= 2)
 		{
-			FScreenPassRenderTarget HalfResolutionTarget = FScreenPassRenderTarget(HalfResolutionTexture, View.ViewRect, ERenderTargetLoadAction::ENoAction);
-			AddMobileAmbientOcclusionPass(GraphBuilder,	View, Parameters, bHalfResolution, bBilateralUpsample, HalfResolutionTarget);
-			AddMobileAmbientOcclusionUpsamplePass(GraphBuilder, View, Parameters, UpsampleQuality, HalfResolutionTexture, FinalTarget);
+			AmbientOcclusionInMip1 =
+				AddAmbientOcclusionSetupPass(
+					GraphBuilder,
+					View,
+					CommonParameters,
+					CommonParameters.SceneDepth);
+
+			FScreenPassTexture AmbientOcclusionPassMip2;
+			if (CommonParameters.Levels >= 3)
+			{
+				FScreenPassTexture AmbientOcclusionInMip2 =
+					AddAmbientOcclusionSetupPass(
+						GraphBuilder,
+						View,
+						CommonParameters,
+						AmbientOcclusionInMip1);
+
+				AmbientOcclusionPassMip2 =
+					AddAmbientOcclusionStepPass(
+						GraphBuilder,
+						View,
+						CommonParameters,
+						AmbientOcclusionInMip2,
+						AmbientOcclusionInMip2,
+						FScreenPassTexture());
+			}
+
+			AmbientOcclusionPassMip1 =
+				AddAmbientOcclusionStepPass(
+					GraphBuilder,
+					View,
+					CommonParameters,
+					AmbientOcclusionInMip1,
+					AmbientOcclusionInMip1,
+					AmbientOcclusionPassMip2);
 		}
-		else
+
+		FScreenPassTexture SetupTexture = CommonParameters.GBufferA;
+		if (Substrate::IsSubstrateEnabled())
 		{
-			AddMobileAmbientOcclusionPass(GraphBuilder, View, Parameters, bHalfResolution, bBilateralUpsample, FinalTarget);
+			// For Substrate, we invalidate the setup texture for the final pass:
+			//	- We do not need GBufferA, the Substrate TopLayer texture will fill in for that.
+			//	- Setting it to nullptr will make the AddAmbientOcclusionPass use a valid viewport from SceneTextures.
+			SetupTexture.Texture = nullptr;
 		}
+
+		return AddAmbientOcclusionFinalPass(
+			GraphBuilder,
+			View,
+			CommonParameters,
+			SetupTexture,
+			AmbientOcclusionInMip1,
+			AmbientOcclusionPassMip1,
+			FinalTarget);
 	}
 }
 
-bool RequiresMobileAmbientOcclusionPass(FScene* Scene, const FViewInfo& View)
-{
-	const int32 Quality = CVarMobileAmbientOcclusionQuality.GetValueOnRenderThread();
-	return Quality > 0
-		&& View.FinalPostProcessSettings.AmbientOcclusionIntensity > 0
-		&& (View.FinalPostProcessSettings.AmbientOcclusionStaticFraction >= 1 / 100.0f || (Scene && Scene->SkyLight && View.Family->EngineShowFlags.SkyLighting))
-		&& View.Family->EngineShowFlags.Lighting
-		&& !View.bIsReflectionCapture
-		&& !View.bIsPlanarReflection
-		&& !View.Family->EngineShowFlags.HitProxies
-		&& !View.Family->EngineShowFlags.VisualizeLightCulling
-		&& !View.Family->UseDebugViewPS();
-}
 
 FRDGTextureRef CreateMobileScreenSpaceAOTexture(FRDGBuilder& GraphBuilder, const FSceneTexturesConfig& Config)
 {
-	bool bGTAO = (CVarMobileAmbientOcclusionTechnique.GetValueOnRenderThread() == 0);
-	const uint32 DownsampleFactor = bGTAO ? 2 : 1;
-	const FIntPoint Extent = FIntPoint::DivideAndRoundUp(Config.Extent, DownsampleFactor);
 	EPixelFormat Format = PF_R8;
-
-	const EMobileAmbientOcclusionShaderType GtaoShaderType = GetCVarMobileAmbientOcclusionShaderTypeOnRenderThread();
-
-	const bool bComputeShader = ( bGTAO && IsComputeShaderType(GtaoShaderType) ) || CVarMobileAmbientOcclusionSSAOCompute.GetValueOnRenderThread() > 0;
+	const bool bComputeShader = FSSAOHelper::IsAmbientOcclusionCompute();
 
 	// G8 isn't supported as UAV on Android OpenGLES, fall back to RGBA8 for compute shader usage.
-	if (bComputeShader && IsOpenGLPlatform(Config.ShaderPlatform) )
+	if (bComputeShader && IsOpenGLPlatform(Config.ShaderPlatform))
 	{
 		Format = PF_R8G8B8A8;
 	}
@@ -1261,11 +327,30 @@ FRDGTextureRef CreateMobileScreenSpaceAOTexture(FRDGBuilder& GraphBuilder, const
 	{
 		TextureCreateFlags |= TexCreate_UAV;
 	}
-	if (UE::PixelFormat::HasCapabilities(Format, EPixelFormatCapabilities::LossyCompressible) && CVarMobileAmbientOcclusionCompression.GetValueOnRenderThread() != 0)
+	if (UE::PixelFormat::HasCapabilities(Format, EPixelFormatCapabilities::LossyCompressible))
 	{
 		TextureCreateFlags |= TexCreate_LossyCompressionLowBitrate;
 	}
-	return GraphBuilder.CreateTexture(FRDGTextureDesc::Create2D(Extent, Format, FClearValueBinding::Black, TextureCreateFlags), TEXT("ScreenSpaceAO"));
+	return GraphBuilder.CreateTexture(FRDGTextureDesc::Create2D(Config.Extent, Format, FClearValueBinding::Black, TextureCreateFlags), TEXT("ScreenSpaceAO"));
+}
+
+bool RequiresMobileAmbientOcclusionPass(FScene* Scene, const FViewInfo& View)
+{
+	int32 CVarLevel = FSSAOHelper::GetNumAmbientOcclusionLevels();
+	return CVarLevel > 0
+		&& View.FinalPostProcessSettings.AmbientOcclusionIntensity > 0
+		&& (View.FinalPostProcessSettings.AmbientOcclusionStaticFraction >= 1 / 100.0f || (Scene && Scene->SkyLight && View.Family->EngineShowFlags.SkyLighting))
+		&& View.Family->EngineShowFlags.Lighting
+		&& !View.bIsReflectionCapture
+		&& !View.bIsPlanarReflection
+		&& !View.Family->EngineShowFlags.HitProxies
+		&& !View.Family->EngineShowFlags.VisualizeLightCulling
+		&& !View.Family->UseDebugViewPS();
+}
+
+bool MobileSSAOUsesCompute()
+{
+	return FSSAOHelper::IsAmbientOcclusionCompute();
 }
 
 // --------------------------------------------------------------------------------------------------------------------
@@ -1275,19 +360,37 @@ FRDGTextureRef FMobileSceneRenderer::RenderAmbientOcclusion(FRDGBuilder& GraphBu
 	{
 		return nullptr;
 	}
-
+	RDG_EVENT_SCOPE_STAT(GraphBuilder, MobileSSAO, "MobileSSAO");
+	const int32 HalfResolutionSetting = CVarMobileSSAOHalfResolution.GetValueOnRenderThread();
 	FRDGTextureRef AmbientOcclusionTexture = CreateMobileScreenSpaceAOTexture(GraphBuilder, SceneTextures.Config);
+	FScreenPassTexture Output;
 
-	const int32 Technique = CVarMobileAmbientOcclusionTechnique.GetValueOnRenderThread();
-	switch (Technique)
+	for (int32 ViewIndex = 0; ViewIndex < Views.Num(); ++ViewIndex)
 	{
-	case 0:
-		RenderGTAO(GraphBuilder, SceneTextures, AmbientOcclusionTexture, Views);
-		break;
-	case 1:
-		RenderSSAO(GraphBuilder, SceneTextures, AmbientOcclusionTexture, Views);
-		break;
-	}
+		FViewInfo& View = Views[ViewIndex];
+		View.BeginRenderView();
+		
+		uint32 Levels = HalfResolutionSetting ? 1 : FSSAOHelper::ComputeAmbientOcclusionPassCount(View);
+		FSSAOCommonParameters Parameters = GetSSAOCommonParameters(GraphBuilder, View, Levels, false);
+		
+		// In case TAA is not active and we do not do half-res upsampling, do a smoothing pass to hide a sampling pattern 
+		Parameters.bNeedSmoothingPass = (Views.Num() == 1) && (!IsTemporalAccumulationBasedMethod(View.AntiAliasingMethod) && HalfResolutionSetting == 0);
+		Parameters.MobileHalfResolutionSetting = HalfResolutionSetting;
+		if (Parameters.FullscreenType == ESSAOType::EPS 
+			&& IsMobileDeferredShadingEnabled(View.GetShaderPlatform()) 
+			&& MobileUsesGBufferCustomData(View.GetShaderPlatform()))
+		{
+			Parameters.MobileShadingModelsExcludeMask = CVarMobileSSAOShadingModelExcludeMask.GetValueOnRenderThread();
+		}
+		
+		FScreenPassRenderTarget FinalTarget = FScreenPassRenderTarget(AmbientOcclusionTexture, View.ViewRect, (ViewIndex == 0 ? ERenderTargetLoadAction::ENoAction : ERenderTargetLoadAction::ELoad));
 
-	return AmbientOcclusionTexture;
+		Output = AddPostProcessingAmbientOcclusion(
+			GraphBuilder,
+			View,
+			Parameters,
+			FinalTarget);
+	}
+	
+	return Output.Texture;
 }
diff --git a/Engine/Source/Runtime/Renderer/Private/PostProcess/PostProcessAmbientOcclusionMobile.h b/Engine/Source/Runtime/Renderer/Private/PostProcess/PostProcessAmbientOcclusionMobile.h
--- a/Engine/Source/Runtime/Renderer/Private/PostProcess/PostProcessAmbientOcclusionMobile.h
+++ b/Engine/Source/Runtime/Renderer/Private/PostProcess/PostProcessAmbientOcclusionMobile.h
@@ -8,4 +8,5 @@
 #include "CoreMinimal.h"
 #include "RendererInterface.h"
 
-bool RequiresMobileAmbientOcclusionPass(class FScene* Scene, const FViewInfo& View);
\ No newline at end of file
+bool RequiresMobileAmbientOcclusionPass(class FScene* Scene, const FViewInfo& View);
+bool MobileSSAOUsesCompute();
\ No newline at end of file
