Fixed: Compile errors in PostProcessAmbientOcclusionMobile.cpp
#Synced-CL 50477682

diff --git a/Engine/Source/Runtime/Renderer/Private/PostProcess/PostProcessAmbientOcclusionMobile.cpp b/Engine/Source/Runtime/Renderer/Private/PostProcess/PostProcessAmbientOcclusionMobile.cpp
--- a/Engine/Source/Runtime/Renderer/Private/PostProcess/PostProcessAmbientOcclusionMobile.cpp
+++ b/Engine/Source/Runtime/Renderer/Private/PostProcess/PostProcessAmbientOcclusionMobile.cpp
@@ -83,7 +83,7 @@ class FMobileAmbientOcclusionUpsamplePS : public FGlobalShader
 IMPLEMENT_GLOBAL_SHADER(FMobileAmbientOcclusionUpsamplePS, "/Engine/Private/PostProcessAmbientOcclusion.usf", "AmbientOcclusionUpsamplePS", SF_Pixel);
 
 // --------------------------------------------------------------------------------------------------------------------
-void AddAmbientOcclusionUpsamplePass(
+static void AddAmbientOcclusionUpsamplePass(
 	FRDGBuilder& GraphBuilder,
 	const FViewInfo& View,
 	const FSSAOCommonParameters& CommonParameters,
@@ -193,8 +193,7 @@ void AddAmbientOcclusionUpsamplePass(
 		});
 }
 
-// @param Levels 0..3, how many different resolution levels we want to render
-FScreenPassTexture AddPostProcessingAmbientOcclusion(
+static FScreenPassTexture AddPostProcessingAmbientOcclusionMobile(
 	FRDGBuilder& GraphBuilder,
 	const FViewInfo& View,
 	const FSSAOCommonParameters& CommonParameters,
@@ -312,7 +311,7 @@ FScreenPassTexture AddPostProcessingAmbientOcclusion(
 }
 
 
-FRDGTextureRef CreateMobileScreenSpaceAOTexture(FRDGBuilder& GraphBuilder, const FSceneTexturesConfig& Config)
+static FRDGTextureRef CreateMobileScreenSpaceAOTexture(FRDGBuilder& GraphBuilder, const FSceneTexturesConfig& Config)
 {
 	EPixelFormat Format = PF_R8;
 	const bool bComputeShader = FSSAOHelper::IsAmbientOcclusionCompute();
@@ -385,7 +384,7 @@ FRDGTextureRef FMobileSceneRenderer::RenderAmbientOcclusion(FRDGBuilder& GraphBu
 		
 		FScreenPassRenderTarget FinalTarget = FScreenPassRenderTarget(AmbientOcclusionTexture, View.ViewRect, (ViewIndex == 0 ? ERenderTargetLoadAction::ENoAction : ERenderTargetLoadAction::ELoad));
 
-		Output = AddPostProcessingAmbientOcclusion(
+		Output = AddPostProcessingAmbientOcclusionMobile(
 			GraphBuilder,
 			View,
 			Parameters,
