Fixed incorrect instance accounting in the new CPU-side hierarchical instance culling.

#Synced-CL 50496977

#Synced-CL 50496977

diff --git a/Engine/Source/Runtime/Renderer/Private/InstanceCulling/InstanceCullingContext.cpp b/Engine/Source/Runtime/Renderer/Private/InstanceCulling/InstanceCullingContext.cpp
--- a/Engine/Source/Runtime/Renderer/Private/InstanceCulling/InstanceCullingContext.cpp
+++ b/Engine/Source/Runtime/Renderer/Private/InstanceCulling/InstanceCullingContext.cpp
@@ -1655,6 +1655,7 @@ void FInstanceCullingContext::SetupDrawCommands(
 							UnpackScreenSizeHelper(VisibleMeshDrawCommand.CullingPayload.MaxScreenSize, CurrentCullingPayloadFlags, EMeshDrawCommandCullingPayloadFlags::MaxScreenSizeCull, std::numeric_limits<float>::max())
 						);
 						
+						int32 NumInstancesDrawn = 0;
 						int32 FirstInstance = 0;
 						for (int32 ItemIndex = 0; ItemIndex < PrimitiveCullingData->Items.Num(); ++ItemIndex)
 						{
@@ -1664,7 +1665,8 @@ void FInstanceCullingContext::SetupDrawCommands(
 							// TODO: track start/end and only add once if e.g., all are visible?
 							if (bIsVisible && CullingQuery->IsLODVisible(*PrimitiveCullingData, ItemIndex, MinMaxLODScreenSize))
 							{
-								AddInstancesToDrawCommand(CurrentIndirectArgsOffset, VisibleMeshDrawCommand.PrimitiveIdInfo.InstanceSceneDataOffset + FirstInstance, FirstInstance, Item.NumInstances, InstanceFlags, MaxGenericBatchSize);
+								AddInstancesToDrawCommand(CurrentIndirectArgsOffset, VisibleMeshDrawCommand.PrimitiveIdInfo.InstanceSceneDataOffset + FirstInstance, NumInstancesDrawn, Item.NumInstances, InstanceFlags, MaxGenericBatchSize);
+								NumInstancesDrawn += Item.NumInstances;
 							}
 							FirstInstance += Item.NumInstances;
 						}
