Added a single-pass optimization for <= 32 bones and removed bUnifiedBoneMap from FSkeletalMeshRenderData. A follow-up CL will introduce a unified bone map storage option that works. Converted 32 bit indices to 16 bits. Removed previous transform histories for auto play, as it's too complicated with blending and over can't resolve play rate effectively.

#Synced-CL 50498661

diff --git a/Engine/Shaders/Private/Skinning/AnimSequenceTransformProvider.usf b/Engine/Shaders/Private/Skinning/AnimSequenceTransformProvider.usf
--- a/Engine/Shaders/Private/Skinning/AnimSequenceTransformProvider.usf
+++ b/Engine/Shaders/Private/Skinning/AnimSequenceTransformProvider.usf
@@ -19,6 +19,8 @@ RWByteAddressBuffer LocalSpaceTransformBufferUAV;
 #define RETARGET_MODE_ANIMATION_RELATIVE 3
 #define RETARGET_MODE_ORIENT_AND_SCALE 4
 
+#define INVALID_BONE_INDEX 0xFFFF
+
 FBoneRetargetingData LoadBoneRetargetingData(ByteAddressBuffer Buffer, uint BaseOffset, uint Index)
 {
 	uint Offset = BaseOffset + Index * (uint)sizeof(FBoneRetargetingData);
@@ -64,7 +66,7 @@ FBoneTransformWithUniformScale EvaluateAndRetargetBone(
 
 	FBoneTransformWithUniformScale RetargetingTransform = LoadBoneTransformWithUniformScale(SequenceTransformBuffer, RetargetingTransformBufferOffset, RetargetingTransformIndex);
 
-	if (RetargetingData.AnimBoneIndex == BONE_RETARGETING_DATA_INVALID_BONE_INDEX)
+	if (RetargetingData.AnimBoneIndex == INVALID_BONE_INDEX)
 	{
 		// When the anim bone doesn't exist the ref pose is stored in the retargeting transform.
 		return RetargetingTransform;
@@ -201,7 +203,7 @@ RWByteAddressBuffer GlobalSpaceTransformBufferUAV;
 
 groupshared FBoneTransformWithUniformScale SharedTransform[BONES_PER_GROUP];
 
-FBoneTransformWithUniformScale LoadSharedTransform(uint TransformBufferOffset, uint BoneIndex)
+FBoneTransformWithUniformScale LoadLocalSpaceTransform(uint TransformBufferOffset, uint BoneIndex)
 {
 	if (BoneIndex < BONES_PER_GROUP)
 	{
@@ -213,11 +215,6 @@ FBoneTransformWithUniformScale LoadSharedTransform(uint TransformBufferOffset, u
 	}
 }
 
-void StoreSharedTransform(uint BoneIndex, FBoneTransformWithUniformScale Transform)
-{
-	SharedTransform[BoneIndex] = Transform;
-}
-
 [numthreads(BONES_PER_GROUP, 1, 1)]
 void ConcatenateLocalToGlobalCS(uint3 GroupId : SV_GroupID, uint GroupIndex : SV_GroupIndex)
 {
@@ -230,7 +227,7 @@ void ConcatenateLocalToGlobalCS(uint3 GroupId : SV_GroupID, uint GroupIndex : SV
 
 	if (PrefetchBoneIndex < BlockHeader.NumPrefetchBones)
 	{
-		StoreSharedTransform(PrefetchBoneIndex, LoadBoneTransformWithUniformScale(LocalSpaceTransformBuffer, BlockHeader.TransformBufferOffset, PrefetchBoneIndex));
+		SharedTransform[PrefetchBoneIndex] = LoadBoneTransformWithUniformScale(LocalSpaceTransformBuffer, BlockHeader.TransformBufferOffset, PrefetchBoneIndex);
 	}
 	GroupMemoryBarrierWithGroupSync();
 
@@ -246,7 +243,7 @@ void ConcatenateLocalToGlobalCS(uint3 GroupId : SV_GroupID, uint GroupIndex : SV
 	const uint PathHeaderOffset = BlockHeader.LeafBonePathBufferOffset + RelativePathHeaderOffset;
 
 	FLeafBonePathHeader PathHeader = LeafBonePathBuffer.Load<FLeafBonePathHeader>(PathHeaderOffset);
-	FBoneTransformWithUniformScale ParentTransform = LoadSharedTransform(BlockHeader.TransformBufferOffset, 0);
+	FBoneTransformWithUniformScale ParentTransform = LoadLocalSpaceTransform(BlockHeader.TransformBufferOffset, 0);
 
 	if (PathHeader.bIsRootOwned)
 	{
@@ -263,7 +260,7 @@ void ConcatenateLocalToGlobalCS(uint3 GroupId : SV_GroupID, uint GroupIndex : SV
 			PathNode = LeafBonePathBuffer.Load<FLeafBonePathNode>(PathNodeOffset);
 			PathNodeOffset += PathElementSize;
 
-			FBoneTransformWithUniformScale ChildTransform = LoadSharedTransform(BlockHeader.TransformBufferOffset, PathNode.BoneIndex);
+			FBoneTransformWithUniformScale ChildTransform = LoadLocalSpaceTransform(BlockHeader.TransformBufferOffset, PathNode.BoneIndex);
 			ParentTransform = MultiplyBoneTransformWithUniformScale(ChildTransform, ParentTransform);
 
 			if (PathNode.bIsOwned)
@@ -284,6 +281,18 @@ ByteAddressBuffer GlobalSpaceTransformBuffer;
 ByteAddressBuffer BoneMapBuffer;
 FCompressedBoneTransformRWBuffer SkinningTransformBufferUAV;
 
+void ScatterGlobalSpaceTransform(uint DstTransformBufferOffset1, uint DstTransformBufferOffset2, uint TransformIndex, FBoneTransformWithUniformScale GlobalTransform)
+{
+	float4x3 OutputMatrix = (float4x3)MatrixFromBoneTransformWithUniformScale(GlobalTransform);
+
+	StoreCompressedBoneTransform(SkinningTransformBufferUAV, DstTransformBufferOffset1, TransformIndex, OutputMatrix);
+
+	if (DstTransformBufferOffset1 != DstTransformBufferOffset2)
+	{
+		StoreCompressedBoneTransform(SkinningTransformBufferUAV, DstTransformBufferOffset2, TransformIndex, OutputMatrix);
+	}
+}
+
 [numthreads(BONES_PER_GROUP, 1, 1)]
 void ScatterTransformsCS(uint3 GroupId : SV_GroupID, uint GroupIndex : SV_GroupIndex)
 {
@@ -299,20 +308,115 @@ void ScatterTransformsCS(uint3 GroupId : SV_GroupID, uint GroupIndex : SV_GroupI
 
 	const uint TransformIndex = LocalTransformIndex + (BlockHeader.BlockLocalIndex * BONES_PER_GROUP);
 
-	const uint CompactBoneIndex = BoneMapBuffer.Load<uint>(BlockHeader.BoneMapOffset + TransformIndex * uint(sizeof(uint)));
+	const uint CompactBoneIndex = BoneMapBuffer.Load<uint>(BlockHeader.BoneMapBufferOffset + TransformIndex * uint(sizeof(uint)));
 
-	if (CompactBoneIndex != ~0u)
+	if (CompactBoneIndex != INVALID_BONE_INDEX)
 	{
 		FBoneTransformWithUniformScale InvGlobalRefPoseTransform = LoadBoneTransformWithUniformScale(SequenceTransformBuffer, BlockHeader.InvGlobalRefPoseTransformBufferOffset, CompactBoneIndex);
 		FBoneTransformWithUniformScale GlobalSpaceTransform = LoadBoneTransformWithUniformScale(GlobalSpaceTransformBuffer, BlockHeader.SrcTransformBufferOffset, CompactBoneIndex);
 		GlobalSpaceTransform = MultiplyBoneTransformWithUniformScale(InvGlobalRefPoseTransform, GlobalSpaceTransform);
-		float4x3 OutputMatrix = (float4x3)MatrixFromBoneTransformWithUniformScale(GlobalSpaceTransform);
 
-		StoreCompressedBoneTransform(SkinningTransformBufferUAV, BlockHeader.DstTransformBufferOffset1, TransformIndex, OutputMatrix);
+		ScatterGlobalSpaceTransform(BlockHeader.DstTransformBufferOffset1, BlockHeader.DstTransformBufferOffset2, TransformIndex, GlobalSpaceTransform);
+	}
+}
 
-		if (BlockHeader.DstTransformBufferOffset1 != BlockHeader.DstTransformBufferOffset2)
-		{
-			StoreCompressedBoneTransform(SkinningTransformBufferUAV, BlockHeader.DstTransformBufferOffset2, TransformIndex, OutputMatrix);
-		}
+FBoneHierarchyData LoadBoneHierarchyData(ByteAddressBuffer Buffer, uint BaseOffset, uint Index)
+{
+	uint Offset = BaseOffset + Index * (uint)sizeof(FBoneHierarchyData);
+	return Buffer.Load<FBoneHierarchyData>(Offset);
+}
+
+FBoneTransformWithUniformScale ConcatenateToGlobal(
+	uint BoneIndex,
+	uint BoneMapBufferOffset,
+	uint InvGlobalRefPoseTransformBufferOffset,
+	out uint OutTransformIndex)
+{
+	FBoneTransformWithUniformScale GlobalTransform = SharedTransform[BoneIndex];
+	
+	FBoneHierarchyData HierarchyData = LoadBoneHierarchyData(BoneMapBuffer, BoneMapBufferOffset, BoneIndex);
+	OutTransformIndex = HierarchyData.TransformIndex;
+
+	while (HierarchyData.ParentIndex != INVALID_BONE_INDEX)
+	{
+		GlobalTransform = MultiplyBoneTransformWithUniformScale(GlobalTransform, SharedTransform[HierarchyData.ParentIndex]);
+		HierarchyData = LoadBoneHierarchyData(BoneMapBuffer, BoneMapBufferOffset, HierarchyData.ParentIndex);
+	}
+
+	FBoneTransformWithUniformScale InvRefPose = LoadBoneTransformWithUniformScale(SequenceTransformBuffer, InvGlobalRefPoseTransformBufferOffset, BoneIndex);
+	return MultiplyBoneTransformWithUniformScale(InvRefPose, GlobalTransform);
+}
+
+[numthreads(BONES_PER_GROUP, 1, 1)]
+void EvaluateConcatenateScatterCS(uint3 GroupId : SV_GroupID, uint GroupIndex : SV_GroupIndex)
+{
+	const uint BlockIndex = GroupId.x;
+	const uint BoneIndex = GroupIndex;
+
+	FEvaluateConcatenateScatterBlockHeader BlockHeader = EvaluateBlockHeaderBuffer.Load<FEvaluateConcatenateScatterBlockHeader>(BlockIndex * uint(sizeof(FEvaluateConcatenateScatterBlockHeader)));
+
+	if (BoneIndex < BlockHeader.NumCompactBones)
+	{
+		SharedTransform[BoneIndex] = EvaluateAndRetargetBone(
+			BoneIndex,
+			BlockHeader.SequenceTransformBufferOffset,
+			BlockHeader.RetargetingTransformBufferOffset,
+			BlockHeader.RetargetingDataBufferOffset,
+			BlockHeader.KeyIndex0,
+			BlockHeader.KeyIndex1,
+			BlockHeader.Alpha,
+			BlockHeader.NumAnimBones);
+	}
+
+	GroupMemoryBarrierWithGroupSync();
+
+	if (BoneIndex < BlockHeader.NumCompactBones)
+	{
+		uint TransformIndex;
+		FBoneTransformWithUniformScale GlobalTransform = ConcatenateToGlobal(BoneIndex, BlockHeader.BoneMapBufferOffset, BlockHeader.InvGlobalRefPoseTransformBufferOffset, TransformIndex);
+		ScatterGlobalSpaceTransform(BlockHeader.DstTransformBufferOffset1, BlockHeader.DstTransformBufferOffset2, TransformIndex, GlobalTransform);
+	}
+}
+
+[numthreads(BONES_PER_GROUP, 1, 1)]
+void EvaluateBlendConcatenateScatterCS(uint3 GroupId : SV_GroupID, uint GroupIndex : SV_GroupIndex)
+{
+	const uint BlockIndex = GroupId.x;
+	const uint BoneIndex = GroupIndex;
+
+	FEvaluateBlendConcatenateScatterBlockHeader BlockHeader = EvaluateBlockHeaderBuffer.Load<FEvaluateBlendConcatenateScatterBlockHeader>(BlockIndex * uint(sizeof(FEvaluateBlendConcatenateScatterBlockHeader)));
+
+	if (BoneIndex < BlockHeader.NumCompactBones)
+	{
+		FBoneTransformWithUniformScale SourceTransform = EvaluateAndRetargetBone(
+			BoneIndex,
+			BlockHeader.SourceSequenceTransformBufferOffset,
+			BlockHeader.SourceRetargetingTransformBufferOffset,
+			BlockHeader.SourceRetargetingDataBufferOffset,
+			BlockHeader.SourceKeyIndex0,
+			BlockHeader.SourceKeyIndex1,
+			BlockHeader.SourceAlpha,
+			BlockHeader.SourceNumAnimBones);
+
+		FBoneTransformWithUniformScale TargetTransform = EvaluateAndRetargetBone(
+			BoneIndex,
+			BlockHeader.TargetSequenceTransformBufferOffset,
+			BlockHeader.TargetRetargetingTransformBufferOffset,
+			BlockHeader.TargetRetargetingDataBufferOffset,
+			BlockHeader.TargetKeyIndex0,
+			BlockHeader.TargetKeyIndex1,
+			BlockHeader.TargetAlpha,
+			BlockHeader.TargetNumAnimBones);
+
+		SharedTransform[BoneIndex] = BlendBoneTransformsWithUniformScale(SourceTransform, TargetTransform, BlockHeader.BlendWeight);
+	}
+
+	GroupMemoryBarrierWithGroupSync();
+
+	if (BoneIndex < BlockHeader.NumCompactBones)
+	{
+		uint TransformIndex;
+		FBoneTransformWithUniformScale GlobalTransform = ConcatenateToGlobal(BoneIndex, BlockHeader.BoneMapBufferOffset, BlockHeader.InvGlobalRefPoseTransformBufferOffset, TransformIndex);
+		ScatterGlobalSpaceTransform(BlockHeader.DstTransformBufferOffset1, BlockHeader.DstTransformBufferOffset2, TransformIndex, GlobalTransform);
 	}
 }
\ No newline at end of file
diff --git a/Engine/Shaders/Shared/SkinningDefinitions.h b/Engine/Shaders/Shared/SkinningDefinitions.h
--- a/Engine/Shaders/Shared/SkinningDefinitions.h
+++ b/Engine/Shaders/Shared/SkinningDefinitions.h
@@ -37,15 +37,13 @@ struct FAnimBankBlockHeader
 	uint BlockTransformCount            : 16;
 };
 
-#define BONE_RETARGETING_DATA_INVALID_BONE_INDEX 0xFFFF
-
 struct FBoneRetargetingData
 {
 	uint AnimBoneIndex : 16;
 	uint RetargetMode  : 15;
 	uint bRequiresRefPoseRetarget : 1;
 };
-SKINNING_STATIC_ASSERT(sizeof(FBoneRetargetingData) == 4, "FBoneRetargetingData must be 32 bytes");
+SKINNING_STATIC_ASSERT(sizeof(FBoneRetargetingData) == 4, "FBoneRetargetingData must be 4 bytes");
 
 struct FEvaluateAnimSequenceBlockHeader
 {
@@ -85,6 +83,56 @@ struct FEvaluateAnimSequenceBlendBlockHeader
  };
 SKINNING_STATIC_ASSERT(sizeof(FEvaluateAnimSequenceBlendBlockHeader) == 64, "FEvaluateAnimSequenceBlendBlockHeader must be 64 bytes");
 
+struct FEvaluateConcatenateScatterBlockHeader
+{
+	float Alpha;
+	uint KeyIndex0;
+	uint KeyIndex1;
+	uint SequenceTransformBufferOffset;
+	uint RetargetingTransformBufferOffset;
+	uint RetargetingDataBufferOffset;
+	uint InvGlobalRefPoseTransformBufferOffset;
+	uint BoneMapBufferOffset;
+	uint DstTransformBufferOffset1;
+	uint DstTransformBufferOffset2;
+	uint NumAnimBones    : 16;
+	uint NumCompactBones : 16;
+};
+SKINNING_STATIC_ASSERT(sizeof(FEvaluateConcatenateScatterBlockHeader) == 44, "FEvaluateConcatenateScatterBlockHeader must be 44 bytes");
+
+struct FEvaluateBlendConcatenateScatterBlockHeader
+{
+	float SourceAlpha;
+	float TargetAlpha;
+	float BlendWeight;
+	uint SourceKeyIndex0;
+	uint SourceKeyIndex1;
+	uint TargetKeyIndex0;
+	uint TargetKeyIndex1;
+	uint SourceSequenceTransformBufferOffset;
+	uint TargetSequenceTransformBufferOffset;
+	uint SourceRetargetingTransformBufferOffset;
+	uint TargetRetargetingTransformBufferOffset;
+	uint SourceRetargetingDataBufferOffset;
+	uint TargetRetargetingDataBufferOffset;
+	uint InvGlobalRefPoseTransformBufferOffset;
+	uint BoneMapBufferOffset;
+	uint DstTransformBufferOffset1;
+	uint DstTransformBufferOffset2;
+	uint SourceNumAnimBones : 16;
+	uint TargetNumAnimBones : 16;
+	uint NumCompactBones    : 16;
+	uint Unused             : 16;
+};
+SKINNING_STATIC_ASSERT(sizeof(FEvaluateBlendConcatenateScatterBlockHeader) == 76, "FEvaluateBlendConcatenateScatterBlockHeader must be 76 bytes");
+
+struct FBoneHierarchyData
+{
+	uint ParentIndex    : 16;
+	uint TransformIndex : 16;
+};
+SKINNING_STATIC_ASSERT(sizeof(FBoneHierarchyData) == 4, "FBoneHierarchyData must be 4 bytes");
+
 struct FLeafBonePathHeader
 {
 	uint PathDepth      : 31;
@@ -116,7 +164,7 @@ struct FScatterTransformsBlockHeader
 	uint SrcTransformBufferOffset;
 	uint DstTransformBufferOffset1;
 	uint DstTransformBufferOffset2;
-	uint BoneMapOffset;
+	uint BoneMapBufferOffset;
 	uint BlockLocalIndex     : 16;
 	uint BlockTransformCount : 16;
 };
diff --git a/Engine/Source/Runtime/Engine/Private/SkeletalMeshRenderData.cpp b/Engine/Source/Runtime/Engine/Private/SkeletalMeshRenderData.cpp
--- a/Engine/Source/Runtime/Engine/Private/SkeletalMeshRenderData.cpp
+++ b/Engine/Source/Runtime/Engine/Private/SkeletalMeshRenderData.cpp
@@ -468,7 +468,6 @@ FSkeletalMeshRenderData::FSkeletalMeshRenderData()
 	, PendingFirstLODIdx(0)
 	, LODBiasModifier(0)
 	, bSupportRayTracing(true)
-	, bUnifiedBoneMap(false)
 	, bInitialized(false)
 {
 	ClearNaniteResources(NaniteResourcesPtr);
@@ -627,38 +626,6 @@ PRAGMA_DISABLE_DEPRECATION_WARNINGS
 PRAGMA_ENABLE_DEPRECATION_WARNINGS
 }
 
-static bool HasUnifiedBoneMap(const TIndirectArray<FSkeletalMeshLODRenderData>& LODRenderData)
-{
-	const int32 NumLODs = LODRenderData.Num();
-
-	if (NumLODs <= 1)
-	{
-		return true;
-	}
-
-	const FSkeletalMeshLODRenderData& LOD0 = LODRenderData[0];
-
-	for (int32 LODIndex = 1; LODIndex < NumLODs; ++LODIndex)
-	{
-		const FSkeletalMeshLODRenderData& LODData = LODRenderData[LODIndex];
-
-		if (LODData.RenderSections.Num() != LOD0.RenderSections.Num())
-		{
-			return false;
-		}
-
-		for (int32 SectionIndex = 0; SectionIndex < LOD0.RenderSections.Num(); ++SectionIndex)
-		{
-			if (LOD0.RenderSections[SectionIndex].BoneMap != LODData.RenderSections[SectionIndex].BoneMap)
-			{
-				return false;
-			}
-		}
-	}
-
-	return true;
-}
-
 void FSkeletalMeshRenderData::InitResources(bool bNeedsVertexColors, USkinnedAsset* Owner)
 {
 	if (!bInitialized)
@@ -685,8 +652,6 @@ void FSkeletalMeshRenderData::InitResources(bool bNeedsVertexColors, USkinnedAss
 			}
 		}
 
-		bUnifiedBoneMap = ::HasUnifiedBoneMap(LODRenderData);
-
 		check(NaniteResourcesPtr.IsValid());
 		NaniteResourcesPtr->InitResources(Owner);
 
diff --git a/Engine/Source/Runtime/Engine/Public/Rendering/SkeletalMeshRenderData.h b/Engine/Source/Runtime/Engine/Public/Rendering/SkeletalMeshRenderData.h
--- a/Engine/Source/Runtime/Engine/Public/Rendering/SkeletalMeshRenderData.h
+++ b/Engine/Source/Runtime/Engine/Public/Rendering/SkeletalMeshRenderData.h
@@ -43,9 +43,6 @@ public:
 	/** Whether ray tracing acceleration structures should be created for this mesh. Derived from owner USkinnedAsset. */
 	bool bSupportRayTracing;
 
-	/** True if all LODs share the same section bone map (enables instanced skinned GPU LOD without indirection) */
-	bool bUnifiedBoneMap;
-
 	/** Optional, name of the owning skeletal mesh */
 	FName OwnerName = NAME_None;
 
@@ -147,7 +144,7 @@ public:
 
 	bool HasUnifiedBoneMap() const
 	{
-		return bUnifiedBoneMap;
+		return false;
 	}
 
 	ENGINE_API bool HasValidNaniteData() const;
diff --git a/Engine/Source/Runtime/Renderer/Private/Skinning/AnimSequenceTransformProvider.cpp b/Engine/Source/Runtime/Renderer/Private/Skinning/AnimSequenceTransformProvider.cpp
--- a/Engine/Source/Runtime/Renderer/Private/Skinning/AnimSequenceTransformProvider.cpp
+++ b/Engine/Source/Runtime/Renderer/Private/Skinning/AnimSequenceTransformProvider.cpp
@@ -36,6 +36,8 @@ DECLARE_MEMORY_STAT(TEXT("Retargeting Data Buffer"), STAT_AnimSequenceTransformP
 DECLARE_MEMORY_STAT(TEXT("Leaf Bone Path Buffer"), STAT_AnimSequenceTransformProvider_LeafBonePathMemory, STATGROUP_AnimSequenceTransformProvider);
 
 DECLARE_DWORD_COUNTER_STAT(TEXT("Active Tracks"), STAT_AnimSequenceTransformProvider_ActiveTracks, STATGROUP_AnimSequenceTransformProvider);
+DECLARE_DWORD_COUNTER_STAT(TEXT("Evaluate SinglePass Blocks"), STAT_AnimSequenceTransformProvider_EvaluateSinglePassBlocks, STATGROUP_AnimSequenceTransformProvider);
+DECLARE_DWORD_COUNTER_STAT(TEXT("Evaluate SinglePass Blend Blocks"), STAT_AnimSequenceTransformProvider_EvaluateSinglePassBlendBlocks, STATGROUP_AnimSequenceTransformProvider);
 DECLARE_DWORD_COUNTER_STAT(TEXT("Evaluate Blocks"), STAT_AnimSequenceTransformProvider_EvaluateBlocks, STATGROUP_AnimSequenceTransformProvider);
 DECLARE_DWORD_COUNTER_STAT(TEXT("Evaluate Blend Blocks"), STAT_AnimSequenceTransformProvider_EvaluateBlendBlocks, STATGROUP_AnimSequenceTransformProvider);
 DECLARE_DWORD_COUNTER_STAT(TEXT("Current Transform Updates"), STAT_AnimSequenceTransformProvider_CurrentTransformUpdates, STATGROUP_AnimSequenceTransformProvider);
@@ -283,6 +285,70 @@ public:
 };
 IMPLEMENT_GLOBAL_SHADER(FScatterTransformsCS, "/Engine/Private/Skinning/AnimSequenceTransformProvider.usf", "ScatterTransformsCS", SF_Compute);
 
+class FEvaluateConcatenateScatterCS : public FGlobalShader
+{
+public:
+	DECLARE_GLOBAL_SHADER(FEvaluateConcatenateScatterCS);
+	SHADER_USE_PARAMETER_STRUCT(FEvaluateConcatenateScatterCS, FGlobalShader);
+
+	static constexpr uint32 BonesPerGroup = 32u;
+
+	static bool ShouldCompilePermutation(const FGlobalShaderPermutationParameters& Parameters)
+	{
+		return true;
+	}
+
+	static void ModifyCompilationEnvironment(const FGlobalShaderPermutationParameters& Parameters, FShaderCompilerEnvironment& OutEnvironment)
+	{
+		FGlobalShader::ModifyCompilationEnvironment(Parameters, OutEnvironment);
+		OutEnvironment.SetDefine(TEXT("BONES_PER_GROUP"), BonesPerGroup);
+
+		OutEnvironment.CompilerFlags.Add(CFLAG_WarningsAsErrors);
+		OutEnvironment.CompilerFlags.Add(CFLAG_HLSL2021);
+	}
+
+	BEGIN_SHADER_PARAMETER_STRUCT(FParameters, )
+		SHADER_PARAMETER_RDG_BUFFER_SRV(ByteAddressBuffer, RetargetingDataBuffer)
+		SHADER_PARAMETER_RDG_BUFFER_SRV(ByteAddressBuffer, SequenceTransformBuffer)
+		SHADER_PARAMETER_RDG_BUFFER_SRV(ByteAddressBuffer, EvaluateBlockHeaderBuffer)
+		SHADER_PARAMETER_RDG_BUFFER_SRV(ByteAddressBuffer, BoneMapBuffer)
+		SHADER_PARAMETER_RDG_COMPRESSED_BONE_TRANSFORM_UAV(SkinningTransformBufferUAV)
+	END_SHADER_PARAMETER_STRUCT()
+};
+IMPLEMENT_GLOBAL_SHADER(FEvaluateConcatenateScatterCS, "/Engine/Private/Skinning/AnimSequenceTransformProvider.usf", "EvaluateConcatenateScatterCS", SF_Compute);
+
+class FEvaluateBlendConcatenateScatterCS : public FGlobalShader
+{
+public:
+	DECLARE_GLOBAL_SHADER(FEvaluateBlendConcatenateScatterCS);
+	SHADER_USE_PARAMETER_STRUCT(FEvaluateBlendConcatenateScatterCS, FGlobalShader);
+
+	static constexpr uint32 BonesPerGroup = FEvaluateConcatenateScatterCS::BonesPerGroup;
+
+	static bool ShouldCompilePermutation(const FGlobalShaderPermutationParameters& Parameters)
+	{
+		return true;
+	}
+
+	static void ModifyCompilationEnvironment(const FGlobalShaderPermutationParameters& Parameters, FShaderCompilerEnvironment& OutEnvironment)
+	{
+		FGlobalShader::ModifyCompilationEnvironment(Parameters, OutEnvironment);
+		OutEnvironment.SetDefine(TEXT("BONES_PER_GROUP"), BonesPerGroup);
+
+		OutEnvironment.CompilerFlags.Add(CFLAG_WarningsAsErrors);
+		OutEnvironment.CompilerFlags.Add(CFLAG_HLSL2021);
+	}
+
+	BEGIN_SHADER_PARAMETER_STRUCT(FParameters, )
+		SHADER_PARAMETER_RDG_BUFFER_SRV(ByteAddressBuffer, RetargetingDataBuffer)
+		SHADER_PARAMETER_RDG_BUFFER_SRV(ByteAddressBuffer, SequenceTransformBuffer)
+		SHADER_PARAMETER_RDG_BUFFER_SRV(ByteAddressBuffer, EvaluateBlockHeaderBuffer)
+		SHADER_PARAMETER_RDG_BUFFER_SRV(ByteAddressBuffer, BoneMapBuffer)
+		SHADER_PARAMETER_RDG_COMPRESSED_BONE_TRANSFORM_UAV(SkinningTransformBufferUAV)
+	END_SHADER_PARAMETER_STRUCT()
+};
+IMPLEMENT_GLOBAL_SHADER(FEvaluateBlendConcatenateScatterCS, "/Engine/Private/Skinning/AnimSequenceTransformProvider.usf", "EvaluateBlendConcatenateScatterCS", SF_Compute);
+
 FAnimSequenceTransformProvider::FProxyUserData::FProxyUserData(FAnimSequenceTransformProviderProxy* Proxy)
 	: Tracks(Proxy->GetTracks())
 	, SceneExtensionProxy(*Proxy->GetSceneExtensionProxy())
@@ -323,70 +389,42 @@ void FAnimSequenceTransformProvider::InitExtension(FScene& InScene)
 
 void FAnimSequenceTransformProvider::BuildBoneTables(
 	FBoneTables& BoneTables,
-	EBoneTransformStorageMode TransformStorageMode,
 	const USkinnedAsset* SkinnedMesh,
 	const FSkeletalMeshLODRenderData& LODRenderData)
 {
 	const USkeleton* TargetSkeleton = SkinnedMesh->GetSkeleton();
 	const FReferenceSkeleton& MeshRefSkeleton = SkinnedMesh->GetRefSkeleton();
 	const FReferenceSkeleton& TargetSkeletonRefSkeleton = TargetSkeleton->GetReferenceSkeleton();
-	const int32 NumMeshBones = MeshRefSkeleton.GetRawBoneNum();
-	const int32 NumCompactBones = LODRenderData.RequiredBones.Num();
+	const uint32 NumMeshBones = MeshRefSkeleton.GetRawBoneNum();
+	const uint32 NumCompactBones = LODRenderData.RequiredBones.Num();
 
-	BoneTables.MeshToCompact.Init(INDEX_NONE, NumMeshBones);
+	BoneTables.MeshToCompact.Init(INVALID_BONE_INDEX, NumMeshBones);
 	BoneTables.CompactToMesh.SetNumUninitialized(NumCompactBones);
 	BoneTables.CompactToParent.SetNumUninitialized(NumCompactBones);
 	BoneTables.CompactToTargetSkeleton.SetNumUninitialized(NumCompactBones);
 	BoneTables.NumMeshBones = NumMeshBones;
-	BoneTables.NumTransforms = NumMeshBones;
 	BoneTables.NumCompactBones = NumCompactBones;
 
-	for (int32 CompactBoneIndex = 0; CompactBoneIndex < NumCompactBones; ++CompactBoneIndex)
+	for (uint32 CompactBoneIndex = 0; CompactBoneIndex < NumCompactBones; ++CompactBoneIndex)
 	{
 		BoneTables.MeshToCompact[LODRenderData.RequiredBones[CompactBoneIndex]] = CompactBoneIndex;
 	}
 
-	for (int32 CompactBoneIndex = 0; CompactBoneIndex < NumCompactBones; ++CompactBoneIndex)
+	for (uint32 CompactBoneIndex = 0; CompactBoneIndex < NumCompactBones; ++CompactBoneIndex)
 	{
-		const int32 MeshBoneIndex = LODRenderData.RequiredBones[CompactBoneIndex];
+		const uint16 MeshBoneIndex = LODRenderData.RequiredBones[CompactBoneIndex];
 		BoneTables.CompactToMesh[CompactBoneIndex] = MeshBoneIndex;
 
-		const int32 MeshParentBoneIndex = MeshRefSkeleton.GetParentIndex(MeshBoneIndex);
-		const int32 CompactParentBoneIndex = MeshParentBoneIndex != INDEX_NONE
+		const uint16 MeshParentBoneIndex = MeshRefSkeleton.GetParentIndex(MeshBoneIndex);
+		const uint16 CompactParentBoneIndex = MeshParentBoneIndex != INVALID_BONE_INDEX
 			? BoneTables.MeshToCompact[MeshParentBoneIndex]
-			: INDEX_NONE;
+			: INVALID_BONE_INDEX;
 		BoneTables.CompactToParent[CompactBoneIndex] = CompactParentBoneIndex;
 
 		const FName BoneName = MeshRefSkeleton.GetBoneName(MeshBoneIndex);
-		const int32 TargetSkeletonBoneIndex = TargetSkeletonRefSkeleton.FindBoneIndex(BoneName);
+		const uint16 TargetSkeletonBoneIndex = TargetSkeletonRefSkeleton.FindBoneIndex(BoneName);
 		BoneTables.CompactToTargetSkeleton[CompactBoneIndex] = TargetSkeletonBoneIndex;
 	}
-
-	if (TransformStorageMode == EBoneTransformStorageMode::BoneMap)
-	{
-		TConstArrayView<FSkelMeshRenderSection> Sections = LODRenderData.RenderSections;
-		uint32 NumTransforms = 0;
-
-		for (const FSkelMeshRenderSection& Section : Sections)
-		{
-			if (Section.IsValid())
-			{
-				NumTransforms += Section.BoneMap.Num();
-			}
-		}
-
-		BoneTables.TransformToMesh.Empty(NumTransforms);
-
-		for (const FSkelMeshRenderSection& Section : Sections)
-		{
-			if (Section.IsValid())
-			{
-				BoneTables.TransformToMesh.Append(Section.BoneMap);
-			}
-		}
-
-		BoneTables.NumTransforms = NumTransforms;
-	}
 }
 
 void FAnimSequenceTransformProvider::BuildRetargetingTables(
@@ -409,7 +447,7 @@ void FAnimSequenceTransformProvider::BuildRetargetingTables(
 
 	RetargetingTables.Mode.Reset(BoneTables.NumCompactBones);
 	RetargetingTables.Transform.Reset(BoneTables.NumCompactBones);
-	RetargetingTables.AnimToCompact.Init(INDEX_NONE, NumAnimBones);
+	RetargetingTables.AnimToCompact.Init(INVALID_BONE_INDEX, NumAnimBones);
 	RetargetingTables.bRequiresRefPoseRetarget = bRequiresRefPoseRetarget;
 
 	if (bRequiresRefPoseRetarget)
@@ -421,10 +459,10 @@ void FAnimSequenceTransformProvider::BuildRetargetingTables(
 	{
 		EBoneTranslationRetargetingMode::Type Mode = EBoneTranslationRetargetingMode::Animation;
 
-		const int32 MeshBoneIndex = BoneTables.CompactToMesh[CompactBoneIndex];
-		const int32 TargetSkeletonBoneIndex = BoneTables.CompactToTargetSkeleton[CompactBoneIndex];
-		
-		const int32 AnimBoneIndex = TargetSkeletonBoneIndex != INDEX_NONE && SkeletonRemapping.IsValid()
+		const uint16 MeshBoneIndex = BoneTables.CompactToMesh[CompactBoneIndex];
+		const uint16 TargetSkeletonBoneIndex = BoneTables.CompactToTargetSkeleton[CompactBoneIndex];
+
+		const uint16 AnimBoneIndex = TargetSkeletonBoneIndex != INVALID_BONE_INDEX && SkeletonRemapping.IsValid()
 			? SkeletonRemapping.GetSourceSkeletonBoneIndex(TargetSkeletonBoneIndex)
 			: TargetSkeletonBoneIndex;
 
@@ -432,7 +470,7 @@ void FAnimSequenceTransformProvider::BuildRetargetingTables(
 		TTuple<FQuat, FQuat> RefPoseRetargetQuats(FQuat::Identity, FQuat::Identity);
 		FTransform3f Transform;
 
-		if (AnimBoneIndex != INDEX_NONE)
+		if (AnimBoneIndex != INVALID_BONE_INDEX)
 		{
 			if (bRequiresRefPoseRetarget)
 			{
@@ -543,7 +581,7 @@ void FAnimSequenceTransformProvider::BuildRetargetingTables(
 
 void FAnimSequenceTransformProvider::BuildLeafBonePathData(FLeafBonePathData& OutData, const FBoneTables& BoneTables)
 {
-	const int32 NumBones = BoneTables.NumCompactBones;
+	const uint32 NumBones = BoneTables.NumCompactBones;
 
 	struct FBoneData
 	{
@@ -554,12 +592,12 @@ void FAnimSequenceTransformProvider::BuildLeafBonePathData(FLeafBonePathData& Ou
 
 	TArray<FBoneData, FConcurrentLinearArrayAllocator> BoneDatas;
 	BoneDatas.SetNum(NumBones);
-	int32 NumLeafBones = NumBones;
+	uint32 NumLeafBones = NumBones;
 
-	for (int32 ChildIndex = 1; ChildIndex < NumBones; ++ChildIndex)
+	for (uint32 ChildIndex = 1; ChildIndex < NumBones; ++ChildIndex)
 	{
-		const int32 ParentIndex = BoneTables.CompactToParent[ChildIndex];
-		if (ParentIndex != INDEX_NONE)
+		const uint32 ParentIndex = BoneTables.CompactToParent[ChildIndex];
+		if (ParentIndex != INVALID_BONE_INDEX)
 		{
 			FBoneData& ParentBone = BoneDatas[ParentIndex];
 			FBoneData& ChildBone  = BoneDatas[ChildIndex];
@@ -574,8 +612,8 @@ void FAnimSequenceTransformProvider::BuildLeafBonePathData(FLeafBonePathData& Ou
 		}
 	}
 
-	int32 NumPathNodes = 0;
-	for (int32 BoneIndex = 0; BoneIndex < NumBones; ++BoneIndex)
+	uint32 NumPathNodes = 0;
+	for (uint32 BoneIndex = 0; BoneIndex < NumBones; ++BoneIndex)
 	{
 		if (BoneDatas[BoneIndex].bIsLeaf)
 		{
@@ -596,7 +634,7 @@ void FAnimSequenceTransformProvider::BuildLeafBonePathData(FLeafBonePathData& Ou
 	int32 LeafIndex = 0;
 	uint32 PathDataOffset = OffsetTableSize * sizeof(uint32);
 
-	for (int32 BoneIndex = 0; BoneIndex < NumBones; ++BoneIndex)
+	for (uint32 BoneIndex = 0; BoneIndex < NumBones; ++BoneIndex)
 	{
 		if (!BoneDatas[BoneIndex].bIsLeaf)
 		{
@@ -627,7 +665,7 @@ void FAnimSequenceTransformProvider::BuildLeafBonePathData(FLeafBonePathData& Ou
 
 			while (CurrentNodeIndex >= 0)
 			{
-				check(CurrentBoneIndex != INDEX_NONE);
+				check(CurrentBoneIndex != INVALID_BONE_INDEX);
 
 				FBoneData& CurrentBone = BoneDatas[CurrentBoneIndex];
 
@@ -672,11 +710,16 @@ void FAnimSequenceTransformProvider::RegisterProxy(FAnimSequenceTransformProvide
 	for (int32 LODIndex = MinLODIndex; LODIndex < NumLODs; ++LODIndex)
 	{
 		FProxyUserData::FLOD& LOD = ProxyUserData->LODs[LODIndex];
-		BuildBoneTables(LOD.BoneTables, ExtensionProxy->GetBoneTransformStorageMode(), SkinnedAsset, RenderData->LODRenderData[LODIndex]);
-		BuildLeafBonePathData(LOD.LeafBonePathData, LOD.BoneTables);
-	}
+		BuildBoneTables(LOD.BoneTables, SkinnedAsset, RenderData->LODRenderData[LODIndex]);
 
-	ProxyUserDatasToUpload.Emplace(ProxyUserData);
+		// Single pass mode skips building leaf bone paths in favor of doing everything in one compute dispatch.
+		const bool bSinglePass = LOD.BoneTables.NumCompactBones <= FEvaluateConcatenateScatterCS::BonesPerGroup && (RenderData->HasUnifiedBoneMap() || ExtensionProxy->GetBoneTransformStorageMode() == EBoneTransformStorageMode::Direct);
+		if (!bSinglePass)
+		{
+			BuildLeafBonePathData(LOD.LeafBonePathData, LOD.BoneTables);
+		}
+		LOD.bSinglePass = bSinglePass;
+	}
 
 	for (const FAnimSequenceTransformProviderSequence& ProviderSequence : ProviderSequences)
 	{
@@ -736,7 +779,7 @@ void FAnimSequenceTransformProvider::RegisterProxy(FAnimSequenceTransformProvide
 				Mapping->InvGlobalRefPoseTransforms.Reserve(Mapping->NumCompactBones);
 				for (uint32 CompactBoneIndex = 0; CompactBoneIndex < Mapping->NumCompactBones; ++CompactBoneIndex)
 				{
-					const int32 MeshBoneIndex = LOD.BoneTables.CompactToMesh[CompactBoneIndex];
+					const uint32 MeshBoneIndex = LOD.BoneTables.CompactToMesh[CompactBoneIndex];
 					const FMatrix44f& InvGlobalRefPose = InvGlobalRefPoseTransforms[MeshBoneIndex];
 					Mapping->InvGlobalRefPoseTransforms.Emplace(InvGlobalRefPose);
 				}
@@ -764,8 +807,6 @@ void FAnimSequenceTransformProvider::UnregisterProxy(FAnimSequenceTransformProvi
 	check(ProxyUserData);
 	ExtensionProxy->SetUserData(nullptr);
 
-	ProxyUserDatasToUpload.RemoveSingleSwap(ProxyUserData);
-
 	for (FProxyUserData::FLOD& LOD : ProxyUserData->LODs)
 	{
 		if (LOD.LeafBonePathBufferAllocation.IsValid())
@@ -796,8 +837,13 @@ void FAnimSequenceTransformProvider::UnregisterProxy(FAnimSequenceTransformProvi
 		}
 	}
 
-	for (FMapping* Mapping : ProxyUserData->Mappings)
+	for (FMapping* Mapping : ProxyUserData->SequenceToMappingTable)
 	{
+		if (!Mapping)
+		{
+			continue;
+		}
+
 		check(Mapping && Mapping->RefCount > 0);
 		--Mapping->RefCount;
 
@@ -831,6 +877,8 @@ template <typename T>
 class TBufferUploadData
 {
 public:
+	TBufferUploadData() = default;
+
 	TBufferUploadData(T* InElements, int32 InMaxElements)
 		: Elements(InElements)
 		, MaxElements(InMaxElements)
@@ -843,28 +891,34 @@ public:
 	}
 
 private:
-	T* Elements;
-	int32 MaxElements;
+	T* Elements = nullptr;
+	int32 MaxElements = 0;
 	int32 NumElements = 0;
 };
 
 template <typename T>
 FRDGBuffer* CreateByteAddressBuffer(FRDGBuilder& GraphBuilder, uint32 MaxNumElements, const TCHAR* Name)
 {
-	return GraphBuilder.CreateBuffer(FRDGBufferDesc::CreateByteAddressDesc(MaxNumElements * sizeof(T)), Name);
+	return MaxNumElements > 0 ? GraphBuilder.CreateBuffer(FRDGBufferDesc::CreateByteAddressDesc(MaxNumElements * sizeof(T)), Name) : nullptr;
 }
 
 template <typename T>
 TBufferUploadData<T> Lock(FRDGBuilder& GraphBuilder, FRDGBuffer* Buffer)
 {
-	GraphBuilder.ConvertToExternalBuffer(Buffer);
-
-	return TBufferUploadData<T>(reinterpret_cast<T*>(GraphBuilder.RHICmdList.LockBuffer(Buffer->GetRHIUnchecked(), 0, Buffer->Desc.GetSize(), RLM_WriteOnly)), Buffer->Desc.GetSize() / sizeof(T));
+	if (Buffer)
+	{
+		GraphBuilder.ConvertToExternalBuffer(Buffer);
+		return TBufferUploadData<T>(reinterpret_cast<T*>(GraphBuilder.RHICmdList.LockBuffer(Buffer->GetRHIUnchecked(), 0, Buffer->Desc.GetSize(), RLM_WriteOnly)), Buffer->Desc.GetSize() / sizeof(T));
+	}
+	return {};
 }
 
 void Unlock(FRDGBuilder& GraphBuilder, FRDGBuffer* Buffer)
 {
-	GraphBuilder.RHICmdList.UnlockBuffer(Buffer->GetRHIUnchecked());
+	if (Buffer)
+	{
+		GraphBuilder.RHICmdList.UnlockBuffer(Buffer->GetRHIUnchecked());
+	}
 }
 
 void FAnimSequenceTransformProvider::ProvideGPUTransforms(FSkinningTransformProvider::FProviderContext& Context)
@@ -906,16 +960,26 @@ void FAnimSequenceTransformProvider::ProvideGPUTransforms(FSkinningTransformProv
 			return static_cast<EPreviousBoneTransformUpdateMode>(UpdateMode);
 		}
 	};
+	
+	TArray<FTrackUpdateRequest, SceneRenderingAllocator> TrackUpdateRequests;
+	TrackUpdateRequests.Reserve(MaxTotalActiveTracks);
+
+	struct FLODUploadRequest
+	{
+		FProxyUserData::FLOD& LOD;
+		const FInstancedSkinningSceneExtensionProxy& ExtensionProxy;
+	};
+
+	TArray<FLODUploadRequest, SceneRenderingAllocator> LODUploadRequests;
 
 	const float GlobalTimeScale = CVarAnimSequenceTransformProviderTimeScale.GetValueOnRenderThread();
 	const double WorldTime = Context.ViewFamily.Time.GetWorldTimeSeconds() * GlobalTimeScale;
 	const float DeltaTime = Context.ViewFamily.Time.GetDeltaWorldTimeSeconds() * GlobalTimeScale;
 
-	TArray<FTrackUpdateRequest, SceneRenderingAllocator> TrackUpdateRequests;
-	TrackUpdateRequests.Reserve(MaxTotalActiveTracks);
-
 	uint32 NumTotalSubTracks = 0;
 	uint32 NumTotalIntermediateTransforms = 0;
+	uint32 NumTotalEvaluateConcatenateScatterBlockHeaders = 0;
+	uint32 NumTotalEvaluateBlendConcatenateScatterBlockHeaders = 0;
 	uint32 NumTotalEvaluateBlockHeaders = 0;
 	uint32 NumTotalEvaluateBlendBlockHeaders = 0;
 	uint32 NumTotalScatterBlockHeaders = 0;
@@ -926,29 +990,30 @@ void FAnimSequenceTransformProvider::ProvideGPUTransforms(FSkinningTransformProv
 		FInstancedSkinningSceneExtensionProxy* ExtensionProxy = static_cast<FInstancedSkinningSceneExtensionProxy*>(Context.Proxies[Indirection.Index]);
 		FProxyUserData& ProxyUserData = static_cast<FProxyUserData&>(ExtensionProxy->GetUserDataChecked());
 		const uint32 LODIndex = ExtensionProxy->GetLOD();
-		const FProxyUserData::FLOD& LOD = ProxyUserData.LODs[LODIndex];
+		FProxyUserData::FLOD& LOD = ProxyUserData.LODs[LODIndex];
 		FAnimSequenceTrackRenderPool& Tracks = ProxyUserData.Tracks;
 		ProxyUserData.TrackHistories.SetNum(Tracks.GetNumTracks());
 
-		const uint32 NumEvaluateBlockHeaders    = FMath::DivideAndRoundUp(LOD.BoneTables.NumCompactBones, BonesPerGroup);
-		const uint32 NumScatterBlockHeaders     = FMath::DivideAndRoundUp(LOD.BoneTables.NumTransforms, BonesPerGroup);
-		const uint32 NumConcatenateBlockHeaders = FMath::DivideAndRoundUp(LOD.LeafBonePathData.NumLeafBones, BonesPerGroup);
-		const uint32 CurrentTransformSlot       = Indirection.CurrentTransformOffset < Indirection.PreviousTransformOffset ? 0 : 1;
-	
-		int32 NumSubTracks = 0;
-		int32 NumEvaluateSubTracks = 0;
-		int32 NumEvaluateBlendSubTracks = 0;
+		const uint32 NumEvaluateConcatenateBlendBlockHeaders = LOD.bSinglePass ? 1 : 0;
+		const uint32 NumEvaluateBlockHeaders                 = LOD.bSinglePass ? 0 : FMath::DivideAndRoundUp<uint32>(LOD.BoneTables.NumCompactBones, BonesPerGroup);
+		const uint32 NumScatterBlockHeaders                  = LOD.bSinglePass ? 0 : FMath::DivideAndRoundUp<uint32>(ExtensionProxy->GetMaxBoneTransformCount(), BonesPerGroup);
+		const uint32 NumConcatenateBlockHeaders              = LOD.bSinglePass ? 0 : FMath::DivideAndRoundUp<uint32>(LOD.LeafBonePathData.NumLeafBones, BonesPerGroup);
+		const uint32 CurrentTransformSlot                    = Indirection.CurrentTransformOffset < Indirection.PreviousTransformOffset ? 0 : 1;
+
+		uint32 NumSubTracks = 0;
+		uint32 NumEvaluateSubTracks = 0;
+		uint32 NumEvaluateBlendSubTracks = 0;
+		uint32 NumIntermediateTransforms = 0;
 
 		Tracks.EnumerateActiveTracks([&] (int32 TrackIndex)
 		{
 			const FAnimSequenceTrackRenderData& Data = Tracks.GetData(TrackIndex);
-			const FTrackData SourceTrack             = ProxyUserData.GetTrackData(Data.GetSource(), WorldTime, LODIndex);
 			const FTrackData TargetTrack             = ProxyUserData.GetTrackData(Data.GetTarget(), WorldTime, LODIndex);
 			const float BlendWeight                  = Data.GetBlendWeight();
 
 			// Dirty previous transforms means the whole allocation is dirty. Rebuild everything.
 			const bool bForceUpdateAll = EnumHasAnyFlags(Indirection.DirtyBoneTransforms, EDirtyBoneTransforms::Previous);
-			FTrackHistory::FUpdateResult UpdateResult = ProxyUserData.TrackHistories[TrackIndex].Update(SourceTrack, TargetTrack, CurrentTransformSlot, bForceUpdateAll);
+			FTrackHistory::FUpdateResult UpdateResult = ProxyUserData.TrackHistories[TrackIndex].Update(TargetTrack, CurrentTransformSlot, bForceUpdateAll);
 
 			FTrackUpdateRequest UpdateRequest;
 			UpdateRequest.SetUpdateMode(EPreviousBoneTransformUpdateMode::None);
@@ -956,26 +1021,44 @@ void FAnimSequenceTransformProvider::ProvideGPUTransforms(FSkinningTransformProv
 
 			if (UpdateResult.bNeedsPrevious)
 			{
-				UpdateRequest.UpdateMode = uint8(UpdateResult.bNeedsCurrent && UpdateResult.bNeedsPrevious && TargetTrack.Position == TargetTrack.PreviousPosition && SourceTrack.Position == SourceTrack.PreviousPosition
+				UpdateRequest.UpdateMode = uint8(UpdateResult.bNeedsCurrent && UpdateResult.bNeedsPrevious && TargetTrack.Position == TargetTrack.PreviousPosition
 					? EPreviousBoneTransformUpdateMode::DuplicateCurrentToPrevious
 					: EPreviousBoneTransformUpdateMode::UpdatePrevious);
 			}
 
 			const int32 NumLocalSubTracks = UpdateResult.bNeedsCurrent + (UpdateRequest.GetUpdateMode() == EPreviousBoneTransformUpdateMode::UpdatePrevious ? 1 : 0);
 
-			NumSubTracks              += NumLocalSubTracks;
-			NumEvaluateSubTracks      += BlendWeight <= 0.0f ? NumLocalSubTracks : 0;
-			NumEvaluateBlendSubTracks += BlendWeight >  0.0f ? NumLocalSubTracks : 0;
+			NumSubTracks                   += NumLocalSubTracks;
+			NumEvaluateSubTracks           += BlendWeight <= 0.0f ? NumLocalSubTracks : 0;
+			NumEvaluateBlendSubTracks      += BlendWeight >  0.0f ? NumLocalSubTracks : 0;
+			NumTotalIntermediateTransforms += LOD.bSinglePass ? 0 : NumLocalSubTracks * LOD.BoneTables.NumCompactBones;
 
 			TrackUpdateRequests.Emplace(UpdateRequest);
 		});
 
-		NumTotalSubTracks                 += NumSubTracks;
-		NumTotalIntermediateTransforms    += NumSubTracks * LOD.BoneTables.NumCompactBones;
-		NumTotalEvaluateBlockHeaders      += NumEvaluateSubTracks * NumEvaluateBlockHeaders;
-		NumTotalEvaluateBlendBlockHeaders += NumEvaluateBlendSubTracks * NumEvaluateBlockHeaders;
-		NumTotalScatterBlockHeaders       += NumSubTracks * NumScatterBlockHeaders;
-		NumTotalConcatenateBlockHeaders   += NumSubTracks * NumConcatenateBlockHeaders;
+		NumTotalSubTracks                                   += NumSubTracks;
+		NumTotalEvaluateConcatenateScatterBlockHeaders      += NumEvaluateSubTracks * NumEvaluateConcatenateBlendBlockHeaders;
+		NumTotalEvaluateBlendConcatenateScatterBlockHeaders += NumEvaluateBlendSubTracks * NumEvaluateConcatenateBlendBlockHeaders;
+		NumTotalEvaluateBlockHeaders                        += NumEvaluateSubTracks * NumEvaluateBlockHeaders;
+		NumTotalEvaluateBlendBlockHeaders                   += NumEvaluateBlendSubTracks * NumEvaluateBlockHeaders;
+		NumTotalScatterBlockHeaders                         += NumSubTracks * NumScatterBlockHeaders;
+		NumTotalConcatenateBlockHeaders                     += NumSubTracks * NumConcatenateBlockHeaders;
+
+		if (!LOD.bUploaded && NumSubTracks > 0)
+		{
+			LOD.bUploaded = true;
+
+			if (LODUploadRequests.IsEmpty())
+			{
+				LODUploadRequests.Reserve(Context.Indirections.Num());
+			}
+
+			LODUploadRequests.Add(FLODUploadRequest
+			{
+				  .LOD = LOD
+				, .ExtensionProxy = *ExtensionProxy
+			});
+		}
 	}
 
 	if (!NumTotalSubTracks)
@@ -989,10 +1072,10 @@ void FAnimSequenceTransformProvider::ProvideGPUTransforms(FSkinningTransformProv
 
 	for (FMapping* Mapping : MappingsToUpload)
 	{
-		const int32 NumInvGlobalRefPoses     = Mapping->NumCompactBones;
-		const int32 NumRetargetingTransforms = (Mapping->RetargetingTables.bRequiresRefPoseRetarget ? 2 : 1);
-		const int32 NumTransforms            = Mapping->NumCompactBones * NumRetargetingTransforms + NumInvGlobalRefPoses;
-		const int32 NumRetargetingDatas      = Mapping->NumCompactBones;
+		const uint32 NumInvGlobalRefPoses     = Mapping->NumCompactBones;
+		const uint32 NumRetargetingTransforms = (Mapping->RetargetingTables.bRequiresRefPoseRetarget ? 2 : 1);
+		const uint32 NumTransforms            = Mapping->NumCompactBones * NumRetargetingTransforms + NumInvGlobalRefPoses;
+		const uint32 NumRetargetingDatas      = Mapping->NumCompactBones;
 
 		Mapping->MappingTransformBufferAllocation.Count  = NumTransforms;
 		Mapping->MappingTransformBufferAllocation.Offset = SequenceTransformAllocator.Allocate(NumTransforms);
@@ -1028,37 +1111,39 @@ void FAnimSequenceTransformProvider::ProvideGPUTransforms(FSkinningTransformProv
 		NumTotalSequenceTransforms += NumTransforms;
 	}
 
-	for (FProxyUserData* ProxyUserData : ProxyUserDatasToUpload)
+	for (const FLODUploadRequest& Request : LODUploadRequests)
 	{
-		for (FProxyUserData::FLOD& LOD : ProxyUserData->LODs)
-		{
-			if (LOD.IsValid())
-			{
-				const int32 NumLeafBonePaths = LOD.LeafBonePathData.Buffer.Num();
-				const int32 NumElements = NumLeafBonePaths + LOD.BoneTables.NumTransforms;
+		FProxyUserData::FLOD& LOD = Request.LOD;
+		const int32 NumLeafBonePaths = LOD.LeafBonePathData.Buffer.Num();
 
-				LOD.LeafBonePathBufferAllocation.Count = NumElements;
-				LOD.LeafBonePathBufferAllocation.Offset = LeafBonePathAllocator.Allocate(NumElements);
+		// When in single pass mode the bone map is in compact bone space, otherwise it's in transform space.
+		const uint32 NumBoneMapEntries = LOD.bSinglePass ? LOD.BoneTables.NumCompactBones : Request.ExtensionProxy.GetMaxBoneTransformCount();
 
-				LOD.BoneMapBufferOffset = LOD.LeafBonePathBufferAllocation.Offset + NumLeafBonePaths;
+		const uint32 NumElements = NumLeafBonePaths + NumBoneMapEntries;
+		LOD.LeafBonePathBufferAllocation.Count = NumElements;
+		LOD.LeafBonePathBufferAllocation.Offset = LeafBonePathAllocator.Allocate(NumElements);
 
-				NumTotalLeafBonePathElements += NumElements;
-			}
-		}
+		LOD.BoneMapBufferOffset = LOD.LeafBonePathBufferAllocation.Offset + NumLeafBonePaths;
+
+		NumTotalLeafBonePathElements += NumElements;
 	}
 
 	int32& NumTotalFramesRequested = *GraphBuilder.AllocPOD<int32>();
 	NumTotalFramesRequested = 0;
-
-	FRDGBuffer* EvaluateBlockHeaderBuffer      = CreateByteAddressBuffer<UE::HLSL::FEvaluateAnimSequenceBlockHeader>(GraphBuilder, NumTotalEvaluateBlockHeaders, TEXT("AnimSequenceTransformProvider.EvaluateBlockHeaders"));
-	FRDGBuffer* EvaluateBlendBlockHeaderBuffer = CreateByteAddressBuffer<UE::HLSL::FEvaluateAnimSequenceBlendBlockHeader>(GraphBuilder, NumTotalEvaluateBlendBlockHeaders, TEXT("AnimSequenceTransformProvider.EvaluateBlendBlockHeaders"));
-	FRDGBuffer* ScatterBlockHeaderBuffer       = CreateByteAddressBuffer<UE::HLSL::FScatterTransformsBlockHeader>(GraphBuilder, NumTotalScatterBlockHeaders, TEXT("AnimSequenceTransformProvider.ScatterBlockHeaders"));
-	FRDGBuffer* ConcatenateBlockHeaderBuffer   = CreateByteAddressBuffer<UE::HLSL::FConcatenateBlockHeader>(GraphBuilder, NumTotalConcatenateBlockHeaders, TEXT("AnimSequenceTransformProvider.ConcatenateBlockHeaders"));
-
-	TBufferUploadData<UE::HLSL::FEvaluateAnimSequenceBlockHeader> EvaluateBlockHeaders           = Lock<UE::HLSL::FEvaluateAnimSequenceBlockHeader>(GraphBuilder, EvaluateBlockHeaderBuffer);
-	TBufferUploadData<UE::HLSL::FEvaluateAnimSequenceBlendBlockHeader> EvaluateBlendBlockHeaders = Lock<UE::HLSL::FEvaluateAnimSequenceBlendBlockHeader>(GraphBuilder, EvaluateBlendBlockHeaderBuffer);
-	TBufferUploadData<UE::HLSL::FScatterTransformsBlockHeader> ScatterBlockHeaders               = Lock<UE::HLSL::FScatterTransformsBlockHeader>(GraphBuilder, ScatterBlockHeaderBuffer);
-	TBufferUploadData<UE::HLSL::FConcatenateBlockHeader> ConcatenateBlockHeaders                 = Lock<UE::HLSL::FConcatenateBlockHeader>(GraphBuilder, ConcatenateBlockHeaderBuffer);
+	
+	FRDGBuffer* EvaluateConcatenateScatterBlockHeaderBuffer      = CreateByteAddressBuffer<UE::HLSL::FEvaluateConcatenateScatterBlockHeader>(GraphBuilder, NumTotalEvaluateConcatenateScatterBlockHeaders, TEXT("AnimSequenceTransformProvider.EvaluateConcatenateScatterBlockHeaders"));
+	FRDGBuffer* EvaluateBlendConcatenateScatterBlockHeaderBuffer = CreateByteAddressBuffer<UE::HLSL::FEvaluateBlendConcatenateScatterBlockHeader>(GraphBuilder, NumTotalEvaluateBlendConcatenateScatterBlockHeaders, TEXT("AnimSequenceTransformProvider.EvaluateBlendConcatenateScatterBlockHeaders"));
+	FRDGBuffer* EvaluateBlockHeaderBuffer                        = CreateByteAddressBuffer<UE::HLSL::FEvaluateAnimSequenceBlockHeader>(GraphBuilder, NumTotalEvaluateBlockHeaders, TEXT("AnimSequenceTransformProvider.EvaluateBlockHeaders"));
+	FRDGBuffer* EvaluateBlendBlockHeaderBuffer                   = CreateByteAddressBuffer<UE::HLSL::FEvaluateAnimSequenceBlendBlockHeader>(GraphBuilder, NumTotalEvaluateBlendBlockHeaders, TEXT("AnimSequenceTransformProvider.EvaluateBlendBlockHeaders"));
+	FRDGBuffer* ScatterBlockHeaderBuffer                         = CreateByteAddressBuffer<UE::HLSL::FScatterTransformsBlockHeader>(GraphBuilder, NumTotalScatterBlockHeaders, TEXT("AnimSequenceTransformProvider.ScatterBlockHeaders"));
+	FRDGBuffer* ConcatenateBlockHeaderBuffer                     = CreateByteAddressBuffer<UE::HLSL::FConcatenateBlockHeader>(GraphBuilder, NumTotalConcatenateBlockHeaders, TEXT("AnimSequenceTransformProvider.ConcatenateBlockHeaders"));
+	
+	TBufferUploadData<UE::HLSL::FEvaluateConcatenateScatterBlockHeader> EvaluateConcatenateScatterBlockHeaders           = Lock<UE::HLSL::FEvaluateConcatenateScatterBlockHeader>(GraphBuilder, EvaluateConcatenateScatterBlockHeaderBuffer);
+	TBufferUploadData<UE::HLSL::FEvaluateBlendConcatenateScatterBlockHeader> EvaluateBlendConcatenateScatterBlockHeaders = Lock<UE::HLSL::FEvaluateBlendConcatenateScatterBlockHeader>(GraphBuilder, EvaluateBlendConcatenateScatterBlockHeaderBuffer);
+	TBufferUploadData<UE::HLSL::FEvaluateAnimSequenceBlockHeader> EvaluateBlockHeaders                                   = Lock<UE::HLSL::FEvaluateAnimSequenceBlockHeader>(GraphBuilder, EvaluateBlockHeaderBuffer);
+	TBufferUploadData<UE::HLSL::FEvaluateAnimSequenceBlendBlockHeader> EvaluateBlendBlockHeaders                         = Lock<UE::HLSL::FEvaluateAnimSequenceBlendBlockHeader>(GraphBuilder, EvaluateBlendBlockHeaderBuffer);
+	TBufferUploadData<UE::HLSL::FScatterTransformsBlockHeader> ScatterBlockHeaders                                       = Lock<UE::HLSL::FScatterTransformsBlockHeader>(GraphBuilder, ScatterBlockHeaderBuffer);
+	TBufferUploadData<UE::HLSL::FConcatenateBlockHeader> ConcatenateBlockHeaders                                         = Lock<UE::HLSL::FConcatenateBlockHeader>(GraphBuilder, ConcatenateBlockHeaderBuffer);
 
 	UE::Tasks::FTask ProcessTask = GraphBuilder.AddSetupTask(
 		[
@@ -1069,6 +1154,8 @@ void FAnimSequenceTransformProvider::ProvideGPUTransforms(FSkinningTransformProv
 			, Proxies = Context.Proxies
 			, &NumTotalFramesRequested
 			, TrackUpdateRequests = MoveTemp(TrackUpdateRequests)
+			, EvaluateConcatenateScatterBlockHeaders
+			, EvaluateBlendConcatenateScatterBlockHeaders
 			, EvaluateBlockHeaders
 			, EvaluateBlendBlockHeaders
 			, ScatterBlockHeaders
@@ -1079,6 +1166,8 @@ void FAnimSequenceTransformProvider::ProvideGPUTransforms(FSkinningTransformProv
 
 		uint32 IntermediateTransformOffset = 0;
 		uint32 NumActiveTracks = 0;
+		uint32 NumEvaluateConcatenateScatterBlocks = 0;
+		uint32 NumEvaluateBlendConcatenateScatterBlocks = 0;
 		uint32 NumEvaluateBlocks = 0;
 		uint32 NumEvaluateBlendBlocks = 0;
 		uint32 NumCurrentTransformUpdates = 0;
@@ -1098,129 +1187,184 @@ void FAnimSequenceTransformProvider::ProvideGPUTransforms(FSkinningTransformProv
 			uint32 DstTransformBufferOffset1,
 			uint32 DstTransformBufferOffset2)
 		{
-			const auto GetBlockTransformCount = [](uint32 BlockIndex, uint32 TotalCount, uint32 BlockSize) -> uint32
-			{
-				return FMath::Min(BlockSize, TotalCount - BlockIndex * BlockSize);
-			};
-
-			int32 SourceKeyIndex0, SourceKeyIndex1;
+			check(SourceMapping.NumCompactBones == TargetMapping.NumCompactBones);
+			const uint32 SourceSequenceTransformBufferOffset    = SourceSequence.BufferAllocation.Offset * sizeof(FBoneTransformWithScale);
+			const uint32 TargetSequenceTransformBufferOffset    = TargetSequence.BufferAllocation.Offset * sizeof(FBoneTransformWithScale);
+			const uint32 SourceRetargetingTransformBufferOffset = SourceMapping.RetargetingTransformBufferOffset * sizeof(FBoneTransformWithScale);
+			const uint32 TargetRetargetingTransformBufferOffset = TargetMapping.RetargetingTransformBufferOffset * sizeof(FBoneTransformWithScale);
+			const uint32 SourceRetargetingDataBufferOffset      = SourceMapping.RetargetingDataBufferAllocation.Offset * sizeof(UE::HLSL::FBoneRetargetingData);
+			const uint32 TargetRetargetingDataBufferOffset      = TargetMapping.RetargetingDataBufferAllocation.Offset * sizeof(UE::HLSL::FBoneRetargetingData);
+			const uint32 InvGlobalRefPoseTransformBufferOffset  = SourceMapping.InvGlobalRefPoseTransformBufferOffset * sizeof(FBoneTransformWithScale);
+			const uint32 LeafBonePathBufferOffset               = LOD.LeafBonePathBufferAllocation.Offset * sizeof(uint32);
+			const uint32 BoneMapBufferOffset                    = LOD.BoneMapBufferOffset * sizeof(uint32);
+
+			int32 SourceKeyIndex0;
+			int32 SourceKeyIndex1;
 			float SourceAlpha;
 			SourceSequence.GetKeyIndicesFromTime(SourceKeyIndex0, SourceKeyIndex1, SourceAlpha, SourceTime);
 
 			NumTotalFramesRequested += SourceSequence.RequestFrame(SourceKeyIndex0);
 			NumTotalFramesRequested += SourceSequence.RequestFrame(SourceKeyIndex1);
 
-			// Evaluation Pass
+			int32 TargetKeyIndex0 = SourceKeyIndex0;
+			int32 TargetKeyIndex1 = SourceKeyIndex1;
+			float TargetAlpha = SourceAlpha;
+
+			if (BlendWeight > 0.0f)
 			{
-				check(SourceMapping.NumCompactBones == TargetMapping.NumCompactBones);
-				const uint32 NumBlocks = FMath::DivideAndRoundUp(TargetMapping.NumCompactBones, BonesPerGroup);
+				TargetSequence.GetKeyIndicesFromTime(TargetKeyIndex0, TargetKeyIndex1, TargetAlpha, TargetTime);
 
-				const uint32 SourceSequenceTransformBufferOffset    = SourceSequence.BufferAllocation.Offset * sizeof(FBoneTransformWithScale);
-				const uint32 SourceRetargetingTransformBufferOffset = SourceMapping.RetargetingTransformBufferOffset * sizeof(FBoneTransformWithScale);
-				const uint32 SourceRetargetingDataBufferOffset      = SourceMapping.RetargetingDataBufferAllocation.Offset * sizeof(UE::HLSL::FBoneRetargetingData);
+				NumTotalFramesRequested += TargetSequence.RequestFrame(TargetKeyIndex0);
+				NumTotalFramesRequested += TargetSequence.RequestFrame(TargetKeyIndex1);
+			}
 
+			if (LOD.bSinglePass)
+			{
 				if (BlendWeight > 0.0f)
 				{
-					int32 TargetKeyIndex0, TargetKeyIndex1;
-					float TargetAlpha;
-					TargetSequence.GetKeyIndicesFromTime(TargetKeyIndex0, TargetKeyIndex1, TargetAlpha, TargetTime);
-
-					NumTotalFramesRequested += TargetSequence.RequestFrame(TargetKeyIndex0);
-					NumTotalFramesRequested += TargetSequence.RequestFrame(TargetKeyIndex1);
-
-					const uint32 TargetSequenceTransformBufferOffset    = TargetSequence.BufferAllocation.Offset * sizeof(FBoneTransformWithScale);
-					const uint32 TargetRetargetingTransformBufferOffset = TargetMapping.RetargetingTransformBufferOffset * sizeof(FBoneTransformWithScale);
-					const uint32 TargetRetargetingDataBufferOffset      = TargetMapping.RetargetingDataBufferAllocation.Offset * sizeof(UE::HLSL::FBoneRetargetingData);
-
-					for (uint32 BlockIndex = 0; BlockIndex < NumBlocks; ++BlockIndex)
-					{
-						UE::HLSL::FEvaluateAnimSequenceBlendBlockHeader& BlockHeader = EvaluateBlendBlockHeaders.Add_GetRef();
-
-						BlockHeader.SourceAlpha                            = SourceAlpha;
-						BlockHeader.TargetAlpha                            = TargetAlpha;
-						BlockHeader.BlendWeight                            = BlendWeight;
-						BlockHeader.SourceKeyIndex0                        = SourceKeyIndex0;
-						BlockHeader.SourceKeyIndex1                        = SourceKeyIndex1;
-						BlockHeader.TargetKeyIndex0                        = TargetKeyIndex0;
-						BlockHeader.TargetKeyIndex1                        = TargetKeyIndex1;
-						BlockHeader.SourceSequenceTransformBufferOffset    = SourceSequenceTransformBufferOffset;
-						BlockHeader.TargetSequenceTransformBufferOffset    = TargetSequenceTransformBufferOffset;
-						BlockHeader.TargetRetargetingTransformBufferOffset = TargetRetargetingTransformBufferOffset;
-						BlockHeader.SourceRetargetingTransformBufferOffset = SourceRetargetingTransformBufferOffset;
-						BlockHeader.TargetRetargetingDataBufferOffset      = TargetRetargetingDataBufferOffset;
-						BlockHeader.SourceRetargetingDataBufferOffset      = SourceRetargetingDataBufferOffset;
-						BlockHeader.SourceNumAnimBones                     = SourceSequence.NumAnimBones;
-						BlockHeader.TargetNumAnimBones                     = TargetSequence.NumAnimBones;
-						BlockHeader.DstTransformBufferOffset               = IntermediateTransformOffset;
-						BlockHeader.BlockLocalIndex                        = BlockIndex;
-						BlockHeader.BlockTransformCount                    = GetBlockTransformCount(BlockIndex, TargetMapping.NumCompactBones, BonesPerGroup);
-					}
-
-					NumEvaluateBlendBlocks += NumBlocks;
+					UE::HLSL::FEvaluateBlendConcatenateScatterBlockHeader& BlockHeader = EvaluateBlendConcatenateScatterBlockHeaders.Add_GetRef();
+
+					BlockHeader.SourceAlpha                            = SourceAlpha;
+					BlockHeader.TargetAlpha                            = TargetAlpha;
+					BlockHeader.BlendWeight                            = BlendWeight;
+					BlockHeader.SourceKeyIndex0                        = SourceKeyIndex0;
+					BlockHeader.SourceKeyIndex1                        = SourceKeyIndex1;
+					BlockHeader.TargetKeyIndex0                        = TargetKeyIndex0;
+					BlockHeader.TargetKeyIndex1                        = TargetKeyIndex1;
+					BlockHeader.SourceSequenceTransformBufferOffset    = SourceSequenceTransformBufferOffset;
+					BlockHeader.TargetSequenceTransformBufferOffset    = TargetSequenceTransformBufferOffset;
+					BlockHeader.TargetRetargetingTransformBufferOffset = TargetRetargetingTransformBufferOffset;
+					BlockHeader.SourceRetargetingTransformBufferOffset = SourceRetargetingTransformBufferOffset;
+					BlockHeader.TargetRetargetingDataBufferOffset      = TargetRetargetingDataBufferOffset;
+					BlockHeader.SourceRetargetingDataBufferOffset      = SourceRetargetingDataBufferOffset;
+					BlockHeader.InvGlobalRefPoseTransformBufferOffset  = InvGlobalRefPoseTransformBufferOffset;
+					BlockHeader.BoneMapBufferOffset                    = BoneMapBufferOffset;
+					BlockHeader.DstTransformBufferOffset1              = DstTransformBufferOffset1;
+					BlockHeader.DstTransformBufferOffset2              = DstTransformBufferOffset2;
+					BlockHeader.SourceNumAnimBones                     = SourceSequence.NumAnimBones;
+					BlockHeader.TargetNumAnimBones                     = TargetSequence.NumAnimBones;
+					BlockHeader.NumCompactBones                        = LOD.BoneTables.NumCompactBones;
+
+					NumEvaluateBlendConcatenateScatterBlocks++;
 				}
 				else
 				{
-					for (uint32 BlockIndex = 0; BlockIndex < NumBlocks; ++BlockIndex)
-					{
-						UE::HLSL::FEvaluateAnimSequenceBlockHeader& BlockHeader = EvaluateBlockHeaders.Add_GetRef();
-
-						BlockHeader.SequenceTransformBufferOffset    = SourceSequenceTransformBufferOffset;
-						BlockHeader.RetargetingTransformBufferOffset = SourceRetargetingTransformBufferOffset;
-						BlockHeader.RetargetingDataBufferOffset      = SourceRetargetingDataBufferOffset;
-						BlockHeader.DstTransformBufferOffset         = IntermediateTransformOffset;
-						BlockHeader.NumAnimBones                     = SourceSequence.NumAnimBones;
-						BlockHeader.KeyIndex0                        = SourceKeyIndex0;
-						BlockHeader.KeyIndex1                        = SourceKeyIndex1;
-						BlockHeader.Alpha                            = SourceAlpha;
-						BlockHeader.BlockLocalIndex                  = BlockIndex;
-						BlockHeader.BlockTransformCount              = GetBlockTransformCount(BlockIndex, TargetMapping.NumCompactBones, BonesPerGroup);
-					}
+					UE::HLSL::FEvaluateConcatenateScatterBlockHeader& BlockHeader = EvaluateConcatenateScatterBlockHeaders.Add_GetRef();
+
+					BlockHeader.Alpha                                 = SourceAlpha;
+					BlockHeader.KeyIndex0                             = SourceKeyIndex0;
+					BlockHeader.KeyIndex1                             = SourceKeyIndex1;
+					BlockHeader.SequenceTransformBufferOffset         = SourceSequenceTransformBufferOffset;
+					BlockHeader.RetargetingTransformBufferOffset      = SourceRetargetingTransformBufferOffset;
+					BlockHeader.RetargetingDataBufferOffset           = SourceRetargetingDataBufferOffset;
+					BlockHeader.InvGlobalRefPoseTransformBufferOffset = InvGlobalRefPoseTransformBufferOffset;
+					BlockHeader.BoneMapBufferOffset                   = BoneMapBufferOffset;
+					BlockHeader.DstTransformBufferOffset1             = DstTransformBufferOffset1;
+					BlockHeader.DstTransformBufferOffset2             = DstTransformBufferOffset2;
+					BlockHeader.NumAnimBones                          = SourceSequence.NumAnimBones;
+					BlockHeader.NumCompactBones                       = LOD.BoneTables.NumCompactBones;
 
-					NumEvaluateBlocks += NumBlocks;
+					NumEvaluateConcatenateScatterBlocks++;
 				}
 			}
-
-			// Concatenate Local to Global Pass
+			else
 			{
-				const uint32 NumLeafBones = LOD.LeafBonePathData.NumLeafBones;
-				const uint32 NumBlocks    = FMath::DivideAndRoundUp(NumLeafBones, BonesPerGroup);
-
-				const uint32 LeafBonePathBufferOffset = LOD.LeafBonePathBufferAllocation.Offset * sizeof(uint32);
+				const auto GetBlockTransformCount = [](uint32 BlockIndex, uint32 TotalCount, uint32 BlockSize) -> uint32
+				{
+					return FMath::Min(BlockSize, TotalCount - BlockIndex * BlockSize);
+				};
 
-				for (uint32 BlockIndex = 0; BlockIndex < NumBlocks; ++BlockIndex)
+				// Evaluation Pass
 				{
-					UE::HLSL::FConcatenateBlockHeader& BlockHeader = ConcatenateBlockHeaders.Add_GetRef();
+					const uint32 NumBlocks = FMath::DivideAndRoundUp<uint32>(TargetMapping.NumCompactBones, BonesPerGroup);
+
+					if (BlendWeight > 0.0f)
+					{
+						for (uint32 BlockIndex = 0; BlockIndex < NumBlocks; ++BlockIndex)
+						{
+							UE::HLSL::FEvaluateAnimSequenceBlendBlockHeader& BlockHeader = EvaluateBlendBlockHeaders.Add_GetRef();
+
+							BlockHeader.SourceAlpha                            = SourceAlpha;
+							BlockHeader.TargetAlpha                            = TargetAlpha;
+							BlockHeader.BlendWeight                            = BlendWeight;
+							BlockHeader.SourceKeyIndex0                        = SourceKeyIndex0;
+							BlockHeader.SourceKeyIndex1                        = SourceKeyIndex1;
+							BlockHeader.TargetKeyIndex0                        = TargetKeyIndex0;
+							BlockHeader.TargetKeyIndex1                        = TargetKeyIndex1;
+							BlockHeader.SourceSequenceTransformBufferOffset    = SourceSequenceTransformBufferOffset;
+							BlockHeader.TargetSequenceTransformBufferOffset    = TargetSequenceTransformBufferOffset;
+							BlockHeader.TargetRetargetingTransformBufferOffset = TargetRetargetingTransformBufferOffset;
+							BlockHeader.SourceRetargetingTransformBufferOffset = SourceRetargetingTransformBufferOffset;
+							BlockHeader.TargetRetargetingDataBufferOffset      = TargetRetargetingDataBufferOffset;
+							BlockHeader.SourceRetargetingDataBufferOffset      = SourceRetargetingDataBufferOffset;
+							BlockHeader.SourceNumAnimBones                     = SourceSequence.NumAnimBones;
+							BlockHeader.TargetNumAnimBones                     = TargetSequence.NumAnimBones;
+							BlockHeader.DstTransformBufferOffset               = IntermediateTransformOffset;
+							BlockHeader.BlockLocalIndex                        = BlockIndex;
+							BlockHeader.BlockTransformCount                    = GetBlockTransformCount(BlockIndex, TargetMapping.NumCompactBones, BonesPerGroup);
+						}
+
+						NumEvaluateBlendBlocks += NumBlocks;
+					}
+					else
+					{
+						for (uint32 BlockIndex = 0; BlockIndex < NumBlocks; ++BlockIndex)
+						{
+							UE::HLSL::FEvaluateAnimSequenceBlockHeader& BlockHeader = EvaluateBlockHeaders.Add_GetRef();
+
+							BlockHeader.SequenceTransformBufferOffset    = SourceSequenceTransformBufferOffset;
+							BlockHeader.RetargetingTransformBufferOffset = SourceRetargetingTransformBufferOffset;
+							BlockHeader.RetargetingDataBufferOffset      = SourceRetargetingDataBufferOffset;
+							BlockHeader.DstTransformBufferOffset         = IntermediateTransformOffset;
+							BlockHeader.NumAnimBones                     = SourceSequence.NumAnimBones;
+							BlockHeader.KeyIndex0                        = SourceKeyIndex0;
+							BlockHeader.KeyIndex1                        = SourceKeyIndex1;
+							BlockHeader.Alpha                            = SourceAlpha;
+							BlockHeader.BlockLocalIndex                  = BlockIndex;
+							BlockHeader.BlockTransformCount              = GetBlockTransformCount(BlockIndex, TargetMapping.NumCompactBones, BonesPerGroup);
+						}
 
-					BlockHeader.LeafBonePathBufferOffset = LeafBonePathBufferOffset;
-					BlockHeader.TransformBufferOffset    = IntermediateTransformOffset;
-					BlockHeader.NumPrefetchBones         = FMath::Min(LOD.BoneTables.NumCompactBones, BonesPerGroup);
-					BlockHeader.BlockLocalIndex          = BlockIndex;
-					BlockHeader.BlockLeafBoneCount       = GetBlockTransformCount(BlockIndex, NumLeafBones, BonesPerGroup);
+						NumEvaluateBlocks += NumBlocks;
+					}
 				}
-			}
 
-			// Scatter Pass
-			{
-				const uint32 NumBlocks           = FMath::DivideAndRoundUp(MaxBoneTransformCount, BonesPerGroup);
-				const uint32 BoneMapBufferOffset = LOD.BoneMapBufferOffset * sizeof(uint32);
+				// Concatenate Local to Global Pass
+				{
+					const uint32 NumLeafBones = LOD.LeafBonePathData.NumLeafBones;
+					const uint32 NumBlocks    = FMath::DivideAndRoundUp(NumLeafBones, BonesPerGroup);
+
+					for (uint32 BlockIndex = 0; BlockIndex < NumBlocks; ++BlockIndex)
+					{
+						UE::HLSL::FConcatenateBlockHeader& BlockHeader = ConcatenateBlockHeaders.Add_GetRef();
 
-				const uint32 InvGlobalRefPoseTransformBufferOffset = SourceMapping.InvGlobalRefPoseTransformBufferOffset * sizeof(FBoneTransformWithScale);
+						BlockHeader.LeafBonePathBufferOffset = LeafBonePathBufferOffset;
+						BlockHeader.TransformBufferOffset    = IntermediateTransformOffset;
+						BlockHeader.NumPrefetchBones         = FMath::Min<uint16>(LOD.BoneTables.NumCompactBones, BonesPerGroup);
+						BlockHeader.BlockLocalIndex          = BlockIndex;
+						BlockHeader.BlockLeafBoneCount       = GetBlockTransformCount(BlockIndex, NumLeafBones, BonesPerGroup);
+					}
+				}
 
-				for (uint32 BlockIndex = 0; BlockIndex < NumBlocks; ++BlockIndex)
+				// Scatter Pass
 				{
-					UE::HLSL::FScatterTransformsBlockHeader& BlockHeader = ScatterBlockHeaders.Add_GetRef();
+					const uint32 NumBlocks = FMath::DivideAndRoundUp(MaxBoneTransformCount, BonesPerGroup);
 
-					BlockHeader.InvGlobalRefPoseTransformBufferOffset = InvGlobalRefPoseTransformBufferOffset;
-					BlockHeader.SrcTransformBufferOffset              = IntermediateTransformOffset;
-					BlockHeader.DstTransformBufferOffset1             = DstTransformBufferOffset1;
-					BlockHeader.DstTransformBufferOffset2             = DstTransformBufferOffset2;
-					BlockHeader.BoneMapOffset                         = BoneMapBufferOffset;
-					BlockHeader.BlockLocalIndex                       = BlockIndex;
-					BlockHeader.BlockTransformCount                   = GetBlockTransformCount(BlockIndex, MaxBoneTransformCount, BonesPerGroup);
+					for (uint32 BlockIndex = 0; BlockIndex < NumBlocks; ++BlockIndex)
+					{
+						UE::HLSL::FScatterTransformsBlockHeader& BlockHeader = ScatterBlockHeaders.Add_GetRef();
+
+						BlockHeader.InvGlobalRefPoseTransformBufferOffset = InvGlobalRefPoseTransformBufferOffset;
+						BlockHeader.SrcTransformBufferOffset              = IntermediateTransformOffset;
+						BlockHeader.DstTransformBufferOffset1             = DstTransformBufferOffset1;
+						BlockHeader.DstTransformBufferOffset2             = DstTransformBufferOffset2;
+						BlockHeader.BoneMapBufferOffset                   = BoneMapBufferOffset;
+						BlockHeader.BlockLocalIndex                       = BlockIndex;
+						BlockHeader.BlockTransformCount                   = GetBlockTransformCount(BlockIndex, MaxBoneTransformCount, BonesPerGroup);
+					}
 				}
-			}
 
-			IntermediateTransformOffset += SourceMapping.NumCompactBones * sizeof(FBoneTransformWithScale);
+				IntermediateTransformOffset += SourceMapping.NumCompactBones * sizeof(FBoneTransformWithScale);
+			}
 		};
 
 		uint32 UpdateResultIndex = 0;
@@ -1245,7 +1389,8 @@ void FAnimSequenceTransformProvider::ProvideGPUTransforms(FSkinningTransformProv
 				const float BlendWeight            = Data.GetBlendWeight();
 
 				const FTrackUpdateRequest UpdateRequest   = TrackUpdateRequests[UpdateResultIndex++];
-				const uint32 TrackTransformOffset         = TrackIndex * LOD.BoneTables.NumTransforms * 2 * sizeof(FCompressedBoneTransform);
+				const uint32 NumTransforms                = ExtensionProxy->GetMaxBoneTransformCount();
+				const uint32 TrackTransformOffset         = TrackIndex * NumTransforms * 2 * sizeof(FCompressedBoneTransform);
 				const uint32 TrackCurrentTransformOffset  = TrackTransformOffset + CurrentTransformOffset;
 				const uint32 TrackPreviousTransformOffset = TrackTransformOffset + PreviousTransformOffset;
 
@@ -1261,7 +1406,7 @@ void FAnimSequenceTransformProvider::ProvideGPUTransforms(FSkinningTransformProv
 						SourceTrack.Sequence, TargetTrack.Sequence,
 						SourceTrack.Position, TargetTrack.Position,
 						BlendWeight,
-						LOD.BoneTables.NumTransforms,
+						NumTransforms,
 						OutputTransformOffset1, OutputTransformOffset2);
 
 					NumCurrentTransformUpdates++;
@@ -1276,7 +1421,7 @@ void FAnimSequenceTransformProvider::ProvideGPUTransforms(FSkinningTransformProv
 						SourceTrack.Sequence, TargetTrack.Sequence,
 						SourceTrack.PreviousPosition, TargetTrack.PreviousPosition,
 						BlendWeight,
-						LOD.BoneTables.NumTransforms,
+						NumTransforms,
 						TrackPreviousTransformOffset, TrackPreviousTransformOffset);
 
 					NumPreviousTransformUpdates++;
@@ -1288,12 +1433,16 @@ void FAnimSequenceTransformProvider::ProvideGPUTransforms(FSkinningTransformProv
 		}
 
 		INC_DWORD_STAT_BY(STAT_AnimSequenceTransformProvider_ActiveTracks, NumActiveTracks);
+		INC_DWORD_STAT_BY(STAT_AnimSequenceTransformProvider_EvaluateSinglePassBlocks, NumEvaluateConcatenateScatterBlocks);
+		INC_DWORD_STAT_BY(STAT_AnimSequenceTransformProvider_EvaluateSinglePassBlendBlocks, NumEvaluateBlendConcatenateScatterBlocks);
 		INC_DWORD_STAT_BY(STAT_AnimSequenceTransformProvider_EvaluateBlocks, NumEvaluateBlocks);
 		INC_DWORD_STAT_BY(STAT_AnimSequenceTransformProvider_EvaluateBlendBlocks, NumEvaluateBlendBlocks);
 		INC_DWORD_STAT_BY(STAT_AnimSequenceTransformProvider_CurrentTransformUpdates, NumCurrentTransformUpdates);
 		INC_DWORD_STAT_BY(STAT_AnimSequenceTransformProvider_PreviousTransformUpdates, NumPreviousTransformUpdates);
 	});
 
+	Unlock(GraphBuilder, EvaluateConcatenateScatterBlockHeaderBuffer);
+	Unlock(GraphBuilder, EvaluateBlendConcatenateScatterBlockHeaderBuffer);
 	Unlock(GraphBuilder, EvaluateBlockHeaderBuffer);
 	Unlock(GraphBuilder, EvaluateBlendBlockHeaderBuffer);
 	Unlock(GraphBuilder, ScatterBlockHeaderBuffer);
@@ -1405,10 +1554,10 @@ void FAnimSequenceTransformProvider::ProvideGPUTransforms(FSkinningTransformProv
 
 					for (int32 TrackIndex = 0; TrackIndex < NumTracks; ++TrackIndex)
 					{
-						const int32 BoneIndex = PlatformCompressedData.GetSkeletonIndexFromTrackIndex(TrackIndex);
+						const uint16 BoneIndex = PlatformCompressedData.GetSkeletonIndexFromTrackIndex(TrackIndex);
 
 						// We only care about raw bones.
-						if (BoneIndex < (int32)Sequence.NumAnimBones)
+						if (BoneIndex < Sequence.NumAnimBones)
 						{
 							TrackToBoneIndexMap.Add(BoneTrackPair(BoneIndex, TrackIndex));
 						}
@@ -1503,46 +1652,89 @@ void FAnimSequenceTransformProvider::ProvideGPUTransforms(FSkinningTransformProv
 	if (NumTotalLeafBonePathElements > 0)
 	{
 		Context.ScatterUploadBuilder.AddPass(GraphBuilder, LeafBonePathUploadBuffer, LeafBonePathBuffer, NumTotalLeafBonePathElements, sizeof(uint32), TEXT("AnimSequenceTransformProvider.LeafBonePaths"),
-			[this] (FRDGScatterUploader& ScatterUploader)
+			[this, LODUploadRequests = MoveTemp(LODUploadRequests)](FRDGScatterUploader& ScatterUploader)
 		{
 			TRACE_CPUPROFILER_EVENT_SCOPE(FAnimSequenceTransformProvider::UploadLeafBonePaths);
+			TArray<uint16, SceneRenderingAllocator> CompactToTransformStorage;
 
-			for (const FProxyUserData* ProxyUserData : ProxyUserDatasToUpload)
+			for (const FLODUploadRequest& Request : LODUploadRequests)
 			{
-				const FSkinningSceneExtensionProxy& ExtensionProxy = ProxyUserData->SceneExtensionProxy;
+				FProxyUserData::FLOD& LOD = Request.LOD;
+				TConstArrayView<FBoneIndexType> BoneMap = Request.ExtensionProxy.GetBoneMap();
+				const FBoneTables& BoneTables = LOD.BoneTables;
+				const int32 NumTransformBones = Request.ExtensionProxy.GetMaxBoneTransformCount();
 
-				for (const FProxyUserData::FLOD& LOD : ProxyUserData->LODs)
+				if (LOD.bSinglePass)
 				{
-					if (!LOD.IsValid())
+					TArrayView<UE::HLSL::FBoneHierarchyData> DstBuffer = ScatterUploader.Add_GetRef<UE::HLSL::FBoneHierarchyData>(LOD.LeafBonePathBufferAllocation.Offset, LOD.LeafBonePathBufferAllocation.Count);
+					TConstArrayView<uint16> CompactToTransform;
+
+					if (!BoneMap.IsEmpty())
 					{
-						continue;
-					}
+						CompactToTransformStorage.Init(INVALID_BONE_INDEX, BoneTables.NumCompactBones);
+
+						for (uint32 TransformIndex = 0; TransformIndex < (uint32)BoneMap.Num(); ++TransformIndex)
+						{
+							const uint16 CompactBoneIndex = BoneTables.MeshToCompact[BoneMap[TransformIndex]];
+
+							if (CompactBoneIndex != INVALID_BONE_INDEX)
+							{
+								CompactToTransformStorage[CompactBoneIndex] = TransformIndex;
+							}
+						}
 
-					const FBoneTables& BoneTables = LOD.BoneTables;
-					const FLeafBonePathData& LeafBonePathData = LOD.LeafBonePathData;
+						CompactToTransform = CompactToTransformStorage;
+					}
+					else
+					{
+						CompactToTransform = BoneTables.CompactToMesh;
+					}
 
+					for (uint32 CompactBoneIndex = 0; CompactBoneIndex < BoneTables.NumCompactBones; ++CompactBoneIndex)
+					{
+						DstBuffer[CompactBoneIndex] = UE::HLSL::FBoneHierarchyData
+						{
+							  .ParentIndex    = (uint16)BoneTables.CompactToParent[CompactBoneIndex]
+							, .TransformIndex = (uint16)CompactToTransform[CompactBoneIndex]
+						};
+					}
+				}
+				else
+				{
+					FLeafBonePathData& LeafBonePathData = LOD.LeafBonePathData;
 					TArrayView<uint32> DstBuffer = ScatterUploader.Add_GetRef<uint32>(LOD.LeafBonePathBufferAllocation.Offset, LOD.LeafBonePathBufferAllocation.Count);
 
 					const int32 NumLeafBonePaths = LeafBonePathData.Buffer.Num();
-
 					FMemory::Memcpy(DstBuffer.GetData(), LeafBonePathData.Buffer.GetData(), NumLeafBonePaths * sizeof(uint32));
 
-					if (!BoneTables.TransformToMesh.IsEmpty())
+					// Release CPU side data now that it's been built.
+					LeafBonePathData.Buffer.Empty();
+
+					if (!BoneMap.IsEmpty())
 					{
-						for (int32 TransformIndex = 0; TransformIndex < BoneTables.TransformToMesh.Num(); ++TransformIndex)
+						// Transform bone space should match bone map space.
+						check(NumTransformBones == BoneMap.Num());
+
+						for (uint32 TransformIndex = 0; TransformIndex < (uint32)BoneMap.Num(); ++TransformIndex)
 						{
-							const int32 MeshBoneIndex = BoneTables.TransformToMesh[TransformIndex];
-							const int32 CompactBoneIndex = BoneTables.MeshToCompact[MeshBoneIndex];
+							const uint16 MeshBoneIndex = BoneMap[TransformIndex];
+							const uint16 CompactBoneIndex = BoneTables.MeshToCompact[MeshBoneIndex];
+
 							DstBuffer[NumLeafBonePaths + TransformIndex] = CompactBoneIndex;
 						}
 					}
 					else
 					{
-						FMemory::Memcpy(DstBuffer.GetData() + NumLeafBonePaths, BoneTables.MeshToCompact.GetData(), BoneTables.MeshToCompact.Num() * sizeof(uint32));
+						// Transform bone space should match mesh bone space.
+						check(NumTransformBones == BoneTables.NumMeshBones);
+
+						for (uint32 MeshBoneIndex = 0; MeshBoneIndex < (uint32)BoneTables.MeshToCompact.Num(); ++MeshBoneIndex)
+						{
+							DstBuffer[NumLeafBonePaths + MeshBoneIndex] = BoneTables.MeshToCompact[MeshBoneIndex];
+						}
 					}
 				}
 			}
-			ProxyUserDatasToUpload.Empty();
 		});
 
 		bExecuteScatterUpload = true;
@@ -1553,13 +1745,59 @@ void FAnimSequenceTransformProvider::ProvideGPUTransforms(FSkinningTransformProv
 		Context.ScatterUploadBuilder.Execute(Context.GraphBuilder);
 	}
 
-	FRDGBuffer* LocalSpaceTransformBuffer = GraphBuilder.CreateBuffer(FRDGBufferDesc::CreateByteAddressDesc(NumTotalIntermediateTransforms * sizeof(FBoneTransformWithScale)), TEXT("AnimSequenceTransformProvider.LocalSpaceBoneTransforms"));
-	FRDGBuffer* GlobalSpaceTransformBuffer = GraphBuilder.CreateBuffer(FRDGBufferDesc::CreateByteAddressDesc(NumTotalIntermediateTransforms * sizeof(FBoneTransformWithScale)), TEXT("AnimSequenceTransformProvider.GlobalSpaceBoneTransforms"));
+	FRDGBuffer* LocalSpaceTransformBuffer = nullptr;
+	FRDGBuffer* GlobalSpaceTransformBuffer = nullptr;
+
+	if (NumTotalIntermediateTransforms > 0)
+	{
+		LocalSpaceTransformBuffer = GraphBuilder.CreateBuffer(FRDGBufferDesc::CreateByteAddressDesc(NumTotalIntermediateTransforms * sizeof(FBoneTransformWithScale)), TEXT("AnimSequenceTransformProvider.LocalSpaceBoneTransforms"));
+		GlobalSpaceTransformBuffer = GraphBuilder.CreateBuffer(FRDGBufferDesc::CreateByteAddressDesc(NumTotalIntermediateTransforms * sizeof(FBoneTransformWithScale)), TEXT("AnimSequenceTransformProvider.GlobalSpaceBoneTransforms"));
+	}
 
 	FGlobalShaderMap* ShaderMap = GetGlobalShaderMap(GMaxRHIFeatureLevel);
 
 	RDG_EVENT_SCOPE_STAT(GraphBuilder, AnimSequenceTransformProvider, "AnimSequenceTransformProvider");
 
+	if (NumTotalEvaluateConcatenateScatterBlockHeaders > 0)
+	{
+		FEvaluateConcatenateScatterCS::FParameters* PassParameters = GraphBuilder.AllocParameters<FEvaluateConcatenateScatterCS::FParameters>();
+		PassParameters->RetargetingDataBuffer = GraphBuilder.CreateSRV(RetargetingDataBuffer);
+		PassParameters->SequenceTransformBuffer = GraphBuilder.CreateSRV(SequenceTransformBuffer);
+		PassParameters->EvaluateBlockHeaderBuffer = GraphBuilder.CreateSRV(EvaluateConcatenateScatterBlockHeaderBuffer);
+		PassParameters->BoneMapBuffer = GraphBuilder.CreateSRV(LeafBonePathBuffer);
+		PassParameters->SkinningTransformBufferUAV = GetCompressedBoneTransformUAV(GraphBuilder, Context.TransformBuffer);
+
+		TShaderMapRef<FEvaluateConcatenateScatterCS> ComputeShader(ShaderMap);
+
+		FComputeShaderUtils::AddPass(
+			GraphBuilder,
+			RDG_EVENT_NAME("EvaluateConcatenateScatter"),
+			ComputeShader,
+			PassParameters,
+			FIntVector(NumTotalEvaluateConcatenateScatterBlockHeaders, 1, 1)
+		);
+	}
+
+	if (NumTotalEvaluateBlendConcatenateScatterBlockHeaders > 0)
+	{
+		FEvaluateBlendConcatenateScatterCS::FParameters* PassParameters = GraphBuilder.AllocParameters<FEvaluateBlendConcatenateScatterCS::FParameters>();
+		PassParameters->RetargetingDataBuffer = GraphBuilder.CreateSRV(RetargetingDataBuffer);
+		PassParameters->SequenceTransformBuffer = GraphBuilder.CreateSRV(SequenceTransformBuffer);
+		PassParameters->EvaluateBlockHeaderBuffer = GraphBuilder.CreateSRV(EvaluateBlendConcatenateScatterBlockHeaderBuffer);
+		PassParameters->BoneMapBuffer = GraphBuilder.CreateSRV(LeafBonePathBuffer);
+		PassParameters->SkinningTransformBufferUAV = GetCompressedBoneTransformUAV(GraphBuilder, Context.TransformBuffer);
+
+		TShaderMapRef<FEvaluateBlendConcatenateScatterCS> ComputeShader(ShaderMap);
+
+		FComputeShaderUtils::AddPass(
+			GraphBuilder,
+			RDG_EVENT_NAME("EvaluateBlendConcatenateScatter"),
+			ComputeShader,
+			PassParameters,
+			FIntVector(NumTotalEvaluateBlendConcatenateScatterBlockHeaders, 1, 1)
+		);
+	}
+
 	// Dispatch 1.A: Evaluate animation sequences. Outputs animated local space bone transforms with retargeting.
 	if (NumTotalEvaluateBlockHeaders > 0)
 	{
@@ -1601,6 +1839,7 @@ void FAnimSequenceTransformProvider::ProvideGPUTransforms(FSkinningTransformProv
 	}
 
 	// Dispatch 2: Concatenate leaf bone transforms (local to global). Outputs global space bone transforms.
+	if (NumTotalConcatenateBlockHeaders > 0)
 	{
 		FConcatenateLocalToGlobalCS::FParameters* PassParameters = GraphBuilder.AllocParameters<FConcatenateLocalToGlobalCS::FParameters>();
 		PassParameters->LocalSpaceTransformBuffer = GraphBuilder.CreateSRV(LocalSpaceTransformBuffer);
@@ -1620,6 +1859,7 @@ void FAnimSequenceTransformProvider::ProvideGPUTransforms(FSkinningTransformProv
 	}
 
 	// Dispatch 3: Scatter transforms to final skinning scene buffer.
+	if (NumTotalScatterBlockHeaders > 0)
 	{
 		FScatterTransformsCS::FParameters* PassParameters = GraphBuilder.AllocParameters<FScatterTransformsCS::FParameters>();
 		PassParameters->GlobalSpaceTransformBuffer = GraphBuilder.CreateSRV(GlobalSpaceTransformBuffer);
@@ -1705,14 +1945,14 @@ void FAnimSequenceTransformProvider::ProvideCPUTransforms(FSkinningTransformProv
 
 			for (uint32 CompressedTrackIndex = 0; CompressedTrackIndex < NumCompressedTracks; ++CompressedTrackIndex)
 			{
-				const int32 AnimBoneIndex = PlatformCompressedData.GetSkeletonIndexFromTrackIndex(CompressedTrackIndex);
+				const uint16 AnimBoneIndex = PlatformCompressedData.GetSkeletonIndexFromTrackIndex(CompressedTrackIndex);
 
 				// We only care about raw bones.
-				if (AnimBoneIndex < (int32)Sequence.NumAnimBones)
+				if (AnimBoneIndex < Sequence.NumAnimBones)
 				{
-					const int32 CompactBoneIndex = Mapping.RetargetingTables.AnimToCompact[AnimBoneIndex];
+					const uint16 CompactBoneIndex = Mapping.RetargetingTables.AnimToCompact[AnimBoneIndex];
 
-					if (CompactBoneIndex != INDEX_NONE)
+					if (CompactBoneIndex != INVALID_BONE_INDEX)
 					{
 						TrackToBoneIndexMap.Add(BoneTrackPair(CompactBoneIndex, CompressedTrackIndex));
 						SampledLocalAnimPose[CompactBoneIndex] = AnimRefLocalPoses[AnimBoneIndex];
@@ -1734,12 +1974,12 @@ void FAnimSequenceTransformProvider::ProvideCPUTransforms(FSkinningTransformProv
 
 			for (uint32 CompactBoneIndex = 0; CompactBoneIndex < Mapping.NumCompactBones; ++CompactBoneIndex)
 			{
-				const int32 AnimBoneIndex = Mapping.RetargetingTables.CompactToAnim[CompactBoneIndex];
-				const int32 MeshBoneIndex = BoneTables.CompactToMesh[CompactBoneIndex];
+				const uint32 AnimBoneIndex = Mapping.RetargetingTables.CompactToAnim[CompactBoneIndex];
+				const uint32 MeshBoneIndex = BoneTables.CompactToMesh[CompactBoneIndex];
 
 				FTransform3f MeshLocalTransform;
 
-				if (AnimBoneIndex != INDEX_NONE)
+				if (AnimBoneIndex != INVALID_BONE_INDEX)
 				{
 					const FTransform3f RetargetTransform = Mapping.RetargetingTables.Transform[CompactBoneIndex];
 
@@ -1857,7 +2097,7 @@ void FAnimSequenceTransformProvider::ProvideCPUTransforms(FSkinningTransformProv
 				SourceSampledGlobalMeshPose[CompactBoneIndex] = TargetMapping.InvGlobalRefPoseTransforms[CompactBoneIndex] * SourceSampledGlobalMeshPose[CompactBoneIndex];
 			}
 
-			if (!BoneTables.TransformToMesh.IsEmpty())
+			if (!ExtensionProxy->GetBoneMap().IsEmpty())
 			{
 				SourceTransformMap.Init(false, TargetMapping.NumCompactBones);
 				SourceTransforms.SetNumUninitialized(TargetMapping.NumCompactBones, EAllowShrinking::No);
@@ -1865,12 +2105,12 @@ void FAnimSequenceTransformProvider::ProvideCPUTransforms(FSkinningTransformProv
 				SetCompressedBoneTransformIdentity(IdentityTransform);
 				int32 DstTransformIndex = 0;
 
-				for (uint32 MeshBoneIndex : BoneTables.TransformToMesh)
+				for (uint32 MeshBoneIndex : ExtensionProxy->GetBoneMap())
 				{
-					const int32 CompactBoneIndex = BoneTables.MeshToCompact[MeshBoneIndex];
+					const uint32 CompactBoneIndex = BoneTables.MeshToCompact[MeshBoneIndex];
 					FCompressedBoneTransform* OutputTransform;
 
-					if (CompactBoneIndex != INDEX_NONE)
+					if (CompactBoneIndex != INVALID_BONE_INDEX)
 					{
 						OutputTransform = &SourceTransforms[CompactBoneIndex];
 
@@ -1899,10 +2139,10 @@ void FAnimSequenceTransformProvider::ProvideCPUTransforms(FSkinningTransformProv
 			{
 				for (uint32 MeshBoneIndex = 0; MeshBoneIndex < TargetMapping.NumMeshBones; ++MeshBoneIndex)
 				{
-					const int32 CompactBoneIndex = BoneTables.MeshToCompact[MeshBoneIndex];
+					const uint32 CompactBoneIndex = BoneTables.MeshToCompact[MeshBoneIndex];
 					FCompressedBoneTransform OutputTransform;
 
-					if (CompactBoneIndex != INDEX_NONE)
+					if (CompactBoneIndex != INVALID_BONE_INDEX)
 					{
 						StoreCompressedBoneTransform(OutputTransform, SourceSampledGlobalMeshPose[CompactBoneIndex].ToMatrixWithScale());
 					}
@@ -1940,7 +2180,8 @@ void FAnimSequenceTransformProvider::ProvideCPUTransforms(FSkinningTransformProv
 			FProxyUserData::FLOD& LOD = ProxyUserData.LODs[LODIndex];
 
 			const uint32 UniqueAnimationCount   = ExtensionProxy->GetUniqueAnimationCount();
-			const uint32 MaxTotalTransformCount = LOD.BoneTables.NumTransforms * 2u; // Current and Previous
+			const uint32 NumTransforms          = ExtensionProxy->GetMaxBoneTransformCount();
+			const uint32 NumTotalTransforms     = NumTransforms * 2;
 			const uint32 CurrentTransformIndex  = Indirection.CurrentTransformOffset  / sizeof(FCompressedBoneTransform);
 			const uint32 PreviousTransformIndex = Indirection.PreviousTransformOffset / sizeof(FCompressedBoneTransform);
 			const uint32 CurrentTransformSlot   = CurrentTransformIndex < PreviousTransformIndex ? 0 : 1;
@@ -1959,7 +2200,7 @@ void FAnimSequenceTransformProvider::ProvideCPUTransforms(FSkinningTransformProv
 
 				// Dirty previous transforms means the whole allocation is dirty. Rebuild everything.
 				const bool bDirtyTransformAllocation = EnumHasAnyFlags(Indirection.DirtyBoneTransforms, EDirtyBoneTransforms::Previous);
-				const FTrackHistory::FUpdateResult UpdateResult = ProxyUserData.TrackHistories[TrackIndex].Update(SourceTrack, TargetTrack, CurrentTransformSlot, bDirtyTransformAllocation);
+				const FTrackHistory::FUpdateResult UpdateResult = ProxyUserData.TrackHistories[TrackIndex].Update(TargetTrack, CurrentTransformSlot, bDirtyTransformAllocation);
 	
 				const float BlendWeight = Data.GetBlendWeight();
 
@@ -1968,16 +2209,16 @@ void FAnimSequenceTransformProvider::ProvideCPUTransforms(FSkinningTransformProv
 
 				if (UpdateResult.bNeedsPrevious)
 				{
-					UpdateMode = UpdateResult.bNeedsCurrent && UpdateResult.bNeedsPrevious && TargetTrack.Position == TargetTrack.PreviousPosition && SourceTrack.Position == SourceTrack.PreviousPosition
+					UpdateMode = UpdateResult.bNeedsCurrent && UpdateResult.bNeedsPrevious && TargetTrack.Position == TargetTrack.PreviousPosition
 						? EPreviousBoneTransformUpdateMode::DuplicateCurrentToPrevious
 						: EPreviousBoneTransformUpdateMode::UpdatePrevious;
 
-					DstTransformsPrevious = ScatterUploader.Add_GetRef<FCompressedBoneTransform>(PreviousTransformIndex + TrackIndex * MaxTotalTransformCount, LOD.BoneTables.NumTransforms);
+					DstTransformsPrevious = ScatterUploader.Add_GetRef<FCompressedBoneTransform>(PreviousTransformIndex + TrackIndex * NumTotalTransforms, NumTransforms);
 				}
 
 				if (UpdateResult.bNeedsCurrent)
 				{
-					TArrayView<FCompressedBoneTransform> DstTransforms = ScatterUploader.Add_GetRef<FCompressedBoneTransform>(CurrentTransformIndex + TrackIndex * MaxTotalTransformCount, LOD.BoneTables.NumTransforms);
+					TArrayView<FCompressedBoneTransform> DstTransforms = ScatterUploader.Add_GetRef<FCompressedBoneTransform>(CurrentTransformIndex + TrackIndex * NumTotalTransforms, NumTransforms);
 					EvaluateAnimSequence(
 						LOD.BoneTables,
 						ExtensionProxy,
diff --git a/Engine/Source/Runtime/Renderer/Private/Skinning/AnimSequenceTransformProvider.h b/Engine/Source/Runtime/Renderer/Private/Skinning/AnimSequenceTransformProvider.h
--- a/Engine/Source/Runtime/Renderer/Private/Skinning/AnimSequenceTransformProvider.h
+++ b/Engine/Source/Runtime/Renderer/Private/Skinning/AnimSequenceTransformProvider.h
@@ -35,19 +35,16 @@ private:
 
 	struct FBoneTables
 	{
-		TArray<int32> TransformToMesh;
-		TArray<int32> MeshToCompact;
-		TArray<int32> CompactToMesh;
-		TArray<int32> CompactToParent;
-		TArray<int32> CompactToTargetSkeleton;
-		uint32 NumTransforms = 0;
+		TArray<uint16> MeshToCompact;
+		TArray<uint16> CompactToMesh;
+		TArray<uint16> CompactToParent;
+		TArray<uint16> CompactToTargetSkeleton;
 		uint32 NumMeshBones = 0;
 		uint32 NumCompactBones = 0;
 	};
 
 	void BuildBoneTables(
 		FBoneTables& BoneTables,
-		EBoneTransformStorageMode TransformStorageMode,
 		const USkinnedAsset* SkinnedMesh,
 		const FSkeletalMeshLODRenderData& LODRenderData);
 
@@ -56,8 +53,8 @@ private:
 		TArray<EBoneTranslationRetargetingMode::Type> Mode;
 		TArray<FTransform3f> Transform;
 		TArray<TTuple<FQuat4f, FQuat4f>> RefPoseRetargetQuats;
-		TArray<int32> CompactToAnim;
-		TArray<int32> AnimToCompact;
+		TArray<uint16> CompactToAnim;
+		TArray<uint16> AnimToCompact;
 		bool bRequiresRefPoseRetarget = false;
 	};
 
@@ -211,6 +208,7 @@ private:
 		float Position;
 		float PreviousPosition;
 		bool bHasPreviousPosition;
+		bool bAutoPlay;
 	};
 
 	class FTrackHistory
@@ -224,39 +222,39 @@ private:
 			uint8 bNeedsPrevious : 1 = 0;
 		};
 
-		FUpdateResult Update(const FTrackData& SourceTrack, const FTrackData& TargetTrack, int32 SlotIndex, bool bForceUpdate)
+		FUpdateResult Update(const FTrackData& Track, int32 SlotIndex, bool bForceUpdate)
 		{
+			check(SlotIndex == 0 || SlotIndex == 1);
+
 			FUpdateResult Result;
-			Result.bNeedsCurrent = bForceUpdate;
+			Result.bNeedsCurrent  = bForceUpdate || Track.bAutoPlay;
 			Result.bNeedsPrevious = bForceUpdate;
-			UpdateSlot(SourceSequenceIndex, SourceSlots, SourceTrack, SlotIndex, Result);
-			UpdateSlot(TargetSequenceIndex, TargetSlots, TargetTrack, SlotIndex, Result);
-			return Result;
-		}
 
-	private:
-		void UpdateSlot(int32& SequenceIndex, FVector2f& Slots, const FTrackData& Track, int32 InCurrentSlot, FUpdateResult& Result)
-		{
-			if (SequenceIndex != Track.SequenceIndex)
+			// Auto-play always rebuilds current transforms. Caching is disabled because blending
+			// between sequences with independent play rates makes reliable change detection impractical.
+			if (!Track.bAutoPlay)
 			{
-				SequenceIndex = Track.SequenceIndex;
-				Slots = FVector2f{ Invalid, Invalid };
-			}
+				if (SequenceIndex != Track.SequenceIndex)
+				{
+					SequenceIndex = Track.SequenceIndex;
+					Slots = FVector2f{ Invalid, Invalid };
+				}
 
-			Result.bNeedsCurrent  |= Slots[ InCurrentSlot] != Track.Position;
-			Result.bNeedsPrevious |= Slots[!InCurrentSlot] != Track.PreviousPosition && Track.bHasPreviousPosition;
-			Slots[InCurrentSlot] = Track.Position;
+				Result.bNeedsCurrent  |= Slots[ SlotIndex] != Track.Position;
+				Result.bNeedsPrevious |= Slots[!SlotIndex] != Track.PreviousPosition && Track.bHasPreviousPosition;
+				Slots[SlotIndex] = Track.Position;
 
-			if (Track.bHasPreviousPosition)
-			{
-				Slots[!InCurrentSlot] = Track.PreviousPosition;
+				if (Track.bHasPreviousPosition)
+				{
+					Slots[!SlotIndex] = Track.PreviousPosition;
+				}
 			}
+			return Result;
 		}
 
-		int32 SourceSequenceIndex = INDEX_NONE;
-		int32 TargetSequenceIndex = INDEX_NONE;
-		FVector2f SourceSlots{ Invalid, Invalid };
-		FVector2f TargetSlots{ Invalid, Invalid };
+	private:
+		int32 SequenceIndex = INDEX_NONE;
+		FVector2f Slots{ Invalid, Invalid };
 	};
 
 	class FProxyUserData : public FInstancedSkinningSceneExtensionProxyUserData
@@ -290,6 +288,7 @@ private:
 				, .Position = Position
 				, .PreviousPosition = PreviousPosition
 				, .bHasPreviousPosition = Data.HasPreviousPosition()
+				, .bAutoPlay = Data.IsAutoPlay()
 			};
 		}
 
@@ -299,7 +298,8 @@ private:
 			FLeafBonePathData LeafBonePathData;
 			FBufferAllocation LeafBonePathBufferAllocation;
 			int32 BoneMapBufferOffset = 0;
-			TArray<int32> SectionBoneMaps;
+			bool bSinglePass = false;
+			bool bUploaded = false;
 
 			bool IsValid() const
 			{
