[HWRT] Fix an assert when trying to look up light function miss shader from the fallback RTPSO.

#rnx

#Synced-CL 50366487

diff --git a/Engine/Source/Runtime/Renderer/Private/PathTracing.cpp b/Engine/Source/Runtime/Renderer/Private/PathTracing.cpp
--- a/Engine/Source/Runtime/Renderer/Private/PathTracing.cpp
+++ b/Engine/Source/Runtime/Renderer/Private/PathTracing.cpp
@@ -2271,7 +2271,12 @@ static void BindLightFunction(
 
 	Shader->GetShaderBindings(Scene, Scene->GetFeatureLevel(), MaterialRenderProxy, Material, View, LightFunctionParameters, SingleShaderBindings);
 
-	int32 MissShaderPipelineIndex = FindRayTracingMissShaderIndex(View.MaterialRayTracingData.PipelineState, Shader.GetRayTracingShader(), true);
+	int32 MissShaderPipelineIndex = FindRayTracingMissShaderIndex(View.MaterialRayTracingData.PipelineState, Shader.GetRayTracingShader(), false);
+
+	if (MissShaderPipelineIndex == INDEX_NONE)
+	{
+		MissShaderPipelineIndex = FindRayTracingMissShaderIndex(View.MaterialRayTracingData.PipelineState, GetPathTracingDefaultMissShader(View.ShaderMap), true);
+	}
 
 	ShaderBindings.SetRayTracingShaderBindingsForMissShader(RHICmdList, SBT, Index, Pipeline, MissShaderPipelineIndex);
 }
diff --git a/Engine/Source/Runtime/Renderer/Private/RayTracing/RayTracingLighting.cpp b/Engine/Source/Runtime/Renderer/Private/RayTracing/RayTracingLighting.cpp
--- a/Engine/Source/Runtime/Renderer/Private/RayTracing/RayTracingLighting.cpp
+++ b/Engine/Source/Runtime/Renderer/Private/RayTracing/RayTracingLighting.cpp
@@ -702,7 +702,12 @@ static void BindLightFunction(
 
 	Shader->GetShaderBindings(Scene, Scene->GetFeatureLevel(), MaterialRenderProxy, Material, View, DeferredLightBuffer, LightFunctionParameters, View.RayTracingLightGridUniformBuffer->GetRHIRef(), SingleShaderBindings);
 
-	int32 MissShaderPipelineIndex = FindRayTracingMissShaderIndex(View.MaterialRayTracingData.PipelineState, Shader.GetRayTracingShader(), true);
+	int32 MissShaderPipelineIndex = FindRayTracingMissShaderIndex(View.MaterialRayTracingData.PipelineState, Shader.GetRayTracingShader(), false);
+
+	if (MissShaderPipelineIndex == INDEX_NONE)
+	{
+		MissShaderPipelineIndex = FindRayTracingMissShaderIndex(View.MaterialRayTracingData.PipelineState, GetRayTracingDefaultMissShader(View.ShaderMap), true);
+	}
 
 	ShaderBindings.SetRayTracingShaderBindingsForMissShader(RHICmdList, SBT, Index, Pipeline, MissShaderPipelineIndex);
 }
