[HWRT] Show "Opaque Only" toggle in ray tracing debug visualization in-editor UI.
- also added on-screen warning saying whether translucent geometry is being force included in ray tracing scene due to debug visualization mode.

#Synced-CL 50568677

diff --git a/Engine/Source/Editor/CommonMenuExtensions/Private/RayTracingDebugVisualizationMenuCommands.cpp b/Engine/Source/Editor/CommonMenuExtensions/Private/RayTracingDebugVisualizationMenuCommands.cpp
--- a/Engine/Source/Editor/CommonMenuExtensions/Private/RayTracingDebugVisualizationMenuCommands.cpp
+++ b/Engine/Source/Editor/CommonMenuExtensions/Private/RayTracingDebugVisualizationMenuCommands.cpp
@@ -95,6 +95,8 @@ void FRayTracingDebugVisualizationMenuCommands::BuildCommandMap()
 
 	UI_COMMAND(ShowNearFieldCommand, "Near Field", "Visualizations will show the near field layer of the ray tracing scene (disables r.RayTracing.Visualize.FarField).", EUserInterfaceActionType::RadioButton, FInputChord());
 	UI_COMMAND(ShowFarFieldCommand, "Far Field", "Visualizations will show the far field layer of the ray tracing scene (enables r.RayTracing.Visualize.FarField).", EUserInterfaceActionType::RadioButton, FInputChord());
+
+	UI_COMMAND(ShowOpaqueOnlyCommand, "Opaque Only", "Whether to only show opaque instances in ray tracing debug visualizations (r.RayTracing.Visualize.OpaqueOnly).", EUserInterfaceActionType::ToggleButton, FInputChord());
 }
 
 void FRayTracingDebugVisualizationMenuCommands::BuildVisualisationSubMenu(UToolMenu* InMenu)
@@ -144,10 +146,8 @@ void FRayTracingDebugVisualizationMenuCommands::BuildVisualisationSubMenu(UToolM
 		Commands.AddCommandTypeToSection(OtherSection, FVisualizationType::Other);
 	}
 
-	{
-		FToolMenuSection& LayersSection = InMenu->AddSection("LevelViewportRayTracingVisualizationLayers", LOCTEXT("RayTracingVisualizationLayers", "Layers"));
-		Commands.AddLayerCommands(LayersSection);
-	}
+	Commands.AddLayersSection(*InMenu);
+	Commands.AddOptionsSection(*InMenu);
 }
 
 bool FRayTracingDebugVisualizationMenuCommands::AddCommandTypeToSection(FToolMenuSection& InSection, const FVisualizationType Type) const
@@ -167,10 +167,19 @@ bool FRayTracingDebugVisualizationMenuCommands::AddCommandTypeToSection(FToolMen
 	return bAddedCommands;
 }
 
-bool FRayTracingDebugVisualizationMenuCommands::AddLayerCommands(FToolMenuSection& InSection) const
+bool FRayTracingDebugVisualizationMenuCommands::AddLayersSection(UToolMenu& InMenu) const
 {
-	InSection.AddMenuEntry(ShowNearFieldCommand);
-	InSection.AddMenuEntry(ShowFarFieldCommand);
+	FToolMenuSection& Section = InMenu.AddSection("LevelViewportRayTracingVisualizationLayers", LOCTEXT("RayTracingVisualizationLayers", "Layers"));
+	Section.AddMenuEntry(ShowNearFieldCommand);
+	Section.AddMenuEntry(ShowFarFieldCommand);
+
+	return true;
+}
+
+bool FRayTracingDebugVisualizationMenuCommands::AddOptionsSection(UToolMenu& InMenu) const
+{
+	FToolMenuSection& Section = InMenu.AddSection("LevelViewportRayTracingVisualizationOptions", LOCTEXT("RayTracingVisualizationOptions", "Options"));
+	Section.AddMenuEntry(ShowOpaqueOnlyCommand);
 
 	return true;
 }
@@ -232,6 +241,24 @@ void FRayTracingDebugVisualizationMenuCommands::BindCommands(FUICommandList& Com
 		FCanExecuteAction::CreateLambda([]() { return true; }),
 		FIsActionChecked::CreateLambda([=]() { return FarFieldCVar->GetBool(); })
 	);
+
+	IConsoleVariable* OpaqueOnlyCVar = IConsoleManager::Get().FindConsoleVariable(TEXT("r.RayTracing.Visualize.OpaqueOnly"));
+	check(OpaqueOnlyCVar);
+
+	CommandList.MapAction(
+		ShowOpaqueOnlyCommand,
+		FExecuteAction::CreateLambda([=]()
+			{
+				OpaqueOnlyCVar->Set(!OpaqueOnlyCVar->GetBool());
+
+				if (TSharedPtr<FEditorViewportClient> Client = WeakClient.Pin())
+				{
+					Client->Invalidate();
+				}
+			}),
+		FCanExecuteAction::CreateLambda([]() { return true; }),
+		FIsActionChecked::CreateLambda([=]() { return OpaqueOnlyCVar->GetBool(); })
+	);
 }
 
 void FRayTracingDebugVisualizationMenuCommands::ChangeRayTracingDebugVisualizationMode(TWeakPtr<FEditorViewportClient> WeakClient, FName InName)
diff --git a/Engine/Source/Editor/CommonMenuExtensions/Public/RayTracingDebugVisualizationMenuCommands.h b/Engine/Source/Editor/CommonMenuExtensions/Public/RayTracingDebugVisualizationMenuCommands.h
--- a/Engine/Source/Editor/CommonMenuExtensions/Public/RayTracingDebugVisualizationMenuCommands.h
+++ b/Engine/Source/Editor/CommonMenuExtensions/Public/RayTracingDebugVisualizationMenuCommands.h
@@ -65,7 +65,8 @@ private:
 	void BuildCommandMap();
 
 	bool AddCommandTypeToSection(FToolMenuSection& InSection, const FVisualizationType Type) const;
-	bool AddLayerCommands(FToolMenuSection& InSection) const;
+	bool AddLayersSection(UToolMenu& InMenu) const;
+	bool AddOptionsSection(UToolMenu& InMenu) const;
 
 	static void ChangeRayTracingDebugVisualizationMode(TWeakPtr<FEditorViewportClient> WeakClient, FName InName);
 	static bool IsRayTracingDebugVisualizationModeSelected(TWeakPtr<FEditorViewportClient> WeakClient, FName InName);
@@ -74,6 +75,8 @@ private:
 
 	TSharedPtr<FUICommandInfo> ShowNearFieldCommand;
 	TSharedPtr<FUICommandInfo> ShowFarFieldCommand;
+
+	TSharedPtr<FUICommandInfo> ShowOpaqueOnlyCommand;
 };
 
 #undef UE_API
diff --git a/Engine/Source/Runtime/Renderer/Private/RayTracing/RayTracing.cpp b/Engine/Source/Runtime/Renderer/Private/RayTracing/RayTracing.cpp
--- a/Engine/Source/Runtime/Renderer/Private/RayTracing/RayTracing.cpp
+++ b/Engine/Source/Runtime/Renderer/Private/RayTracing/RayTracing.cpp
@@ -18,6 +18,7 @@
 #include "RayTracingShadows.h"
 #include "Experimental/Containers/SherwoodHashTable.h"
 #include "RHIShaderBindingLayout.h"
+#include "Engine/Engine.h"
 #include "Async/ParallelFor.h"
 #include <type_traits>
 
@@ -201,6 +202,8 @@ static void RefreshRayTracingInstancesSinkFunction()
 
 static FAutoConsoleVariableSink CVarRefreshRayTracingInstancesSink(FConsoleCommandDelegate::CreateStatic(&RefreshRayTracingInstancesSinkFunction));
 
+extern bool IsRayTracingVisualizeOpaqueOnly();
+
 namespace RayTracing
 {
 	// Configure ray tracing scene options based on currently enabled features and their needs
@@ -221,14 +224,20 @@ namespace RayTracing
 			bTranslucentGeometry = true;
 		}
 
-		if (ViewFamily.EngineShowFlags.RayTracingDebug)
+		if (ViewFamily.EngineShowFlags.PathTracing
+			&& FDataDrivenShaderPlatformInfo::GetSupportsPathTracing(Scene.GetShaderPlatform()))
 		{
-			bTranslucentGeometry = true; // could check r.RayTracing.Visualize.OpaqueOnly, but not critical as this is only for debugging purposes
+			bTranslucentGeometry = true;
 		}
 
-		if (ViewFamily.EngineShowFlags.PathTracing
-			&& FDataDrivenShaderPlatformInfo::GetSupportsPathTracing(Scene.GetShaderPlatform()))
+		if (ViewFamily.EngineShowFlags.RayTracingDebug && !IsRayTracingVisualizeOpaqueOnly())
 		{
+			if (!bTranslucentGeometry)
+			{
+#if !UE_BUILD_SHIPPING
+				GEngine->AddOnScreenDebugMessage(uint64(0x529F427D1), 1.f, FColor::Yellow, TEXT("Translucent geometry is only being included in ray tracing scene due to debug visualization mode."));
+#endif
+			}
 			bTranslucentGeometry = true;
 		}
 
diff --git a/Engine/Source/Runtime/Renderer/Private/RayTracing/RayTracingDebug.cpp b/Engine/Source/Runtime/Renderer/Private/RayTracing/RayTracingDebug.cpp
--- a/Engine/Source/Runtime/Renderer/Private/RayTracing/RayTracingDebug.cpp
+++ b/Engine/Source/Runtime/Renderer/Private/RayTracing/RayTracingDebug.cpp
@@ -68,10 +68,10 @@ TAutoConsoleVariable<int32> CVarRayTracingVisualizePickerDomain(
 
 static FAutoConsoleVariableDeprecated CVarRayTracingDebugPickerDomain_Deprecated(TEXT("r.RayTracing.Debug.PickerDomain"), TEXT("r.RayTracing.Visualize.PickerDomain"), TEXT("5.6"));
 
-static TAutoConsoleVariable<int32> CVarRayTracingVisualizeOpaqueOnly(
+static TAutoConsoleVariable<bool> CVarRayTracingVisualizeOpaqueOnly(
 	TEXT("r.RayTracing.Visualize.OpaqueOnly"),
-	1,
-	TEXT("Sets whether the view mode rendes opaque objects only (default = 1, render only opaque objects, 0 = render all objects)"),
+	true,
+	TEXT("Whether to only show opaque instances in ray tracing debug visualizations."),
 	ECVF_RenderThreadSafe
 );
 
@@ -175,6 +175,11 @@ static FAutoConsoleVariableRef CVarVisualizeProceduralPrimitives(
 
 static FAutoConsoleVariableDeprecated CVarVisualizeProceduralPrimitives_Deprecated(TEXT("r.RayTracing.DebugVisualizationMode.ProceduralPrimitives"), TEXT("r.RayTracing.Visualize.ProceduralPrimitives"), TEXT("5.6"));
 
+bool IsRayTracingVisualizeOpaqueOnly()
+{
+	return CVarRayTracingVisualizeOpaqueOnly.GetValueOnRenderThread();
+}
+
 float GetRayTracingDebugTimingScale()
 {
 	return CVarRayTracingVisualizeTimingScale.GetValueOnRenderThread() / 25000.0f;
@@ -993,7 +998,7 @@ static FRDGBufferRef RayTracingPerformPicking(
 	SharedParameters.InstancesExtraData = GraphBuilder.CreateSRV(RayTracingScene.GetInstanceExtraDataBuffer(LayerToTrace, View.GetRayTracingSceneViewHandle()));
 	SharedParameters.InstancesDebugData = GraphBuilder.CreateSRV(RayTracingScene.GetInstanceDebugBuffer(LayerToTrace));
 	SharedParameters.TLAS = RayTracingScene.GetLayerView(LayerToTrace, View.GetRayTracingSceneViewHandle());
-	SharedParameters.OpaqueOnly = CVarRayTracingVisualizeOpaqueOnly.GetValueOnRenderThread();
+	SharedParameters.OpaqueOnly = IsRayTracingVisualizeOpaqueOnly() ? 1 : 0;
 	SharedParameters.InstanceBuffer = GraphBuilder.CreateSRV(RayTracingScene.GetInstanceBuffer(LayerToTrace, View.GetRayTracingSceneViewHandle()));
 	SharedParameters.ViewUniformBuffer = View.ViewUniformBuffer;
 	SharedParameters.SceneUniformBuffer = GetSceneUniformBufferRef(GraphBuilder, View); // TODO: use a separate params structure
@@ -1313,7 +1318,7 @@ static FRDGBufferRef RayTracingPerformHitStatsPerPrimitive(
 
 	FRayTracingDebugHitStatsRGS::FParameters* RayGenParameters = GraphBuilder.AllocParameters<FRayTracingDebugHitStatsRGS::FParameters>();
 	RayGenParameters->TLAS = RayTracingScene.GetLayerView(LayerToTrace, View.GetRayTracingSceneViewHandle());
-	RayGenParameters->OpaqueOnly = CVarRayTracingVisualizeOpaqueOnly.GetValueOnRenderThread();
+	RayGenParameters->OpaqueOnly = IsRayTracingVisualizeOpaqueOnly() ? 1 : 0;
 	RayGenParameters->ViewUniformBuffer = View.ViewUniformBuffer;
 	RayGenParameters->SceneUniformBuffer = GetSceneUniformBufferRef(GraphBuilder, View); // TODO: use a separate params structure
 	RayGenParameters->NaniteRayTracingUniformBuffer = Nanite::GRayTracingManager.GetUniformBuffer();
@@ -1753,7 +1758,7 @@ void RenderRayTracingDebug(FRDGBuilder& GraphBuilder, const FScene& Scene, const
 	SharedParameters.VisualizationMode = DebugVisualizationMode;
 	SharedParameters.PickerDomain = CVarRayTracingVisualizePickerDomain.GetValueOnRenderThread();
 	SharedParameters.TimingScale = GetRayTracingDebugTimingScale();
-	SharedParameters.OpaqueOnly = CVarRayTracingVisualizeOpaqueOnly.GetValueOnRenderThread();
+	SharedParameters.OpaqueOnly = IsRayTracingVisualizeOpaqueOnly() ? 1 : 0;
 	SharedParameters.TriangleHitCountMaxThreshold = FMath::Clamp((float)CVarRayTracingVisualizeHitCountMaxThreshold.GetValueOnRenderThread(), 1, 100000);
 	SharedParameters.TriangleHitCountPerInstanceMaxThreshold = FMath::Max(1, CVarRayTracingVisualizeHitCountPerInstanceMaxThreshold.GetValueOnRenderThread());
 	SharedParameters.RayTracingDebugHitStatsUniformBuffer = DebugHitStatsUniformBuffer;
