Fix incorrect hair transmittance in splitscreen when VSM one pass is used.

#Synced-CL 50586717

#Synced-CL 50586717

diff --git a/Engine/Source/Runtime/Renderer/Private/LightRendering.cpp b/Engine/Source/Runtime/Renderer/Private/LightRendering.cpp
--- a/Engine/Source/Runtime/Renderer/Private/LightRendering.cpp
+++ b/Engine/Source/Runtime/Renderer/Private/LightRendering.cpp
@@ -2420,7 +2420,10 @@ void FDeferredShadingSceneRenderer::RenderLights(
 							FRDGTextureRef HairShadowMask = nullptr;
 							if (bDrawHairShadow && VirtualShadowMapId != INDEX_NONE)
 							{
-								TransmittanceMaskData.TransmittanceMask = ShadowSceneRenderer.HairTransmittanceMaskBits;
+								if (ShadowSceneRenderer.HairTransmittanceMaskBits.IsValidIndex(ViewIndex))
+								{
+									TransmittanceMaskData.TransmittanceMask = ShadowSceneRenderer.HairTransmittanceMaskBits[ViewIndex];
+								}
 								HairShadowMask = nullptr;
 								check(ScreenShadowMaskSubPixelTexture == nullptr);
 							}
diff --git a/Engine/Source/Runtime/Renderer/Private/Shadows/ShadowSceneRenderer.cpp b/Engine/Source/Runtime/Renderer/Private/Shadows/ShadowSceneRenderer.cpp
--- a/Engine/Source/Runtime/Renderer/Private/Shadows/ShadowSceneRenderer.cpp
+++ b/Engine/Source/Runtime/Renderer/Private/Shadows/ShadowSceneRenderer.cpp
@@ -920,6 +920,7 @@ void FShadowSceneRenderer::RenderVirtualShadowMapProjectionMaskBits(
 		IsVSMOnePassProjectionEnabled(SceneRenderer.ViewFamily.EngineShowFlags) &&
 		bNeedVSMOnePassProjection;
 
+	HairTransmittanceMaskBits.Empty();
 	if (bShouldUseVirtualShadowMapOnePassProjection)
 	{
 		RDG_EVENT_SCOPE(GraphBuilder, "VirtualShadowMapProjectionMaskBits");
@@ -953,7 +954,7 @@ void FShadowSceneRenderer::RenderVirtualShadowMapProjectionMaskBits(
 					VirtualShadowMapMaskBitsHairStrands);
 
 				// Transmittance bits
-				HairTransmittanceMaskBits = RenderHairStrandsOnePassTransmittanceMask(GraphBuilder, View, ViewIndex, VirtualShadowMapMaskBitsHairStrands, VirtualShadowMapArray).TransmittanceMask;
+				HairTransmittanceMaskBits.Add(RenderHairStrandsOnePassTransmittanceMask(GraphBuilder, View, ViewIndex, VirtualShadowMapMaskBitsHairStrands, VirtualShadowMapArray).TransmittanceMask);
 			}
 		}
 	}
@@ -961,7 +962,6 @@ void FShadowSceneRenderer::RenderVirtualShadowMapProjectionMaskBits(
 	{
 		VirtualShadowMapMaskBits = nullptr;//Dummy;
 		VirtualShadowMapMaskBitsHairStrands = nullptr;//Dummy;
-		HairTransmittanceMaskBits = nullptr; //Dummy
 	}
 }
 
diff --git a/Engine/Source/Runtime/Renderer/Private/Shadows/ShadowSceneRenderer.h b/Engine/Source/Runtime/Renderer/Private/Shadows/ShadowSceneRenderer.h
--- a/Engine/Source/Runtime/Renderer/Private/Shadows/ShadowSceneRenderer.h
+++ b/Engine/Source/Runtime/Renderer/Private/Shadows/ShadowSceneRenderer.h
@@ -142,7 +142,7 @@ public:
 	// One pass projection stuff. Set up in RenderVitualShadowMapProjectionMaskBits
 	FRDGTextureRef VirtualShadowMapMaskBits = nullptr;
 	FRDGTextureRef VirtualShadowMapMaskBitsHairStrands = nullptr;
-	FRDGBufferRef HairTransmittanceMaskBits = nullptr;
+	TArray<FRDGBufferRef> HairTransmittanceMaskBits;
 
 	bool UsePackedShadowMaskBits() const
 	{
