Deleted deprecated code from FEyeAdaptationManager.

#Synced-CL 50597675

diff --git a/Engine/Source/Runtime/Renderer/Private/PostProcess/PostProcessEyeAdaptation.cpp b/Engine/Source/Runtime/Renderer/Private/PostProcess/PostProcessEyeAdaptation.cpp
--- a/Engine/Source/Runtime/Renderer/Private/PostProcess/PostProcessEyeAdaptation.cpp
+++ b/Engine/Source/Runtime/Renderer/Private/PostProcess/PostProcessEyeAdaptation.cpp
@@ -1240,9 +1240,6 @@ void FSceneViewState::FEyeAdaptationManager::SafeRelease()
 	for (int32 Index = 0; Index < NUM_BUFFERS; Index++)
 	{
 		ExposureBufferData[Index].SafeRelease();
-PRAGMA_DISABLE_DEPRECATION_WARNINGS
-		PooledRenderTarget[Index].SafeRelease();
-PRAGMA_ENABLE_DEPRECATION_WARNINGS
 	}
 
 	for (int32 Index = 0; Index < ExposureReadbackBuffers.Num(); ++Index)
@@ -1255,33 +1252,6 @@ PRAGMA_ENABLE_DEPRECATION_WARNINGS
 	}
 }
 
-const TRefCountPtr<IPooledRenderTarget>& FSceneViewState::FEyeAdaptationManager::GetTexture(uint32 TextureIndex) const
-{
-	check(0 <= TextureIndex && TextureIndex < NUM_BUFFERS);
-	
-PRAGMA_DISABLE_DEPRECATION_WARNINGS
-	return PooledRenderTarget[TextureIndex];
-PRAGMA_ENABLE_DEPRECATION_WARNINGS
-}
-
-const TRefCountPtr<IPooledRenderTarget>& FSceneViewState::FEyeAdaptationManager::GetOrCreateTexture(FRDGBuilder& GraphBuilder, uint32 TextureIndex)
-{
-	check(0 <= TextureIndex && TextureIndex < NUM_BUFFERS);
-
-PRAGMA_DISABLE_DEPRECATION_WARNINGS
-	// Create texture if needed.
-	if (!PooledRenderTarget[TextureIndex].IsValid())
-	{
-		const FRDGTextureDesc RDGTextureDesc = FRDGTextureDesc::Create2D(FIntPoint(1, 1), PF_A32B32G32R32F, FClearValueBinding::None, TexCreate_UAV | TexCreate_ShaderResource);
-		FRDGTextureRef RDGTexture = GraphBuilder.CreateTexture(RDGTextureDesc, TEXT("EyeAdaptationTexture"), ERDGTextureFlags::MultiFrame);
-
-		PooledRenderTarget[TextureIndex] = GraphBuilder.ConvertToExternalTexture(RDGTexture);
-	}
-
-	return PooledRenderTarget[TextureIndex];
-PRAGMA_ENABLE_DEPRECATION_WARNINGS
-}
-
 const TRefCountPtr<FRDGPooledBuffer>& FSceneViewState::FEyeAdaptationManager::GetBuffer(uint32 BufferIndex) const
 {
 	check(0 <= BufferIndex && BufferIndex < NUM_BUFFERS);
diff --git a/Engine/Source/Runtime/Renderer/Private/ScenePrivate.h b/Engine/Source/Runtime/Renderer/Private/ScenePrivate.h
--- a/Engine/Source/Runtime/Renderer/Private/ScenePrivate.h
+++ b/Engine/Source/Runtime/Renderer/Private/ScenePrivate.h
@@ -898,9 +898,6 @@ private:
 		uint64 GetGPUSizeBytes(bool bLogSizes) const;
 
 	private:
-		const TRefCountPtr<IPooledRenderTarget>& GetTexture(uint32 TextureIndex) const;
-		const TRefCountPtr<IPooledRenderTarget>& GetOrCreateTexture(FRDGBuilder& GraphBuilder, uint32 TextureIndex);
-
 		const TRefCountPtr<FRDGPooledBuffer>& GetBuffer(uint32 BufferIndex) const;
 		const TRefCountPtr<FRDGPooledBuffer>& GetOrCreateBuffer(FRDGBuilder& GraphBuilder, uint32 BufferIndex);
 
@@ -921,9 +918,6 @@ private:
 		float LastAverageLocalExposure = 1.0f;
 		float LastAverageSceneLuminance = 0; // 0 means invalid. Used for Exposure Compensation Curve.
 
-		UE_DEPRECATED(5.2, "Use ExposureBufferData instead.")
-		TRefCountPtr<IPooledRenderTarget> PooledRenderTarget[NUM_BUFFERS];
-
 		TRefCountPtr<FRDGPooledBuffer> ExposureBufferData[NUM_BUFFERS];
 		TArray<FRHIGPUBufferReadback*> ExposureReadbackBuffers;
 
diff --git a/Engine/Source/Runtime/Renderer/Private/SceneViewState.cpp b/Engine/Source/Runtime/Renderer/Private/SceneViewState.cpp
--- a/Engine/Source/Runtime/Renderer/Private/SceneViewState.cpp
+++ b/Engine/Source/Runtime/Renderer/Private/SceneViewState.cpp
@@ -106,16 +106,10 @@ uint64 FPersistentSkyAtmosphereData::GetGPUSizeBytes(bool bLogSizes) const
 uint64 FSceneViewState::FEyeAdaptationManager::GetGPUSizeBytes(bool bLogSizes) const
 {
 	uint64 TotalSize = 0;
-	PRAGMA_DISABLE_DEPRECATION_WARNINGS
-		for (int32 TargetIndex = 0; TargetIndex < UE_ARRAY_COUNT(PooledRenderTarget); TargetIndex++)
-		{
-			TotalSize += GetRenderTargetGPUSizeBytes(PooledRenderTarget[TargetIndex], bLogSizes);
-		}
-	PRAGMA_ENABLE_DEPRECATION_WARNINGS
-		for (int32 BufferIndex = 0; BufferIndex < UE_ARRAY_COUNT(ExposureBufferData); BufferIndex++)
-		{
-			TotalSize += GetBufferGPUSizeBytes(ExposureBufferData[BufferIndex], bLogSizes);
-		}
+	for (int32 BufferIndex = 0; BufferIndex < UE_ARRAY_COUNT(ExposureBufferData); BufferIndex++)
+	{
+		TotalSize += GetBufferGPUSizeBytes(ExposureBufferData[BufferIndex], bLogSizes);
+	}
 	for (FRHIGPUBufferReadback* ReadbackBuffer : ExposureReadbackBuffers)
 	{
 		TotalSize += GetBufferReadbackGPUSizeBytes(ReadbackBuffer, bLogSizes);
