TSR: Cleanup unused ComposeSeparateTranslucencyInTSR function, whose condition was removed at CL 19326281.

#Synced-CL 50602806

diff --git a/Engine/Source/Runtime/Renderer/Private/PostProcess/PostProcessing.cpp b/Engine/Source/Runtime/Renderer/Private/PostProcess/PostProcessing.cpp
--- a/Engine/Source/Runtime/Renderer/Private/PostProcess/PostProcessing.cpp
+++ b/Engine/Source/Runtime/Renderer/Private/PostProcess/PostProcessing.cpp
@@ -347,8 +347,6 @@ namespace UE::Renderer::PostProcess
 	};
 }
 
-bool ComposeSeparateTranslucencyInTSR(const FViewInfo& View);
-
 void AddPostProcessingPasses(
 	FRDGBuilder& GraphBuilder,
 	const FViewInfo& View, 
@@ -738,7 +736,7 @@ void AddPostProcessingPasses(
 			!FMath::IsNearlyEqual(View.FinalPostProcessSettings.LocalExposureDetailStrength, 1.0f);
 
 		// Whether separate translucency is composed in TSR.
-		bool bComposeSeparateTranslucencyInTSR = PostDOFTranslucencyResources.IsValid() && TAAConfig == EMainTAAPassConfig::TSR && ComposeSeparateTranslucencyInTSR(View);
+		bool bComposeSeparateTranslucencyInTSR = PostDOFTranslucencyResources.IsValid() && TAAConfig == EMainTAAPassConfig::TSR;
 
 		const FIntPoint PostTAAViewSize = (View.PrimaryScreenPercentageMethod == EPrimaryScreenPercentageMethod::TemporalUpscale && TAAConfig != EMainTAAPassConfig::Disabled) ? View.GetSecondaryViewRectSize() : View.ViewRect.Size();
 
diff --git a/Engine/Source/Runtime/Renderer/Private/PostProcess/TemporalSuperResolution.cpp b/Engine/Source/Runtime/Renderer/Private/PostProcess/TemporalSuperResolution.cpp
--- a/Engine/Source/Runtime/Renderer/Private/PostProcess/TemporalSuperResolution.cpp
+++ b/Engine/Source/Runtime/Renderer/Private/PostProcess/TemporalSuperResolution.cpp
@@ -1489,11 +1489,6 @@ DECLARE_GPU_STAT(TemporalSuperResolution);
 
 } //! namespace
 
-bool ComposeSeparateTranslucencyInTSR(const FViewInfo& View)
-{
-	return true;
-}
-
 float GetTSRExposureOffsetFactor()
 {
 	float TSRExposureOffsetFactor = FMath::Pow(2.0f, 
