Fix some translucents not rendering due to using DepthEqual instead of DepthGreaterEqual when we use r.EarlyZPassOnlyMaterialMasking=1
Translucent materials with bCastDynamicShadowAsMasked return true for Material.IsMasked() and end up getting the depth state for masked materials.
Change the test to IsMaskedBlendMode(Material) which only returns true for actual masked materials.
#Synced-CL 50587468

diff --git a/Engine/Source/Runtime/Renderer/Private/BasePassRendering.cpp b/Engine/Source/Runtime/Renderer/Private/BasePassRendering.cpp
--- a/Engine/Source/Runtime/Renderer/Private/BasePassRendering.cpp
+++ b/Engine/Source/Runtime/Renderer/Private/BasePassRendering.cpp
@@ -480,7 +480,7 @@ void SetDepthStencilStateForBasePass(
 	bool bEnableReceiveDecalOutput,
 	bool bForceEnableStencilDitherState)
 {
-	const bool bMaskedInEarlyPass = (MaterialResource.IsMasked() || bDitheredLODTransition) && MaskedInEarlyPass(GShaderPlatformForFeatureLevel[FeatureLevel]);
+	const bool bMaskedInEarlyPass = (IsMaskedBlendMode(MaterialResource) || bDitheredLODTransition) && MaskedInEarlyPass(GShaderPlatformForFeatureLevel[FeatureLevel]);
 	if (bEnableReceiveDecalOutput)
 	{
 		if (bMaskedInEarlyPass)
