Unity compile error fix:  ReflectionEnvironment GPU stat declaration was duplicated between indirect light rendering and mobile deferrred shading.

#rnx

#Synced-CL 50617917

diff --git a/Engine/Source/Runtime/Renderer/Private/IndirectLightRendering.cpp b/Engine/Source/Runtime/Renderer/Private/IndirectLightRendering.cpp
--- a/Engine/Source/Runtime/Renderer/Private/IndirectLightRendering.cpp
+++ b/Engine/Source/Runtime/Renderer/Private/IndirectLightRendering.cpp
@@ -103,7 +103,7 @@ static TAutoConsoleVariable<bool> CVarDiffuseIndirectForceCopyPass(
 	TEXT("Forces use of copy pass instead of dual source blend. (for debugging)"),
 	ECVF_RenderThreadSafe);
 
-DECLARE_GPU_STAT_NAMED(ReflectionEnvironment, TEXT("Reflection Environment"));
+DEFINE_GPU_STAT(ReflectionEnvironment);
 DECLARE_GPU_STAT_NAMED(RayTracingReflections, TEXT("Ray Tracing Reflections"));
 DECLARE_GPU_STAT(SkyLightDiffuse);
 
diff --git a/Engine/Source/Runtime/Renderer/Private/MobileDeferredShadingPass.cpp b/Engine/Source/Runtime/Renderer/Private/MobileDeferredShadingPass.cpp
--- a/Engine/Source/Runtime/Renderer/Private/MobileDeferredShadingPass.cpp
+++ b/Engine/Source/Runtime/Renderer/Private/MobileDeferredShadingPass.cpp
@@ -27,7 +27,6 @@ static TAutoConsoleVariable<bool> CVarMobileDiffuseFromCapture(
 );
 
 DECLARE_GPU_STAT(DeferredShading);
-DECLARE_GPU_STAT(ReflectionEnvironment);
 
 static const TCHAR* DeferredMobileLightMaterialPSOCollectorName = TEXT("DeferredMobileLightMaterialPSOCollector");
 
diff --git a/Engine/Source/Runtime/Renderer/Private/ScenePrivate.h b/Engine/Source/Runtime/Renderer/Private/ScenePrivate.h
--- a/Engine/Source/Runtime/Renderer/Private/ScenePrivate.h
+++ b/Engine/Source/Runtime/Renderer/Private/ScenePrivate.h
@@ -73,6 +73,7 @@
 
 /** Extern GPU stats (used in multiple modules) **/
 DECLARE_GPU_STAT_NAMED_EXTERN(ShadowProjection, TEXT("Shadow Projection"));
+DECLARE_GPU_STAT_NAMED_EXTERN(ReflectionEnvironment, TEXT("Reflection Environment"));
 
 class AWorldSettings;
 class ASceneCaptureCube;
