Added early out checking the instance data header to avoid stalls for primitives that don't have precomputed data.

#rnx

#Synced-CL 50589797

diff --git a/Engine/Source/Runtime/Renderer/Private/SceneCulling/SceneCulling.cpp b/Engine/Source/Runtime/Renderer/Private/SceneCulling/SceneCulling.cpp
--- a/Engine/Source/Runtime/Renderer/Private/SceneCulling/SceneCulling.cpp
+++ b/Engine/Source/Runtime/Renderer/Private/SceneCulling/SceneCulling.cpp
@@ -2793,7 +2793,7 @@ void FSceneCulling::FUpdater::PostSceneUpdate(FRDGBuilder& GraphBuilder, const F
 
 	auto AddOrUpdateHierarchicalCPUCUlling = [&](FPersistentPrimitiveIndex PersistentPrimitiveId, FPrimitiveSceneInfo* PrimitiveSceneInfo) -> bool
 	{
-		if (!PrimitiveSceneInfo->Proxy->IsNaniteMesh())
+		if (PrimitiveSceneInfo->GetInstanceDataHeader().Flags.bHasCompressedSpatialHash && !PrimitiveSceneInfo->Proxy->IsNaniteMesh())
 		{
 			// Note: this causes stall, can we move precomputed ref to Header?
 			if (const FInstanceSceneDataBuffers* InstanceSceneDataBuffers = PrimitiveSceneInfo->GetInstanceSceneDataBuffers())
