Volumetric Cloud - fix for late fog on cloud to not have missing fog when applied on translucent. We now have a dedicated side buffer contructed  just for that when enabling the fix, for titles having translucent card clouds behind the volumetric ones for instance.
#tests PC editor dx12 with volumetric cloud and translucent behind them + ApplyFogLate 0/1 showing the problem and ApplyFogLate.CloudOnTranslucent 1 as the improvement/fix.
#rnx

#Synced-CL 50534707

diff --git a/Engine/Shaders/Private/VolumetricRenderTarget.usf b/Engine/Shaders/Private/VolumetricRenderTarget.usf
--- a/Engine/Shaders/Private/VolumetricRenderTarget.usf
+++ b/Engine/Shaders/Private/VolumetricRenderTarget.usf
@@ -15,6 +15,81 @@
 #define CLOUD_MIN_AND_MAX_DEPTH 0
 #endif
 
+
+
+#define APPLY_FOG_LATE (PERMUTATION_APPLY_FOG || PERMUTATION_APPLY_LOCAL_FOG_VOLUME)
+
+#if PERMUTATION_APPLY_FOG
+#include "HeightFogCommon.ush"
+#endif // PERMUTATION_APPLY_FOG
+
+#if PERMUTATION_APPLY_LOCAL_FOG_VOLUME
+#include "LocalFogVolumes/LocalFogVolumeCommon.ush"
+#endif // PERMUTATION_SUPPORT_LOCAL_FOG_VOLUME
+
+// BufferSVPos is needed when using view matrices operation relying on VolumetricRenderTargetViewUniformBuffer such as GetScreenWorldDir.
+float4 ApplyFogToCloudSampleFromVPosAndDepth(float2 FullResViewSVPos, float2 BufferSVPos, float DepthKilometer, float4 CloudRGBT, float2 View_ViewRectMin)
+{
+#if APPLY_FOG_LATE
+
+	const float CloudCoverage = saturate(1.0 - CloudRGBT.a);
+	if (CloudCoverage <= 0.0f)
+	{
+		return CloudRGBT;
+	}
+
+	float4 HeightFogInscatteringAndTransmittance = float4(0, 0, 0, 1);
+	const float DistanceToPoint = DepthKilometer * KILOMETER_TO_CENTIMETER;
+
+	const float3 FogSampleWorldPositionRelativeToCameraCm = normalize(GetScreenWorldDir(float4(BufferSVPos, 0.5f, 1.0f))) * DistanceToPoint; // bug have to div by 4
+
+#if PERMUTATION_APPLY_FOG
+	{
+		HeightFogInscatteringAndTransmittance = CalculateHeightFog(FogSampleWorldPositionRelativeToCameraCm);
+	}
+#endif
+
+#if PERMUTATION_APPLY_LOCAL_FOG_VOLUME
+	uint2 TilePos = (FullResViewSVPos - View_ViewRectMin) / LFVTilePixelSize.xx; // Remove view rect min as we take the view pixel to sample the tiled LFV culling result
+
+	float4 LFVContribution = GetLFVContribution(PrimaryView, TilePos, FogSampleWorldPositionRelativeToCameraCm);
+	if (LFVRenderInVolumetricFog > 0)
+	{
+		HeightFogInscatteringAndTransmittance = float4(LFVContribution.rgb + HeightFogInscatteringAndTransmittance.rgb * LFVContribution.a, LFVContribution.a * HeightFogInscatteringAndTransmittance.a);
+	}
+#endif
+
+#if PERMUTATION_APPLY_FOG
+	if (FogStruct.ApplyVolumetricFog > 0)
+	{
+		float4 ClipPos = mul(float4(FogSampleWorldPositionRelativeToCameraCm, 1.0f), PrimaryView.TranslatedWorldToClip);
+		float3 VolumeUV = ComputeVolumeUVFromNDC(ClipPos);
+		const uint EyeIndex = 0;
+		HeightFogInscatteringAndTransmittance = CombineVolumetricFog(HeightFogInscatteringAndTransmittance, VolumeUV, EyeIndex, DistanceToPoint);
+	}
+#endif
+
+#if PERMUTATION_APPLY_LOCAL_FOG_VOLUME
+	if (LFVRenderInVolumetricFog == 0)
+	{
+		HeightFogInscatteringAndTransmittance = float4(LFVContribution.rgb + HeightFogInscatteringAndTransmittance.rgb * LFVContribution.a, LFVContribution.a * HeightFogInscatteringAndTransmittance.a);
+	}
+#endif
+
+	HeightFogInscatteringAndTransmittance.rgb *= PrimaryView.PreExposure;
+
+	CloudRGBT.rgb = CloudRGBT.rgb * HeightFogInscatteringAndTransmittance.a + CloudCoverage * HeightFogInscatteringAndTransmittance.rgb;
+
+#endif // APPLY_FOG_LATE
+
+	return CloudRGBT;
+}
+
+
+
+
+
+
 #ifdef SHADER_RECONSTRUCT_VOLUMETRICRT
 
 SamplerState LinearTextureSampler;
@@ -125,7 +200,8 @@
 int2 ViewViewRectMin;
 int VolumetricRenderTargetMode;
 int DownSampleFactor;
+float2 SourceViewRectMin;
+float2 SourceViewRectSize;
 
-
 #define USE_YCOCG 0
 
@@ -177,9 +253,15 @@
 	out float4 OutputRt0 : SV_Target0,
 	out float4 OutputRt1 : SV_Target1,
 	out float4 OutputRt2 : SV_Target2
+#if APPLY_FOG_LATE
+	, out float4 OutputCloudAndFog : SV_Target3
+#endif
 #else
 	out float4 OutputRt0 : SV_Target0,
 	out float4 OutputRt1 : SV_Target1
+#if APPLY_FOG_LATE
+	, out float4 OutputCloudAndFog : SV_Target2
 #endif
+#endif
 )
 {
@@ -292,6 +374,13 @@
 				DepthData.SceneDepthFromViewKm,
 				0.0, 0.0);	// unused components
 		#endif
+
+			#if APPLY_FOG_LATE
+			// Apply fog on the a side reconstructed texture if latefog is used for fog on translucent. Otherwise the regular reconstructed cloud texture will be used and this one do not have fog in it.
+			float2 FullResViewSVPos = 0*SourceViewRectMin + ScreenUV * TracingVolumetricTextureUVScale.xy * SourceViewRectSize;
+			float2 BufferSVPos = FullResViewSVPos / 4.0;
+			OutputCloudAndFog = ApplyFogToCloudSampleFromVPosAndDepth(FullResViewSVPos, BufferSVPos, DepthData.CloudFrontDepthFromViewKm, OutputRt0, 0*SourceViewRectMin);
+			#endif
 			return;
 		}
 		else
@@ -566,6 +655,14 @@
 #endif // PERMUTATION_CLOUD_MIN_AND_MAX_DEPTH
 
 #endif // PERMUTATION_HISTORY_AVAILABLE
+
+
+#if APPLY_FOG_LATE
+	// Apply fog on the a side reconstructed texture if latefog is used for fog on translucent. Otherwise the regular reconstructed cloud texture will be used and this one do not have fog in it.
+	float2 FullResViewSVPos = 0*SourceViewRectMin + ScreenUV * SourceViewRectSize;
+	float2 BufferSVPos = FullResViewSVPos / 4.0;
+	OutputCloudAndFog = ApplyFogToCloudSampleFromVPosAndDepth(FullResViewSVPos, BufferSVPos, DepthData.CloudFrontDepthFromViewKm, OutputRt0, 0*SourceViewRectMin);
+#endif
 }
 
 #endif // SHADER_RECONSTRUCT_VOLUMETRICRT
@@ -576,16 +673,6 @@
 
 #include "SceneTexturesCommon.ush"
 
-#if PERMUTATION_APPLY_FOG
-#include "HeightFogCommon.ush"
-#endif // PERMUTATION_APPLY_FOG
-
-#if PERMUTATION_APPLY_LOCAL_FOG_VOLUME
-#include "LocalFogVolumes/LocalFogVolumeCommon.ush"
-#endif // PERMUTATION_SUPPORT_LOCAL_FOG_VOLUME
-
-#define APPLY_FOG_LATE (PERMUTATION_APPLY_FOG || PERMUTATION_APPLY_LOCAL_FOG_VOLUME)
-
 SamplerState LinearTextureSampler;
 
 Texture2D<float4> VolumetricTexture;
@@ -672,68 +759,11 @@
 float MinimumDistanceKmToDisableDisoclusion;
 float2 FullResolutionToVolumetricBufferResolutionScale;
 
-float4 ApplyFogToCloudSampleFromVPosAndDepth(float2 SVPos, float DepthKilometer, float4 CloudRGBT)
-{
-#if APPLY_FOG_LATE
-
-	const float CloudCoverage = saturate(1.0 - CloudRGBT.a);
-	if (CloudCoverage <= 0.0f)
-	{
-		return CloudRGBT;
-	}
-
-	float4 HeightFogInscatteringAndTransmittance = float4(0, 0, 0, 1);
-	const float2 SVPosition = SVPos;
-	const float DistanceToPoint = DepthKilometer * KILOMETER_TO_CENTIMETER;
-
-	const float3 FogSampleWorldPositionRelativeToCameraCm = normalize(GetScreenWorldDir(float4(SVPosition, 0.5f, 1.0f))) * DistanceToPoint;
-
-#if PERMUTATION_APPLY_FOG
-	{
-		HeightFogInscatteringAndTransmittance = CalculateHeightFog(FogSampleWorldPositionRelativeToCameraCm);
-	}
-#endif
-
-#if PERMUTATION_APPLY_LOCAL_FOG_VOLUME
-	uint2 TilePos = (SVPosition - View.ViewRectMin.xy) / LFVTilePixelSize.xx; // Remove view rect min as we take the view pixel to sample the tiled LFV culling result
-	float4 LFVContribution = GetLFVContribution(PrimaryView, TilePos, FogSampleWorldPositionRelativeToCameraCm);
-	if (LFVRenderInVolumetricFog > 0)
-	{
-		HeightFogInscatteringAndTransmittance = float4(LFVContribution.rgb + HeightFogInscatteringAndTransmittance.rgb * LFVContribution.a, LFVContribution.a * HeightFogInscatteringAndTransmittance.a);
-	}
-#endif
-
-#if PERMUTATION_APPLY_FOG
-	if (FogStruct.ApplyVolumetricFog > 0)
-	{
-		float4 ClipPos = mul(float4(FogSampleWorldPositionRelativeToCameraCm, 1.0f), PrimaryView.TranslatedWorldToClip);
-		float3 VolumeUV = ComputeVolumeUVFromNDC(ClipPos);
-		const uint EyeIndex = 0;
-		HeightFogInscatteringAndTransmittance = CombineVolumetricFog(HeightFogInscatteringAndTransmittance, VolumeUV, EyeIndex, DistanceToPoint);
-	}
-#endif
-
-#if PERMUTATION_APPLY_LOCAL_FOG_VOLUME
-	if (LFVRenderInVolumetricFog == 0)
-	{
-		HeightFogInscatteringAndTransmittance = float4(LFVContribution.rgb + HeightFogInscatteringAndTransmittance.rgb * LFVContribution.a, LFVContribution.a * HeightFogInscatteringAndTransmittance.a);
-	}
-#endif
-
-	HeightFogInscatteringAndTransmittance.rgb *= PrimaryView.PreExposure;
-
-	CloudRGBT.rgb = CloudRGBT.rgb * HeightFogInscatteringAndTransmittance.a + CloudCoverage * HeightFogInscatteringAndTransmittance.rgb;
-
-#endif // APPLY_FOG_LATE
-
-	return CloudRGBT;
-}
-
 float4 ApplyFogToCloudFromUVsAndDepth(float2 UVWithOffset, float DepthKilometer, float4 CloudRGBT)
 {
 #if APPLY_FOG_LATE
 	float2 SVPosWithOffset = (UVWithOffset / (FullResolutionToVolumetricBufferResolutionScale.x *View.BufferSizeAndInvSize.xy * VolumetricTextureSizeAndInvSize.zw)) * View.BufferSizeAndInvSize.xy + View.ViewRectMin.xy;
-	CloudRGBT = ApplyFogToCloudSampleFromVPosAndDepth(SVPosWithOffset, DepthKilometer, CloudRGBT);
+	CloudRGBT = ApplyFogToCloudSampleFromVPosAndDepth(SVPosWithOffset, SVPosWithOffset, DepthKilometer, CloudRGBT, View.ViewRectMin.xy);
 #endif // APPLY_FOG_LATE
 
 	return CloudRGBT;
@@ -744,7 +774,7 @@
 #if APPLY_FOG_LATE
 	float2 SVPosWithOffset = (UVWithOffset / (FullResolutionToVolumetricBufferResolutionScale.x * View.BufferSizeAndInvSize.xy * VolumetricTextureSizeAndInvSize.zw)) * View.BufferSizeAndInvSize.xy + View.ViewRectMin.xy;
 	float VolumeFrontDepthKilometer = SafeSampleVolumetricDepthTexture(UVWithOffset).x;
-	CloudRGBT = ApplyFogToCloudSampleFromVPosAndDepth(SVPosWithOffset, VolumeFrontDepthKilometer, CloudRGBT);
+	CloudRGBT = ApplyFogToCloudSampleFromVPosAndDepth(SVPosWithOffset, SVPosWithOffset, VolumeFrontDepthKilometer, CloudRGBT, View.ViewRectMin.xy);
 #endif // APPLY_FOG_LATE
 
 	return CloudRGBT;
@@ -795,7 +825,7 @@
 //	OutputRt0 = SafeLoadVolumetricTexture(VolumeUV * VolumetricTextureSizeAndInvSize.xy);			// Closest
 	OutputRt0 = SafeLoadVolumetricSecondaryTexture(VolumeUV * VolumetricTextureSizeAndInvSize.xy);	// Furthest
 	float4 VolumeFrontDepth0 = SafeSampleVolumetricDepthTexture(VolumeUV);
-	OutputRt0 = ApplyFogToCloudSampleFromVPosAndDepth(CurResPixelCoord, VolumeFrontDepth0.x, OutputRt0);
+	OutputRt0 = ApplyFogToCloudSampleFromVPosAndDepth(CurResPixelCoord, CurResPixelCoord, VolumeFrontDepth0.x, OutputRt0, View.ViewRectMin.xy);
 	return;
 
 #elif PERMUTATION_UPSAMPLINGMODE==1
diff --git a/Engine/Source/Runtime/Renderer/Private/TranslucentRendering.cpp b/Engine/Source/Runtime/Renderer/Private/TranslucentRendering.cpp
--- a/Engine/Source/Runtime/Renderer/Private/TranslucentRendering.cpp
+++ b/Engine/Source/Runtime/Renderer/Private/TranslucentRendering.cpp
@@ -1118,7 +1118,18 @@ TRDGUniformBufferRef<FTranslucentBasePassUniformParameters> CreateTranslucentBas
 			}
 			else
 			{
-				TRefCountPtr<IPooledRenderTarget> VolumetricReconstructRT = View.ViewState->VolumetricCloudRenderTarget.GetDstVolumetricReconstructRT();
+				const bool bShouldVolumetricCloudsApplyFogDuringUpsamplingForCloudOnTranslucent = ShouldVolumetricCloudsApplyFogDuringUpsamplingForCloudOnTranslucent(View);
+
+				TRefCountPtr<IPooledRenderTarget> VolumetricReconstructRT = nullptr;
+				if (bShouldVolumetricCloudsApplyFogDuringUpsamplingForCloudOnTranslucent)
+				{
+					VolumetricReconstructRT = View.ViewState->VolumetricCloudRenderTarget.GetVolumetricReconstructRTForCloudOnTranslucent();
+				}
+				else
+				{
+					VolumetricReconstructRT = View.ViewState->VolumetricCloudRenderTarget.GetDstVolumetricReconstructRT();
+				}
+
 				if (VolumetricReconstructRT.IsValid())
 				{
 					TRefCountPtr<IPooledRenderTarget> VolumetricReconstructRTDepth = View.ViewState->VolumetricCloudRenderTarget.GetDstVolumetricReconstructRTDepth();
diff --git a/Engine/Source/Runtime/Renderer/Private/VolumetricCloudRendering.cpp b/Engine/Source/Runtime/Renderer/Private/VolumetricCloudRendering.cpp
--- a/Engine/Source/Runtime/Renderer/Private/VolumetricCloudRendering.cpp
+++ b/Engine/Source/Runtime/Renderer/Private/VolumetricCloudRendering.cpp
@@ -103,6 +103,11 @@ static TAutoConsoleVariable<int32> CVarVolumetricCloudApplyFogLate(
 	TEXT("When true, fog will be applied on cloud later in the frame. This applies fog on cloud without visible latency due to reprojection and reconstruction. It is a bit more expenssive (upsampling pass cost is roughly x2) but the only correct look when Lumen is enabled and tracing is run on async compute using r.VolumetricRenderTarget.PreferAsyncCompute=1."),
 	ECVF_RenderThreadSafe);
 
+static TAutoConsoleVariable<int32> CVarVolumetricCloudApplyFogLateCloudOnTranslucent(
+	TEXT("r.VolumetricCloud.ApplyFogLate.CloudOnTranslucent"), 0,
+	TEXT("When r.VolumetricCloud.ApplyFogLate=1, fog is not into the volumetric cloud reconstructed buffer. So when applied on translucent, a visual discrepancy can become visible for translucent behind clouds. This enables the reconstruction step to also output an extra cloud buffer with fog included for application on translucent. It is more expenssive and use more memory so it is disabled by default."),
+	ECVF_RenderThreadSafe);
+
 ////////////////////////////////////////////////////////////////////////// Shadow tracing
 
 static TAutoConsoleVariable<float> CVarVolumetricCloudShadowViewRaySampleMaxCount(
@@ -377,7 +382,7 @@ bool VolumetricCloudWantsSeparatedAtmosphereMieRayLeigh(const FScene* Scene)
 	return false;
 }
 
-bool ShouldVolumetricCloudsApplyFogDuringReconstruction(const FViewInfo& ViewInfo)
+bool ShouldVolumetricCloudsApplyFogDuringUpsampling(const FViewInfo& ViewInfo)
 {
 	if (ViewInfo.ViewState)
 	{
@@ -388,6 +393,15 @@ bool ShouldVolumetricCloudsApplyFogDuringReconstruction(const FViewInfo& ViewInf
 	return false;
 }
 
+bool ShouldVolumetricCloudsApplyFogDuringUpsamplingForCloudOnTranslucent(const FViewInfo& ViewInfo)
+{
+	if (ViewInfo.ViewState)
+	{
+		return ShouldVolumetricCloudsApplyFogDuringUpsampling(ViewInfo) && CVarVolumetricCloudApplyFogLateCloudOnTranslucent.GetValueOnRenderThread() > 0;
+	}
+	return false;
+}
+
 //////////////////////////////////////////////////////////////////////////
 
 static float GetVolumetricCloudShadowmapStrength(const FLightSceneProxy* AtmosphericLight)
@@ -732,7 +746,7 @@ void SetupDefaultRenderVolumetricCloudGlobalParameters(FRDGBuilder& GraphBuilder
 		VolumetricCloudParams.EnableHeightFog = CVarVolumetricCloudApplyFogOnAllPixel.GetValueOnRenderThread() > 0 ? 2 : 1;
 	}
 	SetupFogUniformParameters(GraphBuilder, ViewInfo, VolumetricCloudParams.FogStruct);
-	if (ShouldVolumetricCloudsApplyFogDuringReconstruction(ViewInfo) && !bForceNonLateFogOnCloud)
+	if (ShouldVolumetricCloudsApplyFogDuringUpsampling(ViewInfo) && !bForceNonLateFogOnCloud)
 	{
 		VolumetricCloudParams.FogStruct.ApplyVolumetricFog = 0;
 		VolumetricCloudParams.EnableHeightFog = 0;
@@ -2367,7 +2381,7 @@ static TRDGUniformBufferRef<FRenderVolumetricCloudGlobalParameters> CreateCloudP
 	}
 	else
 	{
-		if (ShouldVolumetricCloudsApplyFogDuringReconstruction(MainView))
+		if (ShouldVolumetricCloudsApplyFogDuringUpsampling(MainView))
 		{
 			VolumetricCloudParams.FogStruct.ApplyVolumetricFog = 0;
 			VolumetricCloudParams.EnableHeightFog = 0;
@@ -2614,7 +2628,7 @@ void FSceneRenderer::RenderVolumetricCloudsInternal(FRDGBuilder& GraphBuilder, F
 		PermutationVector.Set<typename FRenderVolumetricCloudRenderViewCS::FCloudSampleSecondLight>(!bCloudDebugViewModeEnabled && bSecondAtmosphereLightEnabled);
 		PermutationVector.Set<typename FRenderVolumetricCloudRenderViewCS::FCloudSampleLocalLights>(!bCloudDebugViewModeEnabled && bCloudEnableLocalLightSampling && !CloudRC.bIsSkyRealTimeReflectionRendering);
 		PermutationVector.Set<typename FRenderVolumetricCloudRenderViewCS::FCloudMinAndMaxDepth>(ShouldVolumetricCloudTraceWithMinMaxDepth(MainView) && CloudRC.SecondaryCloudTracingDataTexture!=nullptr && !CloudRC.bIsReflectionRendering);// Only for mode 0 in non reflection
-		PermutationVector.Set<typename FRenderVolumetricCloudRenderViewCS::FCloudLocalFogVolume>(MainView.LocalFogVolumeViewData.GPUInstanceCount > 0 && !CloudRC.bIsSkyRealTimeReflectionRendering && !ShouldVolumetricCloudsApplyFogDuringReconstruction(MainView)); // LFVs are not culled for real time capture views.
+		PermutationVector.Set<typename FRenderVolumetricCloudRenderViewCS::FCloudLocalFogVolume>(MainView.LocalFogVolumeViewData.GPUInstanceCount > 0 && !CloudRC.bIsSkyRealTimeReflectionRendering && !ShouldVolumetricCloudsApplyFogDuringUpsampling(MainView)); // LFVs are not culled for real time capture views.
 		PermutationVector.Set<typename FRenderVolumetricCloudRenderViewCS::FCloudDebugViewMode>(bCloudDebugViewModeEnabled);
 
 		TShaderRef<FRenderVolumetricCloudRenderViewCS> ComputeShader = MaterialResource->GetShader<FRenderVolumetricCloudRenderViewCS>(&FLocalVertexFactory::StaticType, PermutationVector, false);
diff --git a/Engine/Source/Runtime/Renderer/Private/VolumetricCloudRendering.h b/Engine/Source/Runtime/Renderer/Private/VolumetricCloudRendering.h
--- a/Engine/Source/Runtime/Renderer/Private/VolumetricCloudRendering.h
+++ b/Engine/Source/Runtime/Renderer/Private/VolumetricCloudRendering.h
@@ -116,7 +116,8 @@ bool ShouldVolumetricCloudTraceWithMinMaxDepth(const FViewInfo& ViewInfo);
 bool ShouldVolumetricCloudTraceWithMinMaxDepth(const TArray<FViewInfo>& Views);
 bool VolumetricCloudWantsSeparatedAtmosphereMieRayLeigh(const FScene* Scene);
 
-bool ShouldVolumetricCloudsApplyFogDuringReconstruction(const FViewInfo& ViewInfo);
+bool ShouldVolumetricCloudsApplyFogDuringUpsampling(const FViewInfo& ViewInfo);
+bool ShouldVolumetricCloudsApplyFogDuringUpsamplingForCloudOnTranslucent(const FViewInfo& ViewInfo);
 
 // Structure with data necessary to specify a cloud render.
 struct FCloudRenderContext
diff --git a/Engine/Source/Runtime/Renderer/Private/VolumetricRenderTarget.cpp b/Engine/Source/Runtime/Renderer/Private/VolumetricRenderTarget.cpp
--- a/Engine/Source/Runtime/Renderer/Private/VolumetricRenderTarget.cpp
+++ b/Engine/Source/Runtime/Renderer/Private/VolumetricRenderTarget.cpp
@@ -286,6 +286,7 @@ void FVolumetricRenderTargetViewStateData::Initialise(
 			VolumetricReconstructRT[CurrentRT].SafeRelease();
 			VolumetricReconstructSecondaryRT[CurrentRT].SafeRelease();
 			VolumetricReconstructRTDepth[CurrentRT].SafeRelease();
+			VolumetricReconstructRTForCloudOnTranslucent.SafeRelease();
 		}
 
 		// Regular every frame update
@@ -353,6 +354,7 @@ void FVolumetricRenderTargetViewStateData::Reset()
 		VolumetricReconstructSecondaryRT[i].SafeRelease();
 		VolumetricReconstructRTDepth[i].SafeRelease();
 	}
+	VolumetricReconstructRTForCloudOnTranslucent.SafeRelease();
 }
 
 FRDGTextureRef FVolumetricRenderTargetViewStateData::GetOrCreateVolumetricTracingRT(FRDGBuilder& GraphBuilder)
@@ -436,6 +438,26 @@ FVector2f FVolumetricRenderTargetViewStateData::GetVolumetricTracingUVMax() cons
 	return UVMax;
 }
 
+FRDGTextureRef FVolumetricRenderTargetViewStateData::GetOrCreateVolumetricReconstructRTForCloudOnTranslucent(FRDGBuilder& GraphBuilder)
+{
+	check(VolumetricReconstructRTResolution != FIntPoint::ZeroValue); // check that initialization has been done at least once
+
+	if (!VolumetricReconstructRTForCloudOnTranslucent.IsValid())
+	{
+		FPooledRenderTargetDesc Desc = FPooledRenderTargetDesc::Create2DDesc(
+			VolumetricReconstructRTResolution, PF_FloatRGBA, FClearValueBinding(FLinearColor(0.0f, 0.0f, 0.0f, 1.0f)),
+			TexCreate_None, TexCreate_ShaderResource | TexCreate_RenderTargetable, false);
+		GRenderTargetPool.FindFreeElement(GraphBuilder.RHICmdList, Desc, VolumetricReconstructRTForCloudOnTranslucent, TEXT("VolumetricRenderTarget.Reconstruct.WithFog"));
+	}
+
+	return GraphBuilder.RegisterExternalTexture(VolumetricReconstructRTForCloudOnTranslucent);
+}
+
+TRefCountPtr<IPooledRenderTarget> FVolumetricRenderTargetViewStateData::GetVolumetricReconstructRTForCloudOnTranslucent()
+{
+	return VolumetricReconstructRTForCloudOnTranslucent;
+}
+
 FRDGTextureRef FVolumetricRenderTargetViewStateData::GetDstVolumetricReconstructRT(FRDGBuilder& GraphBuilder)
 {
 	return VolumetricReconstructRT[CurrentRT].IsValid() ? GraphBuilder.RegisterExternalTexture(VolumetricReconstructRT[CurrentRT]) : nullptr;
@@ -704,10 +726,14 @@ class FReconstructVolumetricRenderTargetPS : public FGlobalShader
 	class FHistoryAvailable : SHADER_PERMUTATION_BOOL("PERMUTATION_HISTORY_AVAILABLE");
 	class FReprojectionBoxConstraint : SHADER_PERMUTATION_BOOL("PERMUTATION_REPROJECTION_BOX_CONSTRAINT");
 	class FCloudMinAndMaxDepth : SHADER_PERMUTATION_BOOL("PERMUTATION_CLOUD_MIN_AND_MAX_DEPTH");
-	using FPermutationDomain = TShaderPermutationDomain<FHistoryAvailable, FReprojectionBoxConstraint, FCloudMinAndMaxDepth>;
+	class FApplyFog : SHADER_PERMUTATION_BOOL("PERMUTATION_APPLY_FOG");
+	class FApplyLocalFogVolume : SHADER_PERMUTATION_BOOL("PERMUTATION_APPLY_LOCAL_FOG_VOLUME");
+	using FPermutationDomain = TShaderPermutationDomain<FHistoryAvailable, FReprojectionBoxConstraint, FCloudMinAndMaxDepth, FApplyFog, FApplyLocalFogVolume>;
 
 	BEGIN_SHADER_PARAMETER_STRUCT(FParameters, )
 		SHADER_PARAMETER_STRUCT_REF(FViewUniformShaderParameters, ViewUniformBuffer)
+		SHADER_PARAMETER_STRUCT(FLocalFogVolumeUniformParameters, LFV)
+		SHADER_PARAMETER_RDG_UNIFORM_BUFFER(FFogUniformParameters, FogStruct)
 		SHADER_PARAMETER_RDG_TEXTURE(Texture2D, TracingVolumetricTexture)
 		SHADER_PARAMETER_RDG_TEXTURE(Texture2D, SecondaryTracingVolumetricTexture)
 		SHADER_PARAMETER_RDG_TEXTURE(Texture2D, TracingVolumetricDepthTexture)
@@ -732,6 +758,8 @@ class FReconstructVolumetricRenderTargetPS : public FGlobalShader
 		SHADER_PARAMETER(float, MinimumDistanceKmToDisableDisoclusion)
 		SHADER_PARAMETER(float, HistoryPreExposureCorrection)
 		SHADER_PARAMETER(FVector2f, PreviousFrameVolumetricTextureUVScale)
+		SHADER_PARAMETER(FVector2f, SourceViewRectMin)
+		SHADER_PARAMETER(FVector2f, SourceViewRectSize)
 	END_SHADER_PARAMETER_STRUCT()
 
 	static FPermutationDomain RemapPermutation(FPermutationDomain PermutationVector)
@@ -794,6 +822,7 @@ void ReconstructVolumetricRenderTarget(
 			// No reconstruction needed.
 			continue;
 		}
+		const bool bShouldVolumetricCloudsApplyFogDuringUpsamplingForCloudOnTranslucent = ShouldVolumetricCloudsApplyFogDuringUpsamplingForCloudOnTranslucent(ViewInfo);
 
 		FRDGTextureRef DstVolumetric = VolumetricCloudRT.GetOrCreateDstVolumetricReconstructRT(GraphBuilder);
 		FRDGTextureRef DstVolumetricSecondary = VolumetricCloudRT.GetOrCreateDstVolumetricReconstructSecondaryRT(GraphBuilder);
@@ -813,6 +842,8 @@ void ReconstructVolumetricRenderTarget(
 		PermutationVector.Set<FReconstructVolumetricRenderTargetPS::FHistoryAvailable>(VolumetricCloudRT.GetHistoryValid());
 		PermutationVector.Set<FReconstructVolumetricRenderTargetPS::FReprojectionBoxConstraint>(CVarVolumetricRenderTargetReprojectionBoxConstraint.GetValueOnAnyThread() > 0);
 		PermutationVector.Set<FReconstructVolumetricRenderTargetPS::FCloudMinAndMaxDepth>(bMinMaxDepth);
+		PermutationVector.Set<FReconstructVolumetricRenderTargetPS::FApplyFog>(bShouldVolumetricCloudsApplyFogDuringUpsamplingForCloudOnTranslucent && (ViewInfo.Family->Scene->HasAnyExponentialHeightFog() || ViewInfo.LocalFogVolumeViewData.GPUInstanceCount > 0) && ShouldRenderFog(*ViewInfo.Family));
+		PermutationVector.Set<FReconstructVolumetricRenderTargetPS::FApplyLocalFogVolume>(bShouldVolumetricCloudsApplyFogDuringUpsamplingForCloudOnTranslucent && ViewInfo.LocalFogVolumeViewData.GPUInstanceCount > 0);
 		TShaderMapRef<FReconstructVolumetricRenderTargetPS> PixelShader(ViewInfo.ShaderMap, PermutationVector);
 
 		FReconstructVolumetricRenderTargetPS::FParameters* PassParameters = GraphBuilder.AllocParameters<FReconstructVolumetricRenderTargetPS::FParameters>();
@@ -822,10 +853,18 @@ void ReconstructVolumetricRenderTarget(
 		{
 			PassParameters->RenderTargets[1] = FRenderTargetBinding(DstVolumetricSecondary, ERenderTargetLoadAction::ENoAction);
 			PassParameters->RenderTargets[2] = FRenderTargetBinding(DstVolumetricDepth, ERenderTargetLoadAction::ENoAction);
+			if (bShouldVolumetricCloudsApplyFogDuringUpsamplingForCloudOnTranslucent)
+			{
+				PassParameters->RenderTargets[3] = FRenderTargetBinding(VolumetricCloudRT.GetOrCreateVolumetricReconstructRTForCloudOnTranslucent(GraphBuilder), ERenderTargetLoadAction::ENoAction);
+			}
 		}
 		else
 		{
 			PassParameters->RenderTargets[1] = FRenderTargetBinding(DstVolumetricDepth, ERenderTargetLoadAction::ENoAction);
+			if (bShouldVolumetricCloudsApplyFogDuringUpsamplingForCloudOnTranslucent)
+			{
+				PassParameters->RenderTargets[2] = FRenderTargetBinding(VolumetricCloudRT.GetOrCreateVolumetricReconstructRTForCloudOnTranslucent(GraphBuilder), ERenderTargetLoadAction::ENoAction);
+			}
 		}
 		PassParameters->TracingVolumetricTexture = SrcTracingVolumetric;
 		PassParameters->SecondaryTracingVolumetricTexture = SrcSecondaryTracingVolumetric;
@@ -843,6 +882,15 @@ void ReconstructVolumetricRenderTarget(
 		PassParameters->MinimumDistanceKmToDisableDisoclusion = GetMinimumDistanceKmToDisableDisoclusion();
 		PassParameters->HistoryPreExposureCorrection = ViewInfo.PreExposure / VolumetricCloudRT.GetPrevViewExposure();
 
+		PassParameters->SourceViewRectMin = ViewInfo.ViewRect.Min;
+		PassParameters->SourceViewRectSize = ViewInfo.ViewRect.Size();
+
+		if (bShouldVolumetricCloudsApplyFogDuringUpsamplingForCloudOnTranslucent)
+		{
+			PassParameters->FogStruct = CreateFogUniformBuffer(GraphBuilder, ViewInfo);
+			PassParameters->LFV = ViewInfo.LocalFogVolumeViewData.UniformParametersStruct;
+		}
+
 		const bool bVisualizeConservativeDensity = ShouldViewVisualizeVolumetricCloudConservativeDensity(ViewInfo, ViewInfo.Family->EngineShowFlags);
 		PassParameters->HalfResDepthTexture = bVisualizeConservativeDensity ?
 			SystemTextures.Black :
@@ -863,7 +911,7 @@ void ReconstructVolumetricRenderTarget(
 		PassParameters->PreviousVolumetricTextureSizeAndInvSize = FVector4f(PreviousVolumetricTextureSize.X, PreviousVolumetricTextureSize.Y, 1.0f / PreviousVolumetricTextureSize.X, 1.0f / PreviousVolumetricTextureSize.Y);
 
 		FPixelShaderUtils::AddFullscreenPass<FReconstructVolumetricRenderTargetPS>(
-			GraphBuilder, ViewInfo.ShaderMap, RDG_EVENT_NAME("VolumetricReconstruct"), PixelShader, PassParameters,
+			GraphBuilder, ViewInfo.ShaderMap, RDG_EVENT_NAME("VolumetricReconstruct"), PixelShader, PassParameters, 
 			FIntRect(0, 0, DstVolumetricSize.X, DstVolumetricSize.Y));
 
 		VolumetricCloudRT.PostRenderUpdate(ViewInfo.PreExposure);
@@ -1045,15 +1093,15 @@ void ComposeVolumetricRenderTargetOverScene(
 		const bool bForwardShading = IsForwardShadingEnabled(ViewInfo.GetShaderPlatform());
 		check(bForwardShading || (!bForwardShading && SceneDepthTexture->Desc.NumSamples==1));
 
-		const bool bShouldVolumetricCloudsApplyFogDuringReconstruction = ShouldVolumetricCloudsApplyFogDuringReconstruction(ViewInfo);
+		const bool bShouldVolumetricCloudsApplyFogDuringUpsampling = ShouldVolumetricCloudsApplyFogDuringUpsampling(ViewInfo);
 
 		FComposeVolumetricRTOverScenePS::FPermutationDomain PermutationVector;
 		PermutationVector.Set<FComposeVolumetricRTOverScenePS::FUpsamplingMode>(UpsamplingMode);
 		PermutationVector.Set<FComposeVolumetricRTOverScenePS::FRenderUnderWaterBuffer>(0);
 		PermutationVector.Set<FComposeVolumetricRTOverScenePS::FRenderCameraComposeWithWater>((bComposeWithWater && ViewInfo.IsUnderwater()) ? 1 : 0);
 		PermutationVector.Set<FComposeVolumetricRTOverScenePS::FMSAASampleCount>(SceneDepthTexture->Desc.NumSamples);
-		PermutationVector.Set<FComposeVolumetricRTOverScenePS::FApplyFog>(bShouldVolumetricCloudsApplyFogDuringReconstruction && (ViewInfo.Family->Scene->HasAnyExponentialHeightFog() || ViewInfo.LocalFogVolumeViewData.GPUInstanceCount > 0) && ShouldRenderFog(*ViewInfo.Family));
-		PermutationVector.Set<FComposeVolumetricRTOverScenePS::FApplyLocalFogVolume>(bShouldVolumetricCloudsApplyFogDuringReconstruction && ViewInfo.LocalFogVolumeViewData.GPUInstanceCount > 0);
+		PermutationVector.Set<FComposeVolumetricRTOverScenePS::FApplyFog>(bShouldVolumetricCloudsApplyFogDuringUpsampling && (ViewInfo.Family->Scene->HasAnyExponentialHeightFog() || ViewInfo.LocalFogVolumeViewData.GPUInstanceCount > 0) && ShouldRenderFog(*ViewInfo.Family));
+		PermutationVector.Set<FComposeVolumetricRTOverScenePS::FApplyLocalFogVolume>(bShouldVolumetricCloudsApplyFogDuringUpsampling && ViewInfo.LocalFogVolumeViewData.GPUInstanceCount > 0);
 		PermutationVector.Set<FComposeVolumetricRTOverScenePS::FMinMaxDepthAvailable>(bMinMaxDepth);
 		TShaderMapRef<FComposeVolumetricRTOverScenePS> PixelShader(ViewInfo.ShaderMap, PermutationVector);
 
@@ -1158,7 +1206,7 @@ void ComposeVolumetricRenderTargetOverSceneUnderWater(
 		int32 UpsamplingMode = GetCompositionUpsamplingMode(VRTMode, VolumetricCloudRT.GetUpsamplingMode());
 		const bool bMinMaxDepth = ShouldVolumetricCloudTraceWithMinMaxDepth(ViewInfo);
 
-		const bool bShouldVolumetricCloudsApplyFogDuringReconstruction = ShouldVolumetricCloudsApplyFogDuringReconstruction(ViewInfo);
+		const bool bShouldVolumetricCloudsApplyFogDuringUpsampling = ShouldVolumetricCloudsApplyFogDuringUpsampling(ViewInfo);
 
 		FRDGTextureRef VolumetricTexture = nullptr;
 		FRDGTextureRef VolumetricSecondaryTexture = nullptr;
@@ -1171,8 +1219,8 @@ void ComposeVolumetricRenderTargetOverSceneUnderWater(
 		PermutationVector.Set<FComposeVolumetricRTOverScenePS::FRenderUnderWaterBuffer>(1);
 		PermutationVector.Set<FComposeVolumetricRTOverScenePS::FRenderCameraComposeWithWater>(0);
 		PermutationVector.Set<FComposeVolumetricRTOverScenePS::FMSAASampleCount>(1);
-		PermutationVector.Set<FComposeVolumetricRTOverScenePS::FApplyFog>(bShouldVolumetricCloudsApplyFogDuringReconstruction && (ViewInfo.Family->Scene->HasAnyExponentialHeightFog() || ViewInfo.LocalFogVolumeViewData.GPUInstanceCount > 0) && ShouldRenderFog(*ViewInfo.Family));
-		PermutationVector.Set<FComposeVolumetricRTOverScenePS::FApplyLocalFogVolume>(bShouldVolumetricCloudsApplyFogDuringReconstruction && ViewInfo.LocalFogVolumeViewData.GPUInstanceCount > 0);
+		PermutationVector.Set<FComposeVolumetricRTOverScenePS::FApplyFog>(bShouldVolumetricCloudsApplyFogDuringUpsampling && (ViewInfo.Family->Scene->HasAnyExponentialHeightFog() || ViewInfo.LocalFogVolumeViewData.GPUInstanceCount > 0) && ShouldRenderFog(*ViewInfo.Family));
+		PermutationVector.Set<FComposeVolumetricRTOverScenePS::FApplyLocalFogVolume>(bShouldVolumetricCloudsApplyFogDuringUpsampling && ViewInfo.LocalFogVolumeViewData.GPUInstanceCount > 0);
 		PermutationVector.Set<FComposeVolumetricRTOverScenePS::FMinMaxDepthAvailable>(bMinMaxDepth);
 		TShaderMapRef<FComposeVolumetricRTOverScenePS> PixelShader(ViewInfo.ShaderMap, PermutationVector);
 
@@ -1197,7 +1245,7 @@ void ComposeVolumetricRenderTargetOverSceneUnderWater(
 		PassParameters->ForwardShadingEnable = 0;
 		PassParameters->VolumeTracingStartDistanceFromCamera = VolumetricCloudRT.GetStartTracingDistance();
 
-		if (bShouldVolumetricCloudsApplyFogDuringReconstruction)
+		if (bShouldVolumetricCloudsApplyFogDuringUpsampling)
 		{
 			PassParameters->FogStruct = CreateFogUniformBuffer(GraphBuilder, ViewInfo);
 			PassParameters->LFV = ViewInfo.LocalFogVolumeViewData.UniformParametersStruct;
diff --git a/Engine/Source/Runtime/Renderer/Private/VolumetricRenderTargetViewStateData.h b/Engine/Source/Runtime/Renderer/Private/VolumetricRenderTargetViewStateData.h
--- a/Engine/Source/Runtime/Renderer/Private/VolumetricRenderTargetViewStateData.h
+++ b/Engine/Source/Runtime/Renderer/Private/VolumetricRenderTargetViewStateData.h
@@ -56,6 +56,9 @@ public:
 	FVector2f GetVolumetricTracingUVScale() const;
 	FVector2f GetVolumetricTracingUVMax() const;
 
+	FRDGTextureRef GetOrCreateVolumetricReconstructRTForCloudOnTranslucent(FRDGBuilder& GraphBuilder);
+	TRefCountPtr<IPooledRenderTarget> GetVolumetricReconstructRTForCloudOnTranslucent();
+
 	FRDGTextureRef GetDstVolumetricReconstructRT(FRDGBuilder& GraphBuilder);
 	FRDGTextureRef GetDstVolumetricReconstructSecondaryRT(FRDGBuilder& GraphBuilder);
 	FRDGTextureRef GetOrCreateDstVolumetricReconstructRT(FRDGBuilder& GraphBuilder);
@@ -124,6 +127,8 @@ private:
 	TRefCountPtr<IPooledRenderTarget> VolumetricReconstructRTDepth[kRenderTargetCount];
 	FIntPoint VolumetricReconstructViewRect[kRenderTargetCount];
 
+	TRefCountPtr<IPooledRenderTarget> VolumetricReconstructRTForCloudOnTranslucent;
+
 	TRefCountPtr<IPooledRenderTarget> VolumetricTracingRT;
 	TRefCountPtr<IPooledRenderTarget> VolumetricSecondaryTracingRT;
 	TRefCountPtr<IPooledRenderTarget> VolumetricTracingRTDepth;
