Add VRSImage_Fractional format (VK_EXT_fragment_density_map/FDM) support to r.VRS.DebugForceRate. Automatically disable VRS Image Generators that to do not support VRSImage_Fractional when it is in use.

#Synced-CL 50657708

diff --git a/Engine/Source/Runtime/Renderer/Private/VariableRateShading/ContrastAdaptiveImageGenerator.cpp b/Engine/Source/Runtime/Renderer/Private/VariableRateShading/ContrastAdaptiveImageGenerator.cpp
--- a/Engine/Source/Runtime/Renderer/Private/VariableRateShading/ContrastAdaptiveImageGenerator.cpp
+++ b/Engine/Source/Runtime/Renderer/Private/VariableRateShading/ContrastAdaptiveImageGenerator.cpp
@@ -541,6 +541,7 @@ bool FContrastAdaptiveImageGenerator::IsEnabled() const
 bool FContrastAdaptiveImageGenerator::IsSupportedByView(const FSceneView& View) const
 {
 	// See also the requirments under FVariableRateShadingImageManager::IsVRSCompatibleWithView
+	// :todo-cf: Revise this to account for fractional vs. palette format if CAS support is extended to mobile
 	return !IsMobilePlatform(View.GetShaderPlatform());
 }
 
diff --git a/Engine/Source/Runtime/Renderer/Private/VariableRateShading/FoveatedImageGenerator.cpp b/Engine/Source/Runtime/Renderer/Private/VariableRateShading/FoveatedImageGenerator.cpp
--- a/Engine/Source/Runtime/Renderer/Private/VariableRateShading/FoveatedImageGenerator.cpp
+++ b/Engine/Source/Runtime/Renderer/Private/VariableRateShading/FoveatedImageGenerator.cpp
@@ -253,8 +253,9 @@ bool FFoveatedImageGenerator::IsEnabled() const
 
 bool FFoveatedImageGenerator::IsSupportedByView(const FSceneView& View) const
 {
-	// Only used for XR views
-	return true; // IStereoRendering::IsStereoEyeView(View);
+	// On desktop, we always use palette rates, but may want this available for Software VRS even if DataType == VRSImage_Unsupported
+	// On mobile, we may use palette or fractional rates, and this generator currently only supports palette rates
+	return !IsMobilePlatform(View.GetShaderPlatform()) || GRHIVariableRateShadingImageDataType == VRSImage_Palette;
 }
 
 FVariableRateShadingImageManager::EVRSSourceType FFoveatedImageGenerator::GetType() const
diff --git a/Engine/Source/Runtime/Renderer/Private/VariableRateShading/VariableRateShadingImageManager.cpp b/Engine/Source/Runtime/Renderer/Private/VariableRateShading/VariableRateShadingImageManager.cpp
--- a/Engine/Source/Runtime/Renderer/Private/VariableRateShading/VariableRateShadingImageManager.cpp
+++ b/Engine/Source/Runtime/Renderer/Private/VariableRateShading/VariableRateShadingImageManager.cpp
@@ -805,9 +805,10 @@ FRDGTextureRef FVariableRateShadingImageManager::CombineShadingRateImages(FRDGBu
 FRDGTextureRef FVariableRateShadingImageManager::GetForceRateImage(FRDGBuilder& GraphBuilder, const FSceneViewFamily& ViewFamily, int RateIndex /* = 0*/, bool bGetSoftwareImage)
 {
 	static const TArray<uint32> ValidShadingRates = { VRSSR_1x1, VRSSR_1x2, VRSSR_2x1, VRSSR_2x2, VRSSR_2x4, VRSSR_4x2, VRSSR_4x4 };
+	static const TArray<FVector2D> ValidFragmentDensities = { FVector2D(1.0, 1.0), FVector2D(1.0, 0.5), FVector2D(0.5, 1.0), FVector2D(0.5, 0.5), FVector2D(0.5, 0.25), FVector2D(0.25, 0.5), FVector2D(0.25, 0.25) };
 
 	const bool bImageTypeAvailable = bGetSoftwareImage ? bSoftwareVRSEnabledForFrame : bHardwareVRSEnabledForFrame;
-	if (!bImageTypeAvailable)
+	if (!bImageTypeAvailable || (GRHIVariableRateShadingImageDataType == VRSImage_NotSupported))
 	{
 		return nullptr;
 	}
@@ -816,7 +817,21 @@ FRDGTextureRef FVariableRateShadingImageManager::GetForceRateImage(FRDGBuilder&
 	FRDGTextureUAVRef ForceShadingRateUAV = GraphBuilder.CreateUAV(ForceShadingRateTexture);
 
 	RateIndex = FMath::Clamp(RateIndex, 0, GetNumberOfSupportedRates() - 1);
-	AddClearUAVPass(GraphBuilder, ForceShadingRateUAV, ValidShadingRates[RateIndex]);
+
+	if (GRHIVariableRateShadingImageDataType == VRSImage_Palette)
+	{
+		AddClearUAVPass(GraphBuilder, ForceShadingRateUAV, ValidShadingRates[RateIndex]);
+	}
+	else if (GRHIVariableRateShadingImageDataType == VRSImage_Fractional)
+	{
+		AddClearUAVPass(GraphBuilder, ForceShadingRateUAV, ValidFragmentDensities[RateIndex]);
+	}
+	else
+	{
+		checkNoEntry(); // Invalid GRHIVariableRateShadingImageDataType
+		return nullptr;
+	}
+	
 
 	return ForceShadingRateTexture;
 }
