[HWRT] Fix CachedMeshCommandHash collision.
- hash needs to include GeometrySegmentIndex as positional encoding of each mesh command.

- also improved how hashes are combined using Murmur64 helper function.

#test Run cooked build of MegaLights demo, toggle r.RayTracing.Enable off/on.
#Synced-CL 50672290

diff --git a/Engine/Source/Runtime/Renderer/Private/PrimitiveSceneInfo.cpp b/Engine/Source/Runtime/Renderer/Private/PrimitiveSceneInfo.cpp
--- a/Engine/Source/Runtime/Renderer/Private/PrimitiveSceneInfo.cpp
+++ b/Engine/Source/Runtime/Renderer/Private/PrimitiveSceneInfo.cpp
@@ -1152,12 +1152,14 @@ int32 CacheRayTracingMeshBatch(
 		uint32 ShaderHash = Shader != nullptr ? GetTypeHash(Shader->GetHash()) : 0;
 		uint32 ShaderBindingsHash = RTMeshCommand.ShaderBindings.GetDynamicInstancingHash();
 
-		// Also add the material shader index to the hash because it's used to deduplicate SBT allocations and the 
-		// material shader index is stored in the user data of the SBT binding data (same shader hash can be moved to another material shader index value)
-		ShaderHash = HashCombine(ShaderHash, RTMeshCommand.MaterialShaderIndex);
+		// Add MaterialShaderIndex to the hash because it's used to deduplicate SBT allocations
+		// and the material shader index is stored in the user data of the SBT binding data (same shader hash can be moved to another material shader index value)
 
-		Hash <<= 1; // TODO: It would probably be better to use some kind of proper 64 bit mix here?
-		Hash ^= (uint64(ShaderBindingsHash) << 32) | uint64(ShaderHash);
+		// Add GeometrySegmentIndex to the hash as a positional encoding of the section that the material is assigned to (otherwise it's possible to get collisions in situations such as:
+		// - Mesh A - 2 segments, segment 0 has 0 triangles (no mesh command generated), segment 1 uses material X
+		// - Mesh B - 2 segments, segment 0 uses material X, segment 1 has 0 triangles (no mesh command generated)
+
+		Hash = Murmur64({ Hash, (uint64(ShaderBindingsHash) << 32) | uint64(ShaderHash), (uint64(RTMeshCommand.MaterialShaderIndex) << 32) | uint64(RTMeshCommand.GeometrySegmentIndex) });
 
 		CommandContext.CommandIndex = -1;
 	}
