[HWRT] Update deferred cached uniform expression before re-caching ray tracing data when switching modes.

#rnx

#Synced-CL 50664013

diff --git a/Engine/Source/Runtime/Renderer/Private/RayTracing/RayTracing.cpp b/Engine/Source/Runtime/Renderer/Private/RayTracing/RayTracing.cpp
--- a/Engine/Source/Runtime/Renderer/Private/RayTracing/RayTracing.cpp
+++ b/Engine/Source/Runtime/Renderer/Private/RayTracing/RayTracing.cpp
@@ -561,6 +561,10 @@ namespace RayTracing
 
 			// In some situations, we need to refresh the cached ray tracing mesh commands because they contain data about the currently bound shader. 
 			// This operation is a bit expensive but only happens once as we transition between RT types which should be rare.
+
+			// If any material has been updated process it before re-caching
+			FMaterialRenderProxy::UpdateDeferredCachedUniformExpressions();
+
 			Scene.CachedRayTracingMeshCommandsType = CurrentType;
 			Scene.RefreshCachedRayTracingData();
 			bUpdateCachedRayTracingState = false;
