Fixed parallel bitfield access causing TSAN error.

#Synced-CL 50642529

diff --git a/Engine/Source/Runtime/Renderer/Private/SceneRendering.h b/Engine/Source/Runtime/Renderer/Private/SceneRendering.h
--- a/Engine/Source/Runtime/Renderer/Private/SceneRendering.h
+++ b/Engine/Source/Runtime/Renderer/Private/SceneRendering.h
@@ -1373,8 +1373,6 @@ public:
 
 	/** true if the view has at least one mesh with a translucent material. */
 	uint32 bHasTranslucentViewMeshElements : 1;
-	/** Indicates whether previous frame transforms were reset this frame for any reason. */
-	uint32 bPrevTransformsReset : 1;
 	/** Whether we should ignore queries from last frame (useful to ignoring occlusions on the first frame after a large camera movement). */
 	uint32 bIgnoreExistingQueries : 1;
 	/** Whether we should submit new queries this frame. (used to disable occlusion queries completely. */
@@ -1498,6 +1496,9 @@ public:
 	// Optional stencil dithering optimization during prepasses
 	bool bAllowStencilDither;
 
+	// Indicates whether previous frame transforms were reset this frame for any reason.
+	bool bPrevTransformsReset = false;
+
 	// Max emissive luminance ouput by any material for this view.
 	float MaterialMaxEmissiveValue;
 
