[SingleLayerWater] Reduced maximum value of r.Water.SingleLayer.RefractionDownsampleFactor from 8 down to 2 as values greater than 2 aren't really supported at the moment.

#Synced-CL 50766943

diff --git a/Engine/Source/Runtime/Renderer/Private/SingleLayerWaterRendering.cpp b/Engine/Source/Runtime/Renderer/Private/SingleLayerWaterRendering.cpp
--- a/Engine/Source/Runtime/Renderer/Private/SingleLayerWaterRendering.cpp
+++ b/Engine/Source/Runtime/Renderer/Private/SingleLayerWaterRendering.cpp
@@ -214,7 +214,7 @@ static int32 GetSingleLayerWaterReflectionTechnique()
 
 static int32 GetSingleLayerWaterRefractionDownsampleFactor()
 {
-	const int32 RefractionDownsampleFactor = FMath::Clamp(CVarWaterSingleLayerRefractionDownsampleFactor.GetValueOnRenderThread(), 1, 8);
+	const int32 RefractionDownsampleFactor = FMath::Clamp(CVarWaterSingleLayerRefractionDownsampleFactor.GetValueOnRenderThread(), 1, 2);
 	return RefractionDownsampleFactor;
 }
 
