[SingleLayerWater] Fixed z-fighting issue when using local tessellation and the water depth prepass. The issue was that the prepass would mark water pixels in stencil and then the base pass would draw water meshes again, testing against both depth and stencil equality, resetting stencil to 0 on pass. The issue is that we force earlyZ in the base pass as an optimization, so even though static water meshes cut out the area handled by the water quadtree, they still do earlyZ and by doing so, zero out stencil if they pass depth and stencil tests (which they do when they're z-fighting). Fixed by clearing stencil in a dedicated pass.

#Synced-CL 50767384

diff --git a/Engine/Source/Runtime/Renderer/Private/SingleLayerWaterRendering.cpp b/Engine/Source/Runtime/Renderer/Private/SingleLayerWaterRendering.cpp
--- a/Engine/Source/Runtime/Renderer/Private/SingleLayerWaterRendering.cpp
+++ b/Engine/Source/Runtime/Renderer/Private/SingleLayerWaterRendering.cpp
@@ -1844,7 +1844,7 @@ void FDeferredShadingSceneRenderer::RenderSingleLayerWaterInner(
 	FDepthStencilBinding DepthStencilBinding;
 	if (bHasDepthPrepass)
 	{
-		DepthStencilBinding = FDepthStencilBinding(SingleLayerWaterPrePassResult->DepthPrepassTexture.Target, ERenderTargetLoadAction::ELoad, ERenderTargetLoadAction::ELoad, FExclusiveDepthStencil::DepthRead_StencilWrite);
+		DepthStencilBinding = FDepthStencilBinding(SingleLayerWaterPrePassResult->DepthPrepassTexture.Target, ERenderTargetLoadAction::ELoad, ERenderTargetLoadAction::ELoad, FExclusiveDepthStencil::DepthRead_StencilRead);
 	}
 	else
 	{
@@ -1913,6 +1913,11 @@ void FDeferredShadingSceneRenderer::RenderSingleLayerWaterInner(
 		}
 	}
 
+	if (bHasDepthPrepass)
+	{
+		AddClearStencilPass(GraphBuilder, SingleLayerWaterPrePassResult->DepthPrepassTexture.Target);
+	}
+
 	if (!bHasDepthPrepass)
 	{
 		AddResolveSceneDepthPass(GraphBuilder, InViews, SceneTextures.Depth);
@@ -2106,7 +2111,7 @@ void FSingleLayerWaterPassMeshProcessor::CollectPSOInitializers(const FSceneText
 		}
 
 		const bool bHasDepthPrepass = IsSingleLayerWaterDepthPrepassEnabled(GetFeatureLevelShaderPlatform(FeatureLevel), FeatureLevel);
-		RenderTargetsInfo.DepthStencilAccess = bHasDepthPrepass ? FExclusiveDepthStencil::DepthRead_StencilWrite : FExclusiveDepthStencil::DepthWrite_StencilNop;
+		RenderTargetsInfo.DepthStencilAccess = bHasDepthPrepass ? FExclusiveDepthStencil::DepthRead_StencilRead : FExclusiveDepthStencil::DepthWrite_StencilNop;
 
 		FBasePassMeshProcessor::AddBasePassGraphicsPipelineStateInitializer(
 			FeatureLevel,
@@ -2132,16 +2137,16 @@ FMeshPassProcessor* CreateSingleLayerWaterPassProcessor(ERHIFeatureLevel::Type F
 	FMeshPassProcessorRenderState DrawRenderState;
 
 	// Make sure depth write is enabled if no prepass is used.
-	FExclusiveDepthStencil::Type BasePassDepthStencilAccess_DepthWrite = FExclusiveDepthStencil::Type(bHasDepthPrepass ? FExclusiveDepthStencil::DepthRead_StencilWrite : SceneBasePassDepthStencilAccess | FExclusiveDepthStencil::DepthWrite_StencilNop);
+	FExclusiveDepthStencil::Type BasePassDepthStencilAccess_DepthWrite = FExclusiveDepthStencil::Type(bHasDepthPrepass ? FExclusiveDepthStencil::DepthRead_StencilRead : SceneBasePassDepthStencilAccess | FExclusiveDepthStencil::DepthWrite_StencilNop);
 	SetupBasePassState(BasePassDepthStencilAccess_DepthWrite, false, DrawRenderState);
 	if (bHasDepthPrepass)
 	{
-		// Set depth stencil test to only pass if depth and stencil are equal to the values written by the prepass and then clear the stencil bit
+		// Set depth stencil test to only pass if depth and stencil are equal to the values written by the prepass
 		DrawRenderState.SetDepthStencilState(TStaticDepthStencilState<
 			false, CF_Equal,							// Depth test
-			true, CF_Equal, SO_Keep, SO_Keep, SO_Zero,	// Front face stencil 
-			true, CF_Equal, SO_Keep, SO_Keep, SO_Zero,	// Back face stencil
-			STENCIL_SANDBOX_MASK, STENCIL_SANDBOX_MASK 	// Stencil read/write masks
+			true, CF_Equal, SO_Keep, SO_Keep, SO_Keep,	// Front face stencil 
+			true, CF_Equal, SO_Keep, SO_Keep, SO_Keep,	// Back face stencil
+			STENCIL_SANDBOX_MASK, 0x0 					// Stencil read/write masks
 		>::GetRHI());
 		DrawRenderState.SetStencilRef(STENCIL_SANDBOX_MASK);
 	}
