Subsurface - fixed missing albedos when lighting is disabled (e.g. unlit view or lighting show flag set to 0)

#Synced-CL 50739550

diff --git a/Engine/Source/Runtime/Renderer/Private/PostProcess/PostProcessSubsurface.cpp b/Engine/Source/Runtime/Renderer/Private/PostProcess/PostProcessSubsurface.cpp
--- a/Engine/Source/Runtime/Renderer/Private/PostProcess/PostProcessSubsurface.cpp
+++ b/Engine/Source/Runtime/Renderer/Private/PostProcess/PostProcessSubsurface.cpp
@@ -300,9 +300,9 @@ FSubsurfaceSRVInput GetSubsurfaceSRVInput(FRDGTextureSRVRef Texture, const FScre
 	return Input;
 }
 
-bool IsSubsurfaceEnabled()
+bool IsSubsurfaceEnabled(const FEngineShowFlags& EngineShowFlags)
 {
-	const bool bEnabled = CVarSubsurfaceScattering.GetValueOnAnyThread() != 0;
+	const bool bEnabled = CVarSubsurfaceScattering.GetValueOnAnyThread() != 0 && EngineShowFlags.Lighting;
 	const bool bHasScale = CVarSSSScale.GetValueOnAnyThread() > 0.0f;
 	return (bEnabled && bHasScale);
 }
@@ -310,7 +310,7 @@ bool IsSubsurfaceEnabled()
 bool IsSubsurfaceRequiredForView(const FViewInfo& View)
 {
 	const bool bSimpleDynamicLighting = IsForwardShadingEnabled(View.GetShaderPlatform());
-	const bool bSubsurfaceEnabled = IsSubsurfaceEnabled();
+	const bool bSubsurfaceEnabled = IsSubsurfaceEnabled(View.Family->EngineShowFlags);
 	const bool bViewHasSubsurfaceMaterials = ((View.ShadingModelMaskInView & GetUseSubsurfaceProfileShadingModelMask()) != 0);
 	return (bSubsurfaceEnabled && bViewHasSubsurfaceMaterials && !bSimpleDynamicLighting);
 }
diff --git a/Engine/Source/Runtime/Renderer/Private/PostProcess/PostProcessSubsurface.h b/Engine/Source/Runtime/Renderer/Private/PostProcess/PostProcessSubsurface.h
--- a/Engine/Source/Runtime/Renderer/Private/PostProcess/PostProcessSubsurface.h
+++ b/Engine/Source/Runtime/Renderer/Private/PostProcess/PostProcessSubsurface.h
@@ -6,7 +6,7 @@
 #include "OverridePassSequence.h"
 
 // Returns whether subsurface scattering is globally enabled.
-bool IsSubsurfaceEnabled();
+bool IsSubsurfaceEnabled(const FEngineShowFlags& EngineShowFlags);
 
 // Returns whether subsurface scattering is required for the provided view.
 bool IsSubsurfaceRequiredForView(const FViewInfo& View);
diff --git a/Engine/Source/Runtime/Renderer/Private/SceneRendering.cpp b/Engine/Source/Runtime/Renderer/Private/SceneRendering.cpp
--- a/Engine/Source/Runtime/Renderer/Private/SceneRendering.cpp
+++ b/Engine/Source/Runtime/Renderer/Private/SceneRendering.cpp
@@ -1829,7 +1829,7 @@ void FViewInfo::SetupUniformBufferParameters(
 
 	// Subsurface
 	{
-		ViewUniformShaderParameters.bSubsurfacePostprocessEnabled = IsSubsurfaceEnabled() ? 1.0f : 0.0f;
+		ViewUniformShaderParameters.bSubsurfacePostprocessEnabled = IsSubsurfaceEnabled(Family->EngineShowFlags) ? 1.0f : 0.0f;
 
 		// Subsurface shading model
 		{
