[HWRT] Reduce number of FRayTracingBuildInstanceBufferCS permutations by ~50%.
- Mostly because FUseSeparateHitContributionBuffer is only necessary on Metal RHI.
- Also avoid Nanite Ray Tracing on platforms that don't support Nanite or when permutation doesn't use GPUScene.

#Synced-CL 50761898

diff --git a/Engine/Source/Runtime/Renderer/Private/RayTracing/RayTracingInstanceBufferUtil.cpp b/Engine/Source/Runtime/Renderer/Private/RayTracing/RayTracingInstanceBufferUtil.cpp
--- a/Engine/Source/Runtime/Renderer/Private/RayTracing/RayTracingInstanceBufferUtil.cpp
+++ b/Engine/Source/Runtime/Renderer/Private/RayTracing/RayTracingInstanceBufferUtil.cpp
@@ -611,9 +611,23 @@ struct FRayTracingBuildInstanceBufferCS : public FGlobalShader
 		{
 			return false;
 		}
+
+		// TODO: Once we update FRayTracingInstanceBufferBuilder to use FRDGBuilder
+		// we could replace this permutation with a dynamic branch and simply bind a dummy buffer when none is provided
+		if (PermutationVector.Get<FUseSeparateHitContributionBuffer>() != IsMetalPlatform(Parameters.Platform))
+		{
+			return false;
+		}
 		
-		if (PermutationVector.Get<FNaniteRayTracingDim>() && !PermutationVector.Get<FUseGPUSceneDim>())
+		if (PermutationVector.Get<FNaniteRayTracingDim>() && (!PermutationVector.Get<FUseGPUSceneDim>() || !DoesPlatformSupportNanite(Parameters.Platform)))
+		{
+			// Nanite RayTracing requires GPUScene to access PrimitiveData.NaniteResourceID
+			return false;
+		}
+
+		if (PermutationVector.Get<FGpuCullingDim>() && !PermutationVector.Get<FUseGPUSceneDim>())
 		{
+			// GPU culling requires GPUScene
 			return false;
 		}
 
@@ -624,7 +638,7 @@ struct FRayTracingBuildInstanceBufferCS : public FGlobalShader
 	{
 		FPermutationDomain PermutationVector(Parameters.PermutationId);
 
-		if (PermutationVector.Get<FUseWaveOpsDim>() != GRHISupportsWaveOperations)
+		if (PermutationVector.Get<FUseWaveOpsDim>() != GRHIGlobals.SupportsWaveOperations)
 		{
 			return EShaderPermutationPrecacheRequest::NotPrecached;
 		}
@@ -709,7 +723,7 @@ void BuildRayTracingInstanceBuffer(
 	PermutationVector.Set<FRayTracingBuildInstanceBufferCS::FOutputInstanceExtraDataDim>(InstanceExtraDataUAV != nullptr);
 	PermutationVector.Set<FRayTracingBuildInstanceBufferCS::FOutputStatsDim>(OutputStatsUAV != nullptr);
 	PermutationVector.Set<FRayTracingBuildInstanceBufferCS::FCompactOutputDim>(bCompactOutput);
-	PermutationVector.Set<FRayTracingBuildInstanceBufferCS::FUseWaveOpsDim>(GRHISupportsWaveOperations);
+	PermutationVector.Set<FRayTracingBuildInstanceBufferCS::FUseWaveOpsDim>(GRHIGlobals.SupportsWaveOperations);
 	PermutationVector.Set<FRayTracingBuildInstanceBufferCS::FGpuCullingDim>(CullingParameters != nullptr);
 	PermutationVector.Set<FRayTracingBuildInstanceBufferCS::FUseSeparateHitContributionBuffer>(GRHIGlobals.RayTracing.RequiresSeparateHitGroupContributionsBuffer);
 
