Moved bNeedsCachedReflectionCaptureUpdate to its own byte to avoid data racing with other bits.

#Synced-CL 50783182

diff --git a/Engine/Source/Runtime/Renderer/Private/PrimitiveSceneInfo.cpp b/Engine/Source/Runtime/Renderer/Private/PrimitiveSceneInfo.cpp
--- a/Engine/Source/Runtime/Renderer/Private/PrimitiveSceneInfo.cpp
+++ b/Engine/Source/Runtime/Renderer/Private/PrimitiveSceneInfo.cpp
@@ -307,7 +307,6 @@ FPrimitiveSceneInfo::FPrimitiveSceneInfo(const FPrimitiveSceneInfoAdapter& InAda
 	bPendingAddToScene(false),
 	bPendingAddStaticMeshes(false),
 	bPendingFlushRuntimeVirtualTexture(false),
-	bNeedsCachedReflectionCaptureUpdate(true),
 	bShouldRenderInMainPass(InAdapter.SceneProxy->ShouldRenderInMainPass()),
 	bVisibleInRealTimeSkyCapture(InAdapter.SceneProxy->IsVisibleInRealTimeSkyCaptures()),
 	bWritesRuntimeVirtualTexture(InAdapter.SceneProxy->WritesVirtualTexture()),
diff --git a/Engine/Source/Runtime/Renderer/Public/PrimitiveSceneInfo.h b/Engine/Source/Runtime/Renderer/Public/PrimitiveSceneInfo.h
--- a/Engine/Source/Runtime/Renderer/Public/PrimitiveSceneInfo.h
+++ b/Engine/Source/Runtime/Renderer/Public/PrimitiveSceneInfo.h
@@ -633,9 +633,6 @@ private:
 	bool bPendingFlushRuntimeVirtualTexture : 1;
 
 public:
-	/** Whether the primitive is newly registered or moved and CachedReflectionCaptureProxy needs to be updated on the next render. */
-	bool bNeedsCachedReflectionCaptureUpdate : 1;
-
 	/** Set to true for the primitive to be rendered in the main pass to be visible in a view. */
 	bool bShouldRenderInMainPass : 1;
 
@@ -645,6 +642,9 @@ public:
 	/** True if the primitive can write to runtime virtual texture. */
 	bool bWritesRuntimeVirtualTexture : 1;
 
+	/** Whether the primitive is newly registered or moved and CachedReflectionCaptureProxy needs to be updated on the next render. Not a bitfield to avoid parallel access with other bits. */
+	bool bNeedsCachedReflectionCaptureUpdate = true;
+
 #if RHI_RAYTRACING
 	bool bIsRelevantToPathTracing : 1;
 	bool bIsCachedRayTracingInstanceValid : 1;
