Use local variable instead of re-assigning Scene's World member.

#Synced-CL 50783182

diff --git a/Engine/Source/Runtime/Renderer/Private/ReflectionEnvironmentCapture.cpp b/Engine/Source/Runtime/Renderer/Private/ReflectionEnvironmentCapture.cpp
--- a/Engine/Source/Runtime/Renderer/Private/ReflectionEnvironmentCapture.cpp
+++ b/Engine/Source/Runtime/Renderer/Private/ReflectionEnvironmentCapture.cpp
@@ -2017,12 +2017,10 @@ void FScene::UpdateSkyCaptureContents(
 	{
 		QUICK_SCOPE_CYCLE_COUNTER(STAT_UpdateSkyCaptureContents);
 
-		World = GetWorld();
-
-		if (World)
+		if (UWorld* CurrentWorld = GetWorld())
 		{
 			//guarantee that all render proxies are up to date before kicking off this render
-			World->SendAllEndOfFrameUpdates();
+			CurrentWorld->SendAllEndOfFrameUpdates();
 		}
 
 		const int32 CubemapResolution = CaptureComponent->CubemapResolution;
