[TSR] Thin geometry stability improvement.
1. Add a new shading range = 3 and as default. When temporal analysis is on, run a tailored  thin geo detection on non-foliage materials to improve line stability.
2. Fix a bug where disabling thin geometry antiflickeirng, the whole logic is disabled in reject shading.

#tests TSR level, nanitetest custom levels.

#Synced-CL 50779174

diff --git a/Engine/Shaders/Private/TemporalSuperResolution/TSRDetectThinGeometry.usf b/Engine/Shaders/Private/TemporalSuperResolution/TSRDetectThinGeometry.usf
--- a/Engine/Shaders/Private/TemporalSuperResolution/TSRDetectThinGeometry.usf
+++ b/Engine/Shaders/Private/TemporalSuperResolution/TSRDetectThinGeometry.usf
@@ -22,6 +22,7 @@
 
 #define CONFIG_THIN_GEOMETRY_EDGE_REPROJECTION (DIM_THIN_GEOMETRY_EDGE_REPROJECTION)
 #define CONFIG_SKY_RELAXATION (DIM_SKY_RELAXATION)
+#define CONFIG_VARYING_RANGE_DETECTION (DIM_VARYING_RANGE_DETECTION)
 //------------------------------------------------------- CONSTANTS
 
 #if CONFIG_LDS_DESPILL
@@ -96,6 +97,9 @@ void FetchCommonInput(
 	tsr_tensor_short2 PixelPos
 	,out tsr_tensor_half CurrentCoverage
 	,out tsr_tensor_bool bHasPixelAnimation
+#if CONFIG_VARYING_RANGE_DETECTION
+	,out tsr_tensor_ushort MaterialMask
+#endif
 	)
 {
 	tsr_tensor_bool bHasReprojectionOffset;
@@ -130,8 +134,13 @@ void FetchCommonInput(
 #else
 	CurrentCoverage = select(bHasReprojectionOffset, CurrentCoverage, tsr_tensor_half::Const(0));
 #endif
-	
+
+#if CONFIG_VARYING_RANGE_DETECTION
+	MaterialMask = tsr_tensor_ushort::CastFrom<tsr_half>(CurrentCoverage * tsr_half(2.01f));
+	CurrentCoverage = ceil(CurrentCoverage);
+#else
 	CurrentCoverage = ToBinaryCoverage(CurrentCoverage);
+#endif
 }// FetchCommonInput
 
 void FetchRegularPassInput(
@@ -155,9 +164,14 @@ void FetchRegularPassInput(
 	{
 		float Depth = DepthCache.GetElement(ElementIndex);
 		float DeviceZError = ComputePixelDeviceZError(Depth);
-
+#if		CONFIG_VARYING_RANGE_DETECTION
+		DeviceZError *= max(ErrorMultiplier * 0.05f, 1.0f);
+		
+#else
+		DeviceZError *= ErrorMultiplier;
+#endif
 		SceneDepth.SetElement(ElementIndex, tsr_half(Depth));
-		DeviceZThinGeometryDetectionError.SetElement(ElementIndex, tsr_half(ErrorMultiplier * DeviceZError));
+		DeviceZThinGeometryDetectionError.SetElement(ElementIndex, tsr_half(DeviceZError));
 	}
 } // FetchRegularPassInput
 
@@ -301,10 +315,18 @@ void DetectThinGeometryCS(
 	tsr_tensor_half ReprojectedHistoryCoverage;
 	tsr_tensor_bool bHasPixelAnimation;
 	
+	#if CONFIG_VARYING_RANGE_DETECTION
+	tsr_tensor_ushort MaterialMask;
+	#endif
+	
 	// Fetch data for tiling optimization.
 	FetchCommonInput(PixelPos, 
-		CurrentCoverage,
-		bHasPixelAnimation);
+		CurrentCoverage
+		,bHasPixelAnimation
+	#if CONFIG_VARYING_RANGE_DETECTION
+		,MaterialMask
+	#endif
+	);
 		
 	bool bRegularPass = AnyThread(AnyComponent(and(CurrentCoverage > tsr_tensor_half::Const(0.0), not(bHasPixelAnimation))));
 	tsr_tensor_half ThinGeometryEdge = tsr_tensor_half::Const(0.0);
@@ -388,11 +410,29 @@ void DetectThinGeometryCS(
 			tsr_tensor_bool bValidHistoryRelaxation = and(bThinGeometryHistoryRegion,and(ValidCurrentThinGeometryRegion, bSameDistribution));
 			HistoryRelaxationWeight = select(bValidHistoryRelaxation, HistoryRelaxationWeight, tsr_tensor_half::Const(0));
 			
+			#if CONFIG_VARYING_RANGE_DETECTION
+				tsr_tensor_half NonFoliageThinGeometryEdge = ThinGeometryEdge;
+				tsr_tensor_half NonFoliageRelaxationWeight = HistoryRelaxationWeight;
+				bool bAnyFoliages = AnyThread(AnyComponent(and(MaterialMask == tsr_tensor_ushort::Const(2), not(bHasPixelAnimation))));
+			if (bAnyFoliages)
+			#endif
 			//1. Diffuse a little bit to suppress the history clamping around the valid pixel and apply max relaxation weight.
 			//2. Dilate the edge with (0.8,1,0.8), such that, dilated edges has 0.8 for direct neighbor, 0.64 for diagnal, 1 for consecutive neighbors.
-			Deconcatenate(WeightedAvg3x3(Concatenate(ThinGeometryEdge,HistoryRelaxationWeight),float3(0.8, 1.0, 0.8)),ThinGeometryEdge,HistoryRelaxationWeight);
-			HistoryRelaxationWeight = saturate(WeightedAvg3x3(saturate(HistoryRelaxationWeight * tsr_half(MaxRelaxationWeight)), float3(0.5, 1.0, 0.5)));
-
+			{
+				Deconcatenate(WeightedAvg3x3(Concatenate(ThinGeometryEdge, HistoryRelaxationWeight), float3(0.8, 1.0, 0.8)), ThinGeometryEdge, HistoryRelaxationWeight);
+				HistoryRelaxationWeight = saturate(WeightedAvg3x3(saturate(HistoryRelaxationWeight * tsr_half(MaxRelaxationWeight)), float3(0.5, 1.0, 0.5)));
+				
+				#if CONFIG_VARYING_RANGE_DETECTION
+
+				// Regular non-foliage materials does not require extensive relaxation when temporal flickering logic is active.
+				{
+					tsr_tensor_bool bNonFoliages = MaterialMask == tsr_tensor_ushort::Const(1);
+
+					ThinGeometryEdge = select(bNonFoliages, NonFoliageThinGeometryEdge, ThinGeometryEdge);
+					HistoryRelaxationWeight = select(bNonFoliages, NonFoliageRelaxationWeight, HistoryRelaxationWeight);
+				}
+				#endif
+			}
 			// Not similar patches or has pixel animation should reset relaxation weight to 0
 			// to use more strict history clamping.
 			HistoryRelaxationWeight = select(and(bSameDistribution, not(bHasPixelAnimation)), HistoryRelaxationWeight, tsr_tensor_half::Const(0));
diff --git a/Engine/Shaders/Private/TemporalSuperResolution/TSRRejectShading.usf b/Engine/Shaders/Private/TemporalSuperResolution/TSRRejectShading.usf
--- a/Engine/Shaders/Private/TemporalSuperResolution/TSRRejectShading.usf
+++ b/Engine/Shaders/Private/TemporalSuperResolution/TSRRejectShading.usf
@@ -648,10 +648,16 @@ void RejectShadingCS(
 		tsr_tensor_bool IsThinGeometryRegion;
 		#if CONFIG_MOIRE_DETECTION
 			BRANCH
-			if (bApplyThinGeometryRelaxation)
+			if (bApplyThinGeometryRelaxation || !ThinGeometryAffectMoireError)
 			{
 				FetchThinGeometryEdge(PixelPos, ThinGeometryRelaxationFactor, IsThinGeometryRegion);
 			}
+	
+			BRANCH
+			if (!ThinGeometryAffectMoireError)
+			{
+				bApplyThinGeometryRelaxation = AnyThread(AnyComponent(ThinGeometryRelaxationFactor > tsr_tensor_half::Const(0)));
+			}
 		#else
 			FetchThinGeometryEdge(PixelPos, ThinGeometryRelaxationFactor, IsThinGeometryRegion);
 			bApplyThinGeometryRelaxation = AnyThread(AnyComponent(ThinGeometryRelaxationFactor > tsr_tensor_half::Const(0)));
diff --git a/Engine/Source/Runtime/Renderer/Private/PostProcess/TemporalSuperResolution.cpp b/Engine/Source/Runtime/Renderer/Private/PostProcess/TemporalSuperResolution.cpp
--- a/Engine/Source/Runtime/Renderer/Private/PostProcess/TemporalSuperResolution.cpp
+++ b/Engine/Source/Runtime/Renderer/Private/PostProcess/TemporalSuperResolution.cpp
@@ -345,11 +345,12 @@ TAutoConsoleVariable<bool> CVarTSRThinGeometryCoverageEdgeReprojection(
 	ECVF_RenderThreadSafe);
 
 TAutoConsoleVariable<int32> CVarTSRThinGeometryCoverageShadingRange(
-	TEXT("r.TSR.ThinGeometryDetection.Coverage.ShadingRange"), 1,
+	TEXT("r.TSR.ThinGeometryDetection.Coverage.ShadingRange"), 3,
 	TEXT("The shading model range regarded as thin geometry and accumulate the coverage.\n")
 	TEXT("0: Two-sided foliage only.\n")
-	TEXT("1: Two-sided foliage and hair.(Default)\n")
-	TEXT("2: All shading model except unlit.\n"),
+	TEXT("1: Two-sided foliage and hair.\n")
+	TEXT("2: All shading model except unlit.\n")
+	TEXT("3: Regular non-foliage applies lightweight thin geo logic when `r.TSR.ShadingRejection.Flickering` is active. Otherwise fallback to 1 (Default).\n"),
 	ECVF_RenderThreadSafe);
 
 TAutoConsoleVariable<float> CVarTSRThinGeometryCoverageMaxRelaxationWeight(
@@ -741,38 +742,43 @@ class FTSRMeasureFlickeringLumaCS : public FTSRShader
 	END_SHADER_PARAMETER_STRUCT()
 }; // class FTSRMeasureFlickeringLumaCS
 
+enum class EThinGeometryShadingRange: uint32
+{
+	Foliage,
+	FoliageAndHair,
+	All,
+	VaryingRange,
+	MAX,
+};
+
+static EThinGeometryShadingRange GetShadingRange()
+{
+	int32 ShadingRange = CVarTSRThinGeometryCoverageShadingRange.GetValueOnRenderThread();
+
+	return static_cast<EThinGeometryShadingRange>(FMath::Clamp(ShadingRange,
+		static_cast<int32>(EThinGeometryShadingRange::Foliage), 
+		static_cast<int32>(EThinGeometryShadingRange::MAX)-1));
+}
+
+static const TCHAR* GetShadingRangeName(EThinGeometryShadingRange ShadingRange)
+{
+	static const TCHAR* const kShadingRangeNames[] = {
+		TEXT("Foliage"),
+		TEXT("Foliage&Hair"),
+		TEXT("All"),
+		TEXT("VaryingRange")
+	};
+	static_assert(UE_ARRAY_COUNT(kShadingRangeNames) == int32(EThinGeometryShadingRange::MAX), "Fix me");
+	return kShadingRangeNames[int32(ShadingRange)];
+}
+
 class FTSRMeasureThinGeometryCoverageCS : public FTSRShader
 {
 	DECLARE_GLOBAL_SHADER(FTSRMeasureThinGeometryCoverageCS);
 	SHADER_USE_PARAMETER_STRUCT(FTSRMeasureThinGeometryCoverageCS, FTSRShader);
 
-	enum class EThinGeometryShadingRange: uint32
-	{
-		Foliage,
-		FoliageAndHair,
-		All,
-		MAX
-	};
 
-	static EThinGeometryShadingRange GetShadingRange()
-	{
-		int32 ShadingRange = CVarTSRThinGeometryCoverageShadingRange.GetValueOnRenderThread();
-		
-		return static_cast<EThinGeometryShadingRange>(FMath::Clamp(ShadingRange,
-			static_cast<int32>(EThinGeometryShadingRange::Foliage), 
-			static_cast<int32>(EThinGeometryShadingRange::MAX)-1));
-	}
 
-	static const TCHAR* GetShadingRangeName(EThinGeometryShadingRange ShadingRange)
-	{
-		static const TCHAR* const kShadingRangeNames[] = {
-			TEXT("Foliage"),
-			TEXT("Foliage&Hair"),
-			TEXT("All")
-		};
-		static_assert(UE_ARRAY_COUNT(kShadingRangeNames) == int32(EThinGeometryShadingRange::MAX), "Fix me");
-		return kShadingRangeNames[int32(ShadingRange)];
-	}
 
 	class FThinGeometryShadingRangeDim : SHADER_PERMUTATION_ENUM_CLASS("DIM_THIN_GEOMETRY_SHADING_RANGE",EThinGeometryShadingRange);
 	using FPermutationDomain = TShaderPermutationDomain<FThinGeometryShadingRangeDim>;
@@ -1066,7 +1072,8 @@ class FTSRDetectThinGeometryCS : public FTSRConvolutionNetworkShader
 
 	class FSkyRelaxationDim : SHADER_PERMUTATION_BOOL("DIM_SKY_RELAXATION");
 	class FThinGeometryEdgeReprojectionDim : SHADER_PERMUTATION_BOOL("DIM_THIN_GEOMETRY_EDGE_REPROJECTION");
-	using FPermutationDomain = TShaderPermutationDomain<FTSRConvolutionNetworkShader::FPermutationDomain,FThinGeometryEdgeReprojectionDim, FSkyRelaxationDim>;
+	class FVaryingRangeDetectionDim : SHADER_PERMUTATION_BOOL("DIM_VARYING_RANGE_DETECTION");
+	using FPermutationDomain = TShaderPermutationDomain<FTSRConvolutionNetworkShader::FPermutationDomain, FThinGeometryEdgeReprojectionDim, FSkyRelaxationDim, FVaryingRangeDetectionDim>;
 
 	BEGIN_SHADER_PARAMETER_STRUCT(FParameters, )
 		SHADER_PARAMETER_STRUCT_INCLUDE(FTSRCommonParameters, CommonParameters)
@@ -1765,7 +1772,7 @@ void AddTSRMeasureThinGeometryCoverage(FRDGBuilder& GraphBuilder, FGlobalShaderM
 	
 	PassParameters->ThinGeometryCoverageOutput = GraphBuilder.CreateUAV(ThinGeometryCoverage.Texture);
 
-	FTSRMeasureThinGeometryCoverageCS::EThinGeometryShadingRange ShadingRange = FTSRMeasureThinGeometryCoverageCS::GetShadingRange();
+	EThinGeometryShadingRange ShadingRange = GetShadingRange();
 	FTSRMeasureThinGeometryCoverageCS::FPermutationDomain PermutationVector;
 	PermutationVector.Set<FTSRMeasureThinGeometryCoverageCS::FThinGeometryShadingRangeDim>(ShadingRange);
 
@@ -1777,7 +1784,7 @@ void AddTSRMeasureThinGeometryCoverage(FRDGBuilder& GraphBuilder, FGlobalShaderM
 		GraphBuilder,
 		RDG_EVENT_NAME("TSR MeasureThinGeometryCoverage(#%d%s) %dx%d", 
 			PermutationVector.ToDimensionValueId(),
-			FTSRMeasureThinGeometryCoverageCS::GetShadingRangeName(ShadingRange),
+			GetShadingRangeName(ShadingRange),
 			ViewRect.Width(), 
 			ViewRect.Height()),
 		AsyncComputePasses >= 2 ? ERDGPassFlags::AsyncCompute : ERDGPassFlags::Compute,
@@ -2701,6 +2708,7 @@ FDefaultTemporalUpscaler::FOutputs AddTemporalSuperResolutionPasses(
 	//Create the thin geometry mask to avoid over rejection
 	if (PassConfig.ThinGeometryDetectionEnable)
 	{
+		const bool EnableVaryingRangeDetection = (GetShadingRange() == EThinGeometryShadingRange::VaryingRange) && (FlickeringFramePeriod > 0.0f);
 		{
 			FTSRConvolutionNetworkShader::FPermutationDomain ConvolutionNetworkPermutationVector;
 			ConvolutionNetworkPermutationVector.Set<FTSRConvolutionNetworkShader::FWaveSizeOps>(SelectWaveSize(View.GetShaderPlatform(), { 16, 32, 64 }));
@@ -2709,7 +2717,8 @@ FDefaultTemporalUpscaler::FOutputs AddTemporalSuperResolutionPasses(
 
 			FTSRDetectThinGeometryCS::FPermutationDomain PermutationVector;
 			PermutationVector.Set<FTSRDetectThinGeometryCS::FSkyRelaxationDim>(ShouldApplySkyRelaxation());
-            PermutationVector.Set<FTSRDetectThinGeometryCS::FThinGeometryEdgeReprojectionDim>(CVarTSRThinGeometryCoverageEdgeReprojection.GetValueOnRenderThread());
+            	PermutationVector.Set<FTSRDetectThinGeometryCS::FThinGeometryEdgeReprojectionDim>(CVarTSRThinGeometryCoverageEdgeReprojection.GetValueOnRenderThread());
+			PermutationVector.Set<FTSRDetectThinGeometryCS::FVaryingRangeDetectionDim>(EnableVaryingRangeDetection);
 			PermutationVector.Set<FTSRConvolutionNetworkShader::FPermutationDomain>(ConvolutionNetworkPermutationVector);
 			PermutationVector = FTSRDetectThinGeometryCS::RemapPermutation(PermutationVector);
 
