[Front Layer Translucency] Fix meshes not rendering to FrontLayerTranslucencyGBuffer when Lumen Front Layer Reflections are disabled.
- ShouldRenderFrontLayerTranslucency(...) was not taking IsVSMTranslucentHighQualityEnabled() into account.
- Also added GPU stat for FrontLayerTranslucencyGBuffer.

#Synced-CL 50789286

diff --git a/Engine/Source/Runtime/Renderer/Private/FrontLayerTranslucency.cpp b/Engine/Source/Runtime/Renderer/Private/FrontLayerTranslucency.cpp
--- a/Engine/Source/Runtime/Renderer/Private/FrontLayerTranslucency.cpp
+++ b/Engine/Source/Runtime/Renderer/Private/FrontLayerTranslucency.cpp
@@ -15,6 +15,8 @@
 #include "Lumen/RayTracedTranslucency.h"
 #include "TextureFallbacksRDG.h"
 
+DECLARE_GPU_DRAWCALL_STAT(FrontLayerTranslucencyGBuffer);
+
 // Whether to enable Front Layer Translucency reflections from scalability
 int32 GLumenFrontLayerTranslucencyReflectionsEnabled = 0;
 FAutoConsoleVariableRef CVarLumenTranslucencyReflectionsFrontLayerEnabled(
@@ -78,8 +80,10 @@ namespace
 
 bool ShouldRenderFrontLayerTranslucency(const FViewInfo& View)
 {
-	return ((ShouldRenderLumenReflections(View) && Lumen::UseFrontLayerTranslucencyReflections(View)) 
-			|| RayTracedTranslucency::IsEnabled(View));
+	return ((ShouldRenderLumenReflections(View) && Lumen::UseFrontLayerTranslucencyReflections(View))
+		|| RayTracedTranslucency::IsEnabled(View))
+		// TODO: Should check if VSM is actually enabled for this view
+		|| IsVSMTranslucentHighQualityEnabled();
 }
 
 class FFrontLayerTranslucencyClearGBufferPS : public FGlobalShader
@@ -485,7 +489,7 @@ FFrontLayerTranslucencyData FDeferredShadingSceneRenderer::RenderFrontLayerTrans
 	FFrontLayerTranslucencyData Out;
 	if (bNeedFrontLayerData > 0)
 	{
-		RDG_EVENT_SCOPE(GraphBuilder, "FrontLayerTranslucencyGBuffer");
+		RDG_EVENT_SCOPE_STAT(GraphBuilder, FrontLayerTranslucencyGBuffer, "FrontLayerTranslucencyGBuffer");
 
 		// Allocate resources for all views
 		{
