Remove hit proxy and water selection shaders from cooked PC mobile platform cooked builds.
Could remove from all PC cooked builds if we know that we don't need to support cooked editor.
#rnx
#Synced-CL 50812607

diff --git a/Engine/Plugins/Experimental/Water/Source/Runtime/Public/WaterVertexFactory.inl b/Engine/Plugins/Experimental/Water/Source/Runtime/Public/WaterVertexFactory.inl
--- a/Engine/Plugins/Experimental/Water/Source/Runtime/Public/WaterVertexFactory.inl
+++ b/Engine/Plugins/Experimental/Water/Source/Runtime/Public/WaterVertexFactory.inl
@@ -176,8 +176,12 @@ bool TWaterVertexFactory<bWithWaterSelectionSupport, DrawMode>::ShouldCompilePer
 				return false;
 			}
 		}
-		// Only let the PC platform compile the permutations supporting selection : 
-		return (!bWithWaterSelectionSupport || IsPCPlatform(Parameters.Platform));
+
+		// Only let the editor platforms compile the permutations supporting selection
+		const bool bSupportEditor = EnumHasAllFlags(Parameters.Flags, EShaderPermutationFlags::HasEditorOnlyData);
+		const bool bSupportCookedEditor = IsPCPlatform(Parameters.Platform) && !IsMobilePlatform(Parameters.Platform);
+
+		return !bWithWaterSelectionSupport || bSupportEditor || bSupportCookedEditor;
 	}
 	return false;
 }
diff --git a/Engine/Source/Runtime/Renderer/Private/SceneHitProxyRendering.cpp b/Engine/Source/Runtime/Renderer/Private/SceneHitProxyRendering.cpp
--- a/Engine/Source/Runtime/Renderer/Private/SceneHitProxyRendering.cpp
+++ b/Engine/Source/Runtime/Renderer/Private/SceneHitProxyRendering.cpp
@@ -40,6 +40,13 @@
 #include "GPUSkinCache.h"
 #include "MaterialCache/MaterialCacheTagProvider.h"
 
+static bool RequiresEditorPermutations(const FMeshMaterialShaderPermutationParameters& Parameters)
+{
+	const bool bSupportEditor = EnumHasAllFlags(Parameters.Flags, EShaderPermutationFlags::HasEditorOnlyData);
+	const bool bSupportCookedEditor = IsPCPlatform(Parameters.Platform) && !IsMobilePlatform(Parameters.Platform);
+	return bSupportEditor || bSupportCookedEditor;
+}
+
 class FHitProxyShaderElementData : public FMeshMaterialShaderElementData
 {
 public:
@@ -61,8 +68,9 @@ class FHitProxyVS : public FMeshMaterialShader
 public:
 	static bool ShouldCompilePermutation(const FMeshMaterialShaderPermutationParameters& Parameters)
 	{
-		// Only compile the hit proxy vertex shader on desktop editor platforms
-		return IsPCPlatform(Parameters.Platform)// && EnumHasAllFlags(Parameters.Flags, EShaderPermutationFlags::HasEditorOnlyData)
+
+		// Only compile the hit proxy vertex shader on desktop platforms
+		return RequiresEditorPermutations(Parameters)
 			// and only compile for the default material or materials that are masked.
 			&& (Parameters.MaterialParameters.bIsSpecialEngineMaterial ||
 				!Parameters.MaterialParameters.bWritesEveryPixel ||
@@ -117,8 +125,8 @@ public:
 
 	static bool ShouldCompilePermutation(const FMeshMaterialShaderPermutationParameters& Parameters)
 	{
-		// Only compile the hit proxy vertex shader on desktop editor platforms
-		return IsPCPlatform(Parameters.Platform)// && EnumHasAllFlags(Parameters.Flags, EShaderPermutationFlags::HasEditorOnlyData)
+		// Only compile the hit proxy vertex shader on editor platforms
+		return RequiresEditorPermutations(Parameters)
 			// and only compile for default materials or materials that are masked.
 			&& (Parameters.MaterialParameters.bIsSpecialEngineMaterial ||
 				!Parameters.MaterialParameters.bWritesEveryPixel ||
