Skip copying members from FViewInfo that are unused by the shadow view and are triggering TSAN errors.
- View snapshots are just memcpy's of the main view. We have several systems that use them, shadows being one of them.
- The shadow views are created on task threads. Shadow views don't use 90% of the struct, but we don't have a slim interface to pass data through, so FViewInfo is it for now.
- TSAN reports that the render thread is mutating members of the FViewInfo from several locations that shadow views don't care about. This change adds an enum to specify which kind of view snapshot is being created. Shadow views skip copying the offending data.

#Synced-CL 50790641

diff --git a/Engine/Source/Runtime/Renderer/Private/Renderer/ViewSnapshotCache.cpp b/Engine/Source/Runtime/Renderer/Private/Renderer/ViewSnapshotCache.cpp
--- a/Engine/Source/Runtime/Renderer/Private/Renderer/ViewSnapshotCache.cpp
+++ b/Engine/Source/Runtime/Renderer/Private/Renderer/ViewSnapshotCache.cpp
@@ -8,7 +8,7 @@ class FViewSnapshotCache
 public:
 	FViewSnapshotCache() = default;
 
-	FViewInfo* Create(const FViewInfo* InView)
+	FViewInfo* Create(const FViewInfo* InView, EViewSnapshotType Type)
 	{
 		TRACE_CPUPROFILER_EVENT_SCOPE(ViewSnapshotCache::Create);
 		check(IsInParallelRenderingThread()); // we do not want this popped before the end of the scene and it better be the scene allocator
@@ -27,7 +27,55 @@ public:
 		Snapshots.Push(Result);
 		NumSnapshots++;
 
-		FMemory::Memcpy(*Result, *InView);
+		if (Type == EViewSnapshotType::Main)
+		{
+			FMemory::Memcpy(*Result, *InView);
+
+			static_assert(std::is_trivially_destructible_v<FGPUScenePrimitiveCollector> != 0, "The destructor is not invoked properly because of FMemory::Memcpy above");
+			Result->DynamicPrimitiveCollector = FGPUScenePrimitiveCollector(InView->DynamicPrimitiveCollector);
+		}
+		else // Shadow view
+		{
+			struct FSkipRange
+			{
+				size_t Offset;
+				size_t Size;
+			};
+
+			// These members are unused by the shadow view but are modified on the main view in parallel on the render thread.
+			const FSkipRange SkipRanges[] =
+			{
+				// Modified in FGPUOcclusionParallel::Finalize
+				{ offsetof(FViewInfo, IndividualOcclusionQueries), sizeof(FViewInfo::IndividualOcclusionQueries) + sizeof(FViewInfo::GroupedOcclusionQueries) },
+				// Modified in FSceneRenderer::ComputeLightGrid
+				{ offsetof(FViewInfo, ForwardLightingResources),   sizeof(FViewInfo::ForwardLightingResources)  },
+				// Modified in FGPUScene::UploadDynamicPrimitiveShaderDataForView()
+				{ offsetof(FViewInfo, DynamicPrimitiveCollector),  sizeof(FViewInfo::DynamicPrimitiveCollector) },
+			};
+
+			const uint8* SrcData = reinterpret_cast<const uint8*>(InView);
+			uint8* DstData = reinterpret_cast<uint8*>(Result);
+			size_t CopyOffset = 0;
+
+			for (const FSkipRange& Skip : SkipRanges)
+			{
+				// Keep ranges in order.
+				check(Skip.Offset >= CopyOffset);
+
+				if (Skip.Offset > CopyOffset)
+				{
+					FMemory::Memcpy(DstData + CopyOffset, SrcData + CopyOffset, Skip.Offset - CopyOffset);
+				}
+
+				FMemory::Memzero(DstData + Skip.Offset, Skip.Size);
+				CopyOffset = Skip.Offset + Skip.Size;
+			}
+
+			if (CopyOffset < sizeof(FViewInfo))
+			{
+				FMemory::Memcpy(DstData + CopyOffset, SrcData + CopyOffset, sizeof(FViewInfo) - CopyOffset);
+			}
+		}
 
 		// we want these to start null without a reference count, since we clear a ref later
 		TUniformBufferRef<FViewUniformShaderParameters> NullViewUniformBuffer;
@@ -41,10 +89,6 @@ public:
 		TArray<FPrimitiveUniformShaderParameters> NullDynamicPrimitiveShaderData;
 		FMemory::Memset(&Result->ParallelMeshDrawCommandPasses, 0, sizeof(Result->ParallelMeshDrawCommandPasses));
 
-		// Ensure the internal state is maintained, needed because we've just Memcpy'd the member data.
-		static_assert(std::is_trivially_destructible_v<FGPUScenePrimitiveCollector> != 0, "The destructor is not invoked properly because of FMemory::Memcpy(*Result, *this) above");
-		Result->DynamicPrimitiveCollector = FGPUScenePrimitiveCollector(InView->DynamicPrimitiveCollector);
-
 		Result->SnapshotOriginView = InView;
 		return Result;
 	}
@@ -122,9 +166,9 @@ namespace ViewSnapshotCache
 {
 	static FViewSnapshotCache Cache;
 
-	FViewInfo* Create(const FViewInfo* View)
+	FViewInfo* Create(const FViewInfo* View, EViewSnapshotType Type)
 	{
-		return Cache.Create(View);
+		return Cache.Create(View, Type);
 	}
 
 	void Deallocate()
diff --git a/Engine/Source/Runtime/Renderer/Private/Renderer/ViewSnapshotCache.h b/Engine/Source/Runtime/Renderer/Private/Renderer/ViewSnapshotCache.h
--- a/Engine/Source/Runtime/Renderer/Private/Renderer/ViewSnapshotCache.h
+++ b/Engine/Source/Runtime/Renderer/Private/Renderer/ViewSnapshotCache.h
@@ -4,10 +4,12 @@
 
 class FViewInfo;
 
+enum class EViewSnapshotType;
+
 namespace ViewSnapshotCache
 {
 	// Creates and returns a snapshot from the input view.
-	extern FViewInfo* Create(const FViewInfo* View);
+	extern FViewInfo* Create(const FViewInfo* View, EViewSnapshotType Type);
 
 	// Called on the render thread after rendering to prepare snapshots for destruction.
 	extern void Deallocate();
diff --git a/Engine/Source/Runtime/Renderer/Private/SceneRendering.cpp b/Engine/Source/Runtime/Renderer/Private/SceneRendering.cpp
--- a/Engine/Source/Runtime/Renderer/Private/SceneRendering.cpp
+++ b/Engine/Source/Runtime/Renderer/Private/SceneRendering.cpp
@@ -2382,9 +2382,9 @@ const FViewInfo* FViewInfo::GetPrimaryView() const
 
 extern TUniformBufferRef<FMobileDirectionalLightShaderParameters>& GetNullMobileDirectionalLightShaderParameters();
 
-FViewInfo* FViewInfo::CreateSnapshot() const
+FViewInfo* FViewInfo::CreateSnapshot(EViewSnapshotType Type) const
 {
-	return ViewSnapshotCache::Create(this);
+	return ViewSnapshotCache::Create(this, Type);
 }
 
 FInt32Range FViewInfo::GetDynamicMeshElementRange(uint32 PrimitiveIndex) const
diff --git a/Engine/Source/Runtime/Renderer/Private/SceneRendering.h b/Engine/Source/Runtime/Renderer/Private/SceneRendering.h
--- a/Engine/Source/Runtime/Renderer/Private/SceneRendering.h
+++ b/Engine/Source/Runtime/Renderer/Private/SceneRendering.h
@@ -1142,6 +1142,15 @@ struct FPrimitiveInstanceRange
 	int32 NumInstances;
 };
 
+enum class EViewSnapshotType
+{
+	// The copy contains all members for rendering. Must be done on the rendering thread.
+	Main,
+
+	// The copy contains only members relevant for shadow views.
+	Shadow
+};
+
 /** A FSceneView with additional state used by the scene renderer. */
 class FViewInfo : public FSceneView
 {
@@ -1799,7 +1808,7 @@ public:
 	inline const FViewInfo* GetInstancedView() const { return static_cast<const FViewInfo*>(GetInstancedSceneView()); }
 
 	/** Create a snapshot of this view info on the scene allocator. */
-	FViewInfo* CreateSnapshot() const;
+	FViewInfo* CreateSnapshot(EViewSnapshotType Type = EViewSnapshotType::Main) const;
 
 	void Cleanup();
 
diff --git a/Engine/Source/Runtime/Renderer/Private/ShadowDepthRendering.cpp b/Engine/Source/Runtime/Renderer/Private/ShadowDepthRendering.cpp
--- a/Engine/Source/Runtime/Renderer/Private/ShadowDepthRendering.cpp
+++ b/Engine/Source/Runtime/Renderer/Private/ShadowDepthRendering.cpp
@@ -1338,7 +1338,7 @@ void FProjectedShadowInfo::SetupShadowDepthView(FSceneRenderer* SceneRenderer)
 {
 	FViewInfo* FoundView = FindViewForShadow(SceneRenderer);
 	check(FoundView && IsInParallelRenderingThread());
-	FViewInfo* DepthPassView = FoundView->CreateSnapshot();
+	FViewInfo* DepthPassView = FoundView->CreateSnapshot(EViewSnapshotType::Shadow);
 	// We are starting a new collection of dynamic primitives for the shadow views.
 	DepthPassView->DynamicPrimitiveCollector = FGPUScenePrimitiveCollector(&SceneRenderer->GetGPUSceneDynamicContext());
 
