[TSR] Clean up unnecessary permutations sky relaxation and reprojection field for thin geo pass.
Now it's equivalent to sky relaxation disabled as there is no need to have a specific logic to deal with volumetric clouds, and reprojection field always enabled to have better stability.

#Synced-CL 50817990

diff --git a/Engine/Shaders/Private/TemporalSuperResolution/TSRDetectThinGeometry.usf b/Engine/Shaders/Private/TemporalSuperResolution/TSRDetectThinGeometry.usf
--- a/Engine/Shaders/Private/TemporalSuperResolution/TSRDetectThinGeometry.usf
+++ b/Engine/Shaders/Private/TemporalSuperResolution/TSRDetectThinGeometry.usf
@@ -20,8 +20,6 @@
 
 #define EDGE_TYPE EDGE_HORIZONTAL_VERTICAL_DIAGNAL
 
-#define CONFIG_THIN_GEOMETRY_EDGE_REPROJECTION (DIM_THIN_GEOMETRY_EDGE_REPROJECTION)
-#define CONFIG_SKY_RELAXATION (DIM_SKY_RELAXATION)
 #define CONFIG_VARYING_RANGE_DETECTION (DIM_VARYING_RANGE_DETECTION)
 //------------------------------------------------------- CONSTANTS
 
@@ -103,9 +101,7 @@ void FetchCommonInput(
 	)
 {
 	tsr_tensor_bool bHasReprojectionOffset;
-#if CONFIG_THIN_GEOMETRY_EDGE_REPROJECTION
 	tsr_tensor_short2 ReprojectionOffset;
-#endif
 
 	UNROLL_N(SIMD_SIZE)
 	for (uint ElementIndex = 0; ElementIndex < SIMD_SIZE; ElementIndex++)
@@ -115,14 +111,11 @@ void FetchCommonInput(
 
 		tsr_ushort EncodedDilateMask = R8Output[int3(ClampedFetchPixelPos,0)];
 		bHasReprojectionOffset.SetElement(ElementIndex, (EncodedDilateMask & tsr_ushort(0x80)) != tsr_ushort(0));
-#if CONFIG_THIN_GEOMETRY_EDGE_REPROJECTION
 		ReprojectionOffset.SetElement(ElementIndex, DecodeReprojectionOffset(R8Output[int3(ClampedFetchPixelPos, ThinGeometryTextureIndex)]));
-#endif
 		tsr_ushort BitMask = tsr_ushort(DecimateMaskTexture[ClampedFetchPixelPos].r * tsr_half(255.0));
 		bHasPixelAnimation.SetElement(ElementIndex,(BitMask & tsr_ushort(0x4)) != tsr_ushort(0));
 	}
 
-#if CONFIG_THIN_GEOMETRY_EDGE_REPROJECTION
 	tsr_tensor_half ReprojectedCurrentCoverage;
 	UNROLL_N(SIMD_SIZE)
 	for (uint ElementIndex = 0; ElementIndex < SIMD_SIZE; ElementIndex++)
@@ -131,9 +124,6 @@ void FetchCommonInput(
 		ReprojectedCurrentCoverage.SetElement(ElementIndex, CurrentCoverageTexture[ClampedFetchPixelPos]);
 	}
 	CurrentCoverage = select(bHasReprojectionOffset, max(CurrentCoverage, ReprojectedCurrentCoverage), tsr_tensor_half::Const(0));
-#else
-	CurrentCoverage = select(bHasReprojectionOffset, CurrentCoverage, tsr_tensor_half::Const(0));
-#endif
 
 #if CONFIG_VARYING_RANGE_DETECTION
 	MaterialMask = tsr_tensor_ushort::CastFrom<tsr_half>(CurrentCoverage * tsr_half(2.01f));
@@ -177,7 +167,7 @@ void FetchRegularPassInput(
 
 //------------------------------------------------------- FUNCTIONS
 
-#if !CONFIG_SKY_RELAXATION
+
 // Horizontal thin edge detection.
 // 0 0 0
 // x x x
@@ -267,8 +257,6 @@ TLaneVector2D<ScalarType, VectorSize, LaneStrideX, LaneStrideY> VerticalThinEdge
 	return ThinGeometryEdge;
 }
 
-#endif // !CONFIG_SKY_RELAXATION
-
 tsr_tensor_ushort EncodeThinGeometryWeight(tsr_tensor_half ThinGeometryEdge
 	, tsr_tensor_bool IsClusterHoleRegion,
 	tsr_tensor_half Weight
@@ -343,13 +331,11 @@ void DetectThinGeometryCS(
 			ReprojectedHistoryCoverage
 		);
 		
-		#if !CONFIG_SKY_RELAXATION
 		{
 			ThinGeometryEdge = HorizontalThinEdge3x3(DeviceZ, DeviceZThinGeometryDetectionError);
 			ThinGeometryEdge = ThinGeometryEdge + VerticalThinEdge3x3(DeviceZ, DeviceZThinGeometryDetectionError);
 			//Debug[0].SetComponent(2, tsr_tensor_half(ThinGeometryEdge));
 		}
-		#endif
 
 		// Student t test of current observation and history to check if they are of the same distribution.
 		tsr_tensor_bool bSameDistribution = tsr_tensor_bool::Const(false);
diff --git a/Engine/Shaders/Private/TemporalSuperResolution/TSRWeightRelaxation.usf b/Engine/Shaders/Private/TemporalSuperResolution/TSRWeightRelaxation.usf
--- a/Engine/Shaders/Private/TemporalSuperResolution/TSRWeightRelaxation.usf
+++ b/Engine/Shaders/Private/TemporalSuperResolution/TSRWeightRelaxation.usf
@@ -14,8 +14,6 @@
 
 #define DEBUG_DIGESTABLE 0
 
-#define CONFIG_SKY_RELAXATION (DIM_SKY_RELAXATION)
-
 //------------------------------------------------------- CONSTANTS
 
 #if CONFIG_LDS_DESPILL
@@ -120,10 +118,6 @@ void FetchInput(
 	out tsr_tensor_half2 Translucency,
 	out tsr_tensor_half HistoryRelaxationWeight,
 	out tsr_tensor_bool IsClusterHoleRegion
-#if CONFIG_SKY_RELAXATION
-	,out tsr_tensor_half InputLumaBeforeComposition
-	,out tsr_tensor_half CurrentCoverage
-#endif
 )
 {
 	tsr_tensor_ushort ThinGeometryCache;
@@ -136,11 +130,6 @@ void FetchInput(
 		tsr_half4 FullTranslucency = InputSceneTranslucencyTexture[ClampedFetchPixelPos];
 		Translucency.SetElement(ElementIndex, tsr_half2(ComputeRGBLuma(FullTranslucency.rgb),FullTranslucency.a));
 		ThinGeometryCache.SetElement(ElementIndex, tsr_ushort(R8Output[int3(ClampedFetchPixelPos, ThinGeometryTextureIndex)]));
-
-#if CONFIG_SKY_RELAXATION
-		CurrentCoverage.SetElement(ElementIndex,CurrentCoverageTexture[ClampedFetchPixelPos]);
-		InputLumaBeforeComposition.SetElement(ElementIndex,InputMoireLumaTexture[ClampedFetchPixelPos]);
-#endif
 	}
 
 	DecodeThinGeometry(ThinGeometryCache, HistoryRelaxationWeight, IsClusterHoleRegion);	
@@ -173,47 +162,15 @@ void WeightRelaxationCS(
 	tsr_tensor_half2 Translucency;
 	tsr_tensor_half HistoryRelaxationWeight;
 	tsr_tensor_bool IsPartialCoverageRegion;
-#if CONFIG_SKY_RELAXATION
-	tsr_tensor_half InputLumaBeforeComposition;
-	tsr_tensor_half CurrentCoverage;
-#endif
+
 	FetchInput(
 		PixelPos,
  		/*out */ InputLuma,
 		/*out */ Translucency,
 		/*out */ HistoryRelaxationWeight,
 		/*out */ IsPartialCoverageRegion
-#if CONFIG_SKY_RELAXATION
-		,/*out */ InputLumaBeforeComposition
-		,/*out */ CurrentCoverage
-#endif
 	);
 
-	// Edge case: thin geometry against volumetric clouds.
-#if CONFIG_SKY_RELAXATION
-	tsr_tensor_bool IsEarlyTranslucencyEffect;
-	{
-		tsr_tensor_half EarlyTranslucencyAE = abs(InputLuma - InputLumaBeforeComposition);
-#if DEBUG_OUTPUT
-		Debug[0].SetComponent(0, EarlyTranslucencyAE);
-#endif
-
-		tsr_tensor_half MinCurrentCoverageNeighbor;
-		MinCurrentCoverageNeighbor = Min3x3(CurrentCoverage);
-
-		tsr_tensor_bool bIsLeafAgainstSky = and(and(EarlyTranslucencyAE > tsr_tensor_half::Const(0.3), InputLumaBeforeComposition > tsr_tensor_half::Const(0.01)), 
-			and(MinCurrentCoverageNeighbor == tsr_tensor_half::Const(0), IsThinGeometry(CurrentCoverage)));
-#if DEBUG_OUTPUT		
-		Debug[0].SetComponent(1, select(bIsLeafAgainstSky,abs(InputLuma - InputLumaBeforeComposition), tsr_tensor_half::Const(0)));
-#endif
-		IsEarlyTranslucencyEffect = bIsLeafAgainstSky;
-		tsr_tensor_half EarlyTranslucencyAddWeight = select(IsEarlyTranslucencyEffect, tsr_tensor_half::Const(1.0), tsr_tensor_half::Const(0));
-		EarlyTranslucencyAddWeight = WeightedAvg3x3(EarlyTranslucencyAddWeight, float3(0.8, 1.0, 0.8));
-		IsPartialCoverageRegion = select(EarlyTranslucencyAddWeight > tsr_tensor_half::Const(0),tsr_tensor_bool::Const(false),IsPartialCoverageRegion);
-		HistoryRelaxationWeight = saturate(HistoryRelaxationWeight + EarlyTranslucencyAddWeight);
-	}
-#endif
-
 	// If translucency has data turn down the weights aggresively. if alpha < 0.95
 	tsr_tensor_half TranslucencyRatio = Translucency[0] * rcp(Translucency[0] + InputLuma * Translucency[1] + tsr_tensor_half::Const(0.0005));
 
@@ -223,12 +180,10 @@ void WeightRelaxationCS(
 	RatioToLerpWeight = select(or(IsPartialCoverageRegion,TranslucencyRatio > tsr_tensor_half::Const(0.6)), RatioToLerpWeight, tsr_tensor_half::Const(1));
 	
 	// Write only when the weights has been modified.
-	#if !CONFIG_SKY_RELAXATION
 	bool bRelaxationWeightsModified = AnyComponent(and(HistoryRelaxationWeight > tsr_tensor_half::Const(tsr_half(0)), RatioToLerpWeight < tsr_tensor_half::Const(tsr_half(1.0))));
 	
 	BRANCH
 	if (bRelaxationWeightsModified)
-	#endif
 	{
 		HistoryRelaxationWeight = lerp(tsr_tensor_half::Const(0), HistoryRelaxationWeight, RatioToLerpWeight); 
 		tsr_tensor_ushort EncodedThinGeometry = EncodeThinGeometry(HistoryRelaxationWeight,IsPartialCoverageRegion);
diff --git a/Engine/Source/Runtime/Renderer/Private/PostProcess/TemporalSuperResolution.cpp b/Engine/Source/Runtime/Renderer/Private/PostProcess/TemporalSuperResolution.cpp
--- a/Engine/Source/Runtime/Renderer/Private/PostProcess/TemporalSuperResolution.cpp
+++ b/Engine/Source/Runtime/Renderer/Private/PostProcess/TemporalSuperResolution.cpp
@@ -339,11 +339,6 @@ TAutoConsoleVariable<int32> CVarTSRThinGeometryDetection(
 	TEXT("and a value < 0 must be used to disable it.\n"),
 	ECVF_RenderThreadSafe);
 
-TAutoConsoleVariable<bool> CVarTSRThinGeometryCoverageEdgeReprojection(
-	TEXT("r.TSR.ThinGeometryDetection.Coverage.EdgeReprojection"), true,
-	TEXT("Whether thin geometry edge against not thin geometry should be considered for coverage. Better stability for sparse thin geometry clusters."),
-	ECVF_RenderThreadSafe);
-
 TAutoConsoleVariable<int32> CVarTSRThinGeometryCoverageShadingRange(
 	TEXT("r.TSR.ThinGeometryDetection.Coverage.ShadingRange"), 3,
 	TEXT("The shading model range regarded as thin geometry and accumulate the coverage.\n")
@@ -363,11 +358,6 @@ TAutoConsoleVariable<int32> CVarTSRThinGeometryDetectionWeightRelaxation(
 	TEXT("Adaptively trim the history relaxation to avoid ghosting."),
 	ECVF_RenderThreadSafe);
 
-TAutoConsoleVariable<int32> CVarTSRThinGeometryDetectionWeightRelaxationSky(
-	TEXT("r.TSR.ThinGeometryDetection.WeightRelaxation.Sky"), 0,
-	TEXT("Apply weight adjust to deal with thin geometry shimmering against sky."),
-	ECVF_RenderThreadSafe);
-
 TAutoConsoleVariable<float> CVarTSRThinGeometryDetectionErrorMultiplier(
 	TEXT("r.TSR.ThinGeometryDetection.ErrorMultiplier"), 200.0,
 	TEXT("Define the depth difference multiplier between neighbors will be used to detect thin geometry. The larger the upscaler, the smaller the multiplier should be for temporal stability. ")
@@ -483,12 +473,6 @@ static bool ShouldEnableThinGeometryDetection()
 	return bNaniteFoliageEnabled || ThinGeometryDetection != 0;
 }
 
-bool ShouldApplySkyRelaxation()
-{
-	return (ShouldEnableThinGeometryDetection() && CVarTSRThinGeometryDetectionWeightRelaxation.GetValueOnRenderThread()!=0
-		&& CVarTSRThinGeometryDetectionWeightRelaxationSky.GetValueOnRenderThread()!=0);
-}
-
 class FTSRShader : public FGlobalShader
 {
 public:
@@ -1070,10 +1054,8 @@ class FTSRDetectThinGeometryCS : public FTSRConvolutionNetworkShader
 	DECLARE_GLOBAL_SHADER(FTSRDetectThinGeometryCS);
 	SHADER_USE_PARAMETER_STRUCT(FTSRDetectThinGeometryCS, FTSRConvolutionNetworkShader);
 
-	class FSkyRelaxationDim : SHADER_PERMUTATION_BOOL("DIM_SKY_RELAXATION");
-	class FThinGeometryEdgeReprojectionDim : SHADER_PERMUTATION_BOOL("DIM_THIN_GEOMETRY_EDGE_REPROJECTION");
 	class FVaryingRangeDetectionDim : SHADER_PERMUTATION_BOOL("DIM_VARYING_RANGE_DETECTION");
-	using FPermutationDomain = TShaderPermutationDomain<FTSRConvolutionNetworkShader::FPermutationDomain, FThinGeometryEdgeReprojectionDim, FSkyRelaxationDim, FVaryingRangeDetectionDim>;
+	using FPermutationDomain = TShaderPermutationDomain<FTSRConvolutionNetworkShader::FPermutationDomain, FVaryingRangeDetectionDim>;
 
 	BEGIN_SHADER_PARAMETER_STRUCT(FParameters, )
 		SHADER_PARAMETER_STRUCT_INCLUDE(FTSRCommonParameters, CommonParameters)
@@ -1145,8 +1127,7 @@ class FTSRTSRWeightRelaxationCS :  public FTSRConvolutionNetworkShader
 	DECLARE_GLOBAL_SHADER(FTSRTSRWeightRelaxationCS);
 	SHADER_USE_PARAMETER_STRUCT(FTSRTSRWeightRelaxationCS, FTSRConvolutionNetworkShader);
 
-	class FSkyRelaxationDim : SHADER_PERMUTATION_BOOL("DIM_SKY_RELAXATION");
-	using FPermutationDomain = TShaderPermutationDomain<FTSRConvolutionNetworkShader::FPermutationDomain,FSkyRelaxationDim>;
+	using FPermutationDomain = TShaderPermutationDomain<FTSRConvolutionNetworkShader::FPermutationDomain>;
 
 	BEGIN_SHADER_PARAMETER_STRUCT(FParameters,)
 		SHADER_PARAMETER_STRUCT_INCLUDE(FTSRCommonParameters, CommonParameters)
@@ -1154,8 +1135,6 @@ class FTSRTSRWeightRelaxationCS :  public FTSRConvolutionNetworkShader
 		SHADER_PARAMETER(int32, ThinGeometryTextureIndex)
 		SHADER_PARAMETER(float, MaxRelaxationWeight)
 
-		SHADER_PARAMETER_RDG_TEXTURE_SRV(Texture2D, CurrentCoverageTexture)
-		SHADER_PARAMETER_RDG_TEXTURE_SRV(Texture2D, InputMoireLumaTexture)
 		SHADER_PARAMETER_RDG_TEXTURE(Texture2D, InputTexture)                  // Scene with
 		SHADER_PARAMETER_RDG_TEXTURE(Texture2D, InputSceneTranslucencyTexture) // Translucency texture
 		SHADER_PARAMETER_RDG_TEXTURE_UAV(RWTexture2DArray, R8Output)
@@ -1687,7 +1666,7 @@ bool NeedTSRTranslucentClippedDepthPass(const FViewInfo& View)
 
 bool ShouldAddTSRMainFlickeringLumaPass()
 {
-	return  (CVarTSRFlickeringEnable.GetValueOnRenderThread()!=0 && CVarTSRFlickeringPeriod.GetValueOnRenderThread() != 0.0f) || ShouldApplySkyRelaxation();
+	return  (CVarTSRFlickeringEnable.GetValueOnRenderThread()!=0 && CVarTSRFlickeringPeriod.GetValueOnRenderThread() != 0.0f);
 }
 
 bool ShouldAddTSRMainThinGeometryCoveragePass()
@@ -2483,7 +2462,7 @@ FDefaultTemporalUpscaler::FOutputs AddTemporalSuperResolutionPasses(
 
 		int32 TileSize = 8;
 		FTSRDilateVelocityCS::FPermutationDomain PermutationVector;
-		PermutationVector.Set<FTSRDilateVelocityCS::FThinGeometryEdgeReprojectionDim>(PassConfig.ThinGeometryDetectionEnable && CVarTSRThinGeometryCoverageEdgeReprojection.GetValueOnRenderThread());
+		PermutationVector.Set<FTSRDilateVelocityCS::FThinGeometryEdgeReprojectionDim>(PassConfig.ThinGeometryDetectionEnable);
 
 		FTSRDilateVelocityCS::FParameters* PassParameters = GraphBuilder.AllocParameters<FTSRDilateVelocityCS::FParameters>();
 		PassParameters->CommonParameters = CommonParameters;
@@ -2716,8 +2695,6 @@ FDefaultTemporalUpscaler::FOutputs AddTemporalSuperResolutionPasses(
 			ConvolutionNetworkPermutationVector.Set<FTSRShader::FAlphaChannelDim>(bSupportsAlpha);
 
 			FTSRDetectThinGeometryCS::FPermutationDomain PermutationVector;
-			PermutationVector.Set<FTSRDetectThinGeometryCS::FSkyRelaxationDim>(ShouldApplySkyRelaxation());
-            	PermutationVector.Set<FTSRDetectThinGeometryCS::FThinGeometryEdgeReprojectionDim>(CVarTSRThinGeometryCoverageEdgeReprojection.GetValueOnRenderThread());
 			PermutationVector.Set<FTSRDetectThinGeometryCS::FVaryingRangeDetectionDim>(EnableVaryingRangeDetection);
 			PermutationVector.Set<FTSRConvolutionNetworkShader::FPermutationDomain>(ConvolutionNetworkPermutationVector);
 			PermutationVector = FTSRDetectThinGeometryCS::RemapPermutation(PermutationVector);
@@ -2803,12 +2780,11 @@ FDefaultTemporalUpscaler::FOutputs AddTemporalSuperResolutionPasses(
 			ConvolutionNetworkPermutationVector.Set<FTSRShader::FAlphaChannelDim>(bSupportsAlpha);
 
 			FTSRTSRWeightRelaxationCS::FPermutationDomain PermutationVector;
-			PermutationVector.Set<FTSRTSRWeightRelaxationCS::FSkyRelaxationDim>(ShouldApplySkyRelaxation());
 			PermutationVector.Set<FTSRConvolutionNetworkShader::FPermutationDomain>(ConvolutionNetworkPermutationVector);
 			PermutationVector = FTSRTSRWeightRelaxationCS::RemapPermutation(PermutationVector);
 
 			const int32 GroupTileSize = 32;
-			const int32 TileOverscan = ShouldApplySkyRelaxation() ? 2 : 0;//FMath::Clamp(CVarTSRShadingTileOverscan.GetValueOnRenderThread(), 3, GroupTileSize / 2 - 1);
+			const int32 TileOverscan = 0;//FMath::Clamp(CVarTSRShadingTileOverscan.GetValueOnRenderThread(), 3, GroupTileSize / 2 - 1);
 			const int32 TileSize = GroupTileSize - 2 * TileOverscan;
 
 			FTSRTSRWeightRelaxationCS::FParameters* PassParameters = GraphBuilder.AllocParameters<FTSRTSRWeightRelaxationCS::FParameters>();
@@ -2817,31 +2793,6 @@ FDefaultTemporalUpscaler::FOutputs AddTemporalSuperResolutionPasses(
 			PassParameters->ThinGeometryTextureIndex = ThinGeometryTextureIndex;
 			PassParameters->MaxRelaxationWeight = FMath::Clamp(CVarTSRThinGeometryCoverageMaxRelaxationWeight.GetValueOnRenderThread(), 0.0f, 1.0f);
 
-			
-			if (PassInputs.FlickeringInputTexture.IsValid() && !bCameraCut)
-			{
-				ensure(InputRect == PassInputs.FlickeringInputTexture.ViewRect);
-				int32 ArrayIndex = GetTSRMainFlickeringLumaTextureArraySize() - 1;
-				PassParameters->CurrentCoverageTexture = GraphBuilder.CreateSRV(
-					FRDGTextureSRVDesc::CreateForSlice(PassInputs.FlickeringInputTexture.Texture,  /* SliceIndex = */ ArrayIndex));
-			}
-			else
-			{
-				PassParameters->CurrentCoverageTexture = GraphBuilder.CreateSRV(FRDGTextureSRVDesc(BlackDummy));
-			}
-			
-
-			if (PassInputs.FlickeringInputTexture.IsValid())
-			{
-				ensure(InputRect == PassInputs.FlickeringInputTexture.ViewRect);
-				PassParameters->InputMoireLumaTexture = GraphBuilder.CreateSRV(
-					FRDGTextureSRVDesc::CreateForSlice(PassInputs.FlickeringInputTexture.Texture, 0));
-			}
-			else
-			{
-				PassParameters->InputMoireLumaTexture = GraphBuilder.CreateSRV(FRDGTextureSRVDesc(BlackDummy));
-			}
-
 			PassParameters->InputTexture = PassInputs.SceneColor.Texture;
 			PassParameters->InputSceneTranslucencyTexture = SeparateTranslucencyTexture;
 			PassParameters->R8Output = GraphBuilder.CreateUAV(R8OutputTexture);
