Only use HDR path in GetTonemapperOutputDeviceParameters if the scene viewport requires it

Seemed like OutputMaxLuminance was 0 for scene captures and leading to slow NaN usage in UE::Color::ACES::GetTransformResources somehow.

#Synced-CL 50851879

diff --git a/Engine/Source/Runtime/Renderer/Private/PostProcess/PostProcessTonemap.cpp b/Engine/Source/Runtime/Renderer/Private/PostProcess/PostProcessTonemap.cpp
--- a/Engine/Source/Runtime/Renderer/Private/PostProcess/PostProcessTonemap.cpp
+++ b/Engine/Source/Runtime/Renderer/Private/PostProcess/PostProcessTonemap.cpp
@@ -243,17 +243,15 @@ FTonemapperOutputDeviceParameters GetTonemapperOutputDeviceParameters(const FSce
 	static TConsoleVariableData<float>* CVarOutputGamma = IConsoleManager::Get().FindTConsoleVariableDataFloat(TEXT("r.TonemapperGamma"));
 
 	EDisplayOutputFormat OutputDeviceValue;
-	bool bIsHDR;
+	bool bIsHDR = false;
 
 	if (Family.SceneCaptureSource == SCS_FinalColorHDR)
 	{
 		OutputDeviceValue = EDisplayOutputFormat::HDR_LinearNoToneCurve;
-		bIsHDR = true;
 	}
 	else if (Family.SceneCaptureSource == SCS_FinalToneCurveHDR)
 	{
 		OutputDeviceValue = EDisplayOutputFormat::HDR_LinearWithToneCurve;
-		bIsHDR = true;
 	}
 	else
 	{
@@ -273,7 +271,6 @@ FTonemapperOutputDeviceParameters GetTonemapperOutputDeviceParameters(const FSce
 	if (Gamma > 0.0f && (OutputDeviceValue == EDisplayOutputFormat::SDR_sRGB || OutputDeviceValue == EDisplayOutputFormat::SDR_Rec709))
 	{
 		OutputDeviceValue = EDisplayOutputFormat::SDR_ExplicitGammaMapping;
-		bIsHDR = false;
 	}
 
 	FVector3f InvDisplayGammaValue;
