Add a legend for LOD Coloration view modes.

#Synced-CL 50874393

diff --git a/Engine/Source/Runtime/Renderer/Private/PostProcess/PostProcessLODColorationLegend.cpp b/Engine/Source/Runtime/Renderer/Private/PostProcess/PostProcessLODColorationLegend.cpp
new file mode 100644
--- /dev/null
+++ b/Engine/Source/Runtime/Renderer/Private/PostProcess/PostProcessLODColorationLegend.cpp
@@ -0,0 +1,77 @@
+// Copyright Epic Games, Inc. All Rights Reserved.
+
+#include "PostProcess/PostProcessLODColorationLegend.h"
+#include "UnrealEngine.h"
+#include "DebugViewModeRendering.h"
+#include "SceneRendering.h"
+
+namespace UE::Private
+{
+static void DrawLODBox(FCanvas& Canvas, int32 PosX, int32 PosY, int32 SizeX, int32 SizeY, const FLinearColor& Color, int32 LODValue)
+{
+	// the color coded tile
+	Canvas.DrawTile(PosX, PosY, SizeX, SizeY, 0, 0, 1, 1, Color);
+
+	// and the lil number label that goes on it
+	const UFont* StatsFont = GetStatsFont();
+	const FString LODString = FString::FromInt(LODValue);
+	int32 TextSizeX = 0, TextSizeY = 0;
+	StringSize(StatsFont, TextSizeX, TextSizeY, LODString);
+	Canvas.DrawShadowedText(
+		PosX + (SizeX / 2) - (TextSizeX / 2), 
+		PosY + (SizeY / 2) - (TextSizeY / 2), 
+		FText::FromString(LODString), 
+		StatsFont, 
+		FLinearColor(0.7f, 0.7f, 0.7f));
+}
+} //! namespace
+
+FScreenPassTexture AddLODColorationLegendPass(FRDGBuilder& GraphBuilder, const FViewInfo& View, const FLODColorationLegendInputs& Inputs)
+{
+	FScreenPassRenderTarget Output = Inputs.OverrideOutput;
+
+	if (Output.IsValid())
+	{
+		AddDrawTexturePass(GraphBuilder, View, Inputs.SceneColor, Output);
+	}
+	else
+	{
+		Output = FScreenPassRenderTarget(Inputs.SceneColor, ERenderTargetLoadAction::ELoad);
+	}
+
+	const TArrayView<const FLinearColor> Colors = Inputs.Colors;
+
+	AddDrawCanvasPass(GraphBuilder, RDG_EVENT_NAME("LODColorationLegend"), View, Output, [Colors](FCanvas& Canvas)
+	{
+		const FIntRect& OutputViewRect = Canvas.GetViewRect();
+		// adjust canvas scale for DPI
+		const float DPIScale = Canvas.GetDPIScale();
+		Canvas.SetBaseTransform(FMatrix(FScaleMatrix(DPIScale)* Canvas.CalcBaseTransform2D(OutputViewRect.Width(), OutputViewRect.Height())));
+
+		const int32 NumColors = Colors.Num();
+
+		// some named variables make size calculations easy to follow
+		const int32 BorderWidth = 2;
+		const int32 OffsetFromViewBottom = 50;
+		const int32 SectionWidth = 50;
+		const int32 LegendHeight = 20;
+
+		// figuring out size + placement
+		const int32 LegendWidth = SectionWidth * NumColors;
+		const int32 LegendHalfWidth = LegendWidth / 2;
+
+		const int32 LegendPosX = OutputViewRect.Min.X + (OutputViewRect.Width() / 2) - LegendHalfWidth;
+		const int32 LegendPosY = OutputViewRect.Max.Y - OffsetFromViewBottom;
+
+		// draw a thin black border around our legend
+		Canvas.DrawTile(LegendPosX - BorderWidth, LegendPosY - BorderWidth, LegendWidth + BorderWidth*2, LegendHeight + BorderWidth*2, 0, 0, 1, 1, FLinearColor::Black);
+
+		// and then the color coded boxes with their associated LOD number
+		for (int32 i = 0; i < NumColors; ++i)
+		{
+			UE::Private::DrawLODBox(Canvas, LegendPosX + (SectionWidth * i), LegendPosY, SectionWidth, LegendHeight, Colors[i], i);
+		}
+	});
+
+	return MoveTemp(Output);
+}
diff --git a/Engine/Source/Runtime/Renderer/Private/PostProcess/PostProcessLODColorationLegend.h b/Engine/Source/Runtime/Renderer/Private/PostProcess/PostProcessLODColorationLegend.h
new file mode 100644
--- /dev/null
+++ b/Engine/Source/Runtime/Renderer/Private/PostProcess/PostProcessLODColorationLegend.h
@@ -0,0 +1,14 @@
+// Copyright Epic Games, Inc. All Rights Reserved.
+
+#pragma once
+
+#include "ScreenPass.h"
+
+struct FLODColorationLegendInputs
+{
+	FScreenPassRenderTarget OverrideOutput;
+	FScreenPassTexture SceneColor;
+	TArrayView<const FLinearColor> Colors;
+};
+
+FScreenPassTexture AddLODColorationLegendPass(FRDGBuilder& GraphBuilder, const FViewInfo& View, const FLODColorationLegendInputs& Inputs);
\ No newline at end of file
diff --git a/Engine/Source/Runtime/Renderer/Private/PostProcess/PostProcessing.cpp b/Engine/Source/Runtime/Renderer/Private/PostProcess/PostProcessing.cpp
--- a/Engine/Source/Runtime/Renderer/Private/PostProcess/PostProcessing.cpp
+++ b/Engine/Source/Runtime/Renderer/Private/PostProcess/PostProcessing.cpp
@@ -42,6 +42,7 @@
 #include "PostProcess/PostProcessTestImage.h"
 #include "PostProcess/PostProcessVisualizeCalibrationMaterial.h"
 #include "PostProcess/PostProcessFFTBloom.h"
+#include "PostProcess/PostProcessLODColorationLegend.h"
 #include "PostProcess/PostProcessStreamingAccuracyLegend.h"
 #include "PostProcess/PostProcessSubsurface.h"
 #include "PostProcess/SubpixelMorphologicalAA.h"
@@ -2372,7 +2373,24 @@ void AddDebugViewPostProcessingPasses(
 			break;
 		}
 		case DVSM_LODColoration:
+		{
+			FLODColorationLegendInputs PassInputs{};
+			PassInputs.OverrideOutput = OverrideOutput;
+			PassInputs.SceneColor = SceneColor;
+			const FEngineShowFlags& Flags = View.Family->EngineShowFlags;
+			if (Flags.HLODColoration != 0)
+			{
+				PassInputs.Colors = GEngine->HLODColorationColors;
+			}
+			else if (Flags.LODColoration != 0)
+			{
+				PassInputs.Colors = GEngine->LODColorationColors;
+			}
+
+			SceneColor = AddLODColorationLegendPass(GraphBuilder, View, PassInputs);
+
 			break;
+		}
 		case DVSM_ShadowCasters:
 			break;
 		default:
