Material Cache, minor bug fixes
- Updated page table uniforms, packed sign bit no longer relevant
- Added missing Nanite binding
- Fixed invalid render target format from cache attribute

#Synced-CL 50905358

diff --git a/Engine/Source/Runtime/Engine/Internal/MaterialCache/MaterialCacheAttribute.h b/Engine/Source/Runtime/Engine/Internal/MaterialCache/MaterialCacheAttribute.h
--- a/Engine/Source/Runtime/Engine/Internal/MaterialCache/MaterialCacheAttribute.h
+++ b/Engine/Source/Runtime/Engine/Internal/MaterialCache/MaterialCacheAttribute.h
@@ -372,7 +372,7 @@ static void PackMaterialCacheAttributeLayers(TArray<EMaterialCacheAttribute>& At
 		switch (Attribute)
 		{
 		case EMaterialCacheAttribute::BaseColor:
-			Layer.RenderFormat     = PF_R8G8B8;
+			Layer.RenderFormat     = PF_R8G8B8A8;
 			Layer.CompressedFormat = PF_DXT5;
 			Layer.bIsSRGB          = true;
 			break;
@@ -388,7 +388,7 @@ static void PackMaterialCacheAttributeLayers(TArray<EMaterialCacheAttribute>& At
 			Layer.CompressedFormat = PF_DXT1;
 			break;
 		case EMaterialCacheAttribute::WorldPosition:
-			Layer.RenderFormat     = PF_R8G8B8;
+			Layer.RenderFormat     = PF_R8G8B8A8;
 			Layer.CompressedFormat = PF_DXT5;
 			break;
 		case EMaterialCacheAttribute::WorldHeight:
diff --git a/Engine/Source/Runtime/Renderer/Private/MaterialCache/MaterialCacheRenderer.cpp b/Engine/Source/Runtime/Renderer/Private/MaterialCache/MaterialCacheRenderer.cpp
--- a/Engine/Source/Runtime/Renderer/Private/MaterialCache/MaterialCacheRenderer.cpp
+++ b/Engine/Source/Runtime/Renderer/Private/MaterialCache/MaterialCacheRenderer.cpp
@@ -1565,6 +1565,7 @@ static void MaterialCacheSetupNaniteContext(FSceneRendererBase* Renderer, FRDGBu
 	ShadingUniformParameters->MultiViewIndices = GraphBuilder.CreateSRV(GSystemTextures.GetDefaultStructuredBuffer<uint32>(GraphBuilder));
 	ShadingUniformParameters->MultiViewRectScaleOffsets = GraphBuilder.CreateSRV(GSystemTextures.GetDefaultStructuredBuffer<FVector4>(GraphBuilder));
 	ShadingUniformParameters->InViews = GraphBuilder.CreateSRV(PackedViewBuffer);
+	ShadingUniformParameters->ThreadGroupData = GraphBuilder.CreateSRV(GSystemTextures.GetDefaultStructuredBuffer<uint32>(GraphBuilder));
 
 	FMaterialCacheNaniteShadeParameters* PassParameters = GraphBuilder.AllocParameters<FMaterialCacheNaniteShadeParameters>();
 	PassParameters->NaniteRaster = GraphBuilder.CreateUniformBuffer(RasterUniformParameters);
diff --git a/Engine/Source/Runtime/Renderer/Private/MaterialCache/MaterialCacheTagProvider.cpp b/Engine/Source/Runtime/Renderer/Private/MaterialCache/MaterialCacheTagProvider.cpp
--- a/Engine/Source/Runtime/Renderer/Private/MaterialCache/MaterialCacheTagProvider.cpp
+++ b/Engine/Source/Runtime/Renderer/Private/MaterialCache/MaterialCacheTagProvider.cpp
@@ -64,7 +64,7 @@ void FMaterialCacheTagProvider::CreateDeviceBuffersOrResize(struct FMaterialCach
 	const uint64 MinBufferSize = sizeof(UE::HLSL::FMaterialCacheTagEntry) * Bucket.TagEntries.Num();
 
 	// Can we accomodate for the current entry count?
-	if (Bucket.EntryBuffer.IsValid() && Bucket.EntryBuffer->GetDesc().Size < MinBufferSize)
+	if (Bucket.EntryBuffer.IsValid() && Bucket.EntryBuffer->GetDesc().Size >= MinBufferSize)
 	{
 		return;
 	}
@@ -229,10 +229,12 @@ FMaterialCacheTagBindingData FMaterialCacheTagProvider::GetBindingData(const FGu
 	
 	FMaterialCacheTagBucket* Bucket = *It;
 	Out.TagBufferSRV = Bucket->EntryBufferSRV;
+	check(Out.TagBufferSRV);
 
 	// All virtual textures in the bucket share the same page table
 	IAllocatedVirtualTexture* ReferenceTexture = Bucket->VirtualTextures[0];
 	Out.PageTableSRV = ReferenceTexture->GetPageTableTexture(0);
+	check(Out.PageTableSRV);
 
 	// And all physical textures
 	for (uint32 i = 0; i < ReferenceTexture->GetNumTextureLayers(); i++)
@@ -265,13 +267,11 @@ FMaterialCacheTagUniformData FMaterialCacheTagProvider::GetUniformData(const FGu
 	const uint32 PageBorderSize         = ReferenceTexture->GetTileBorderSize();
 	const float  RcpPhysicalTextureSize = 1.0f / static_cast<float>(ReferenceTexture->GetPhysicalTextureSize(0));
 	const uint32 PageSizeWithBorder     = PageSize + PageBorderSize * 2u;
-	const bool   bPageTableExtraBits    = ReferenceTexture->GetPageTableFormat() == EVTPageTableFormat::UInt32;
-	const float  PackedSignBit          = bPageTableExtraBits ? 1.0f : -1.0f;
 
 	Out.PackedTableUniform.X  = 0xFFFFFFFF;
 	Out.PackedTableUniform.Y  = FMath::AsUInt(static_cast<float>(PageSize) * RcpPhysicalTextureSize);
 	Out.PackedTableUniform.Z  = FMath::AsUInt(static_cast<float>(PageBorderSize) * RcpPhysicalTextureSize);
-	Out.PackedTableUniform.W  = FMath::AsUInt(static_cast<float>(PageSizeWithBorder) * RcpPhysicalTextureSize * PackedSignBit);
+	Out.PackedTableUniform.W  = FMath::AsUInt(static_cast<float>(PageSizeWithBorder) * RcpPhysicalTextureSize);
 	
 	return Out;
 }
