[HWRT] On-screen warning when visualizing far field but it's either not enabled or doesn't have any instances.

#Synced-CL 50872947

diff --git a/Engine/Source/Runtime/Renderer/Private/DeferredShadingRenderer.cpp b/Engine/Source/Runtime/Renderer/Private/DeferredShadingRenderer.cpp
--- a/Engine/Source/Runtime/Renderer/Private/DeferredShadingRenderer.cpp
+++ b/Engine/Source/Runtime/Renderer/Private/DeferredShadingRenderer.cpp
@@ -1722,7 +1722,7 @@ extern void RenderRayTracingDebug(
 	FSceneTextures& SceneTextures,
 	TConstArrayView<FRayTracingShaderBindingData> VisibleRayTracingShaderBindings,
 	FRayTracingPickingFeedback& PickingFeedback);
-extern void RayTracingDebugDisplayOnScreenMessages(FScreenMessageWriter& Writer, const FViewInfo& View);
+extern void RayTracingDebugDisplayOnScreenMessages(FScreenMessageWriter& Writer, const FViewInfo& View, bool bViewHasFarFieldInstances);
 #endif // RHI_RAYTRACING
 
 void FDeferredShadingSceneRenderer::Render(FRDGBuilder& GraphBuilder, const FSceneRenderUpdateInputs* SceneUpdateInputs)
@@ -3796,9 +3796,10 @@ void FDeferredShadingSceneRenderer::Render(FRDGBuilder& GraphBuilder, const FSce
 					FRayTracingPickingFeedback PickingFeedback = {};
 					RenderRayTracingDebug(GraphBuilder, *Scene, View, SceneTextures, VisibleRayTracingShaderBindings, PickingFeedback);
 
-					OnGetOnScreenMessages.AddLambda([this, &View, PickingFeedback](FScreenMessageWriter& ScreenMessageWriter)->void
+					const bool bViewHasFarFieldInstances = RayTracingScene.GetNumNativeInstances(ERayTracingSceneLayer::FarField, View.GetRayTracingSceneViewHandle()) > 0;
+					OnGetOnScreenMessages.AddLambda([this, &View, PickingFeedback, bViewHasFarFieldInstances](FScreenMessageWriter& ScreenMessageWriter)->void
 						{
-							RayTracingDebugDisplayOnScreenMessages(ScreenMessageWriter, View);
+							RayTracingDebugDisplayOnScreenMessages(ScreenMessageWriter, View, bViewHasFarFieldInstances);
 							RayTracingDisplayPicking(PickingFeedback, ScreenMessageWriter);
 						});
 				}
diff --git a/Engine/Source/Runtime/Renderer/Private/RayTracing/RayTracingDebug.cpp b/Engine/Source/Runtime/Renderer/Private/RayTracing/RayTracingDebug.cpp
--- a/Engine/Source/Runtime/Renderer/Private/RayTracing/RayTracingDebug.cpp
+++ b/Engine/Source/Runtime/Renderer/Private/RayTracing/RayTracingDebug.cpp
@@ -1925,7 +1925,7 @@ void RenderRayTracingDebug(FRDGBuilder& GraphBuilder, const FScene& Scene, const
 	}
 }
 
-extern void RayTracingDebugDisplayOnScreenMessages(FScreenMessageWriter& Writer, const FViewInfo& View)
+void RayTracingDebugDisplayOnScreenMessages(FScreenMessageWriter& Writer, const FViewInfo& View, bool bViewHasFarFieldInstances)
 {
 	const uint32 DebugVisualizationMode = GetRaytracingDebugViewModeID(View);
 
@@ -1942,6 +1942,20 @@ extern void RayTracingDebugDisplayOnScreenMessages(FScreenMessageWriter& Writer,
 	default:
 		break;
 	}
+
+	if (CVarRayTracingVisualizeFarField.GetValueOnRenderThread())
+	{
+		if (!View.RayTracingCullingParameters.bIsRayTracingFarField)
+		{
+			static const FText Message = NSLOCTEXT("Renderer", "RayTracingDebugVizFarFieldNotEnabled", "Ray Tracing Far Field is automatically disabled as there are no features using it.");
+			Writer.DrawLine(Message);
+		}
+		else if (!bViewHasFarFieldInstances)
+		{
+			static const FText Message = NSLOCTEXT("Renderer", "RayTracingDebugVizFarFieldEmpty", "Ray Tracing Far Field doesn't have any instances.");
+			Writer.DrawLine(Message);
+		}
+	}
 }
 
 void FDeferredShadingSceneRenderer::RayTracingDisplayPicking(const FRayTracingPickingFeedback& PickingFeedback, FScreenMessageWriter& Writer)
