[HWRT] Handle large scratch buffer (>2GB) when doing dynamic geometry update: use RHI function to build geometries as it will internally split them and allocate smaller size scratch buffer for each batch.

#rnx

#Synced-CL 50782841

diff --git a/Engine/Source/Runtime/Renderer/Private/RayTracing/RayTracingDynamicGeometryUpdateManager.cpp b/Engine/Source/Runtime/Renderer/Private/RayTracing/RayTracingDynamicGeometryUpdateManager.cpp
--- a/Engine/Source/Runtime/Renderer/Private/RayTracing/RayTracingDynamicGeometryUpdateManager.cpp
+++ b/Engine/Source/Runtime/Renderer/Private/RayTracing/RayTracingDynamicGeometryUpdateManager.cpp
@@ -16,6 +16,8 @@
 
 #include "Materials/MaterialRenderProxy.h"
 
+DECLARE_LOG_CATEGORY_CLASS(LogRayTracingDynamicGeometryManager, Log, All);
+
 DECLARE_GPU_STAT(RayTracingDynamicGeometry);
 
 DECLARE_DWORD_COUNTER_STAT(TEXT("Ray tracing dynamic build primitives"), STAT_RayTracingDynamicBuildPrimitives, STATGROUP_SceneRendering);
@@ -896,12 +898,23 @@ void FRayTracingDynamicGeometryUpdateManager::AddDynamicGeometryUpdatePass(FRDGB
 	const uint32 ScratchAlignment = GRHIRayTracingScratchBufferAlignment;
 	const uint32 BLASScratchSize = ScratchBufferSize;
 
+	// Slow path when scratch buffer exceeds 2GBs
+	const uint32 MaxBLASScratchSize = 2147483647u;
+	const bool bSlowPath = BLASScratchSize >= MaxBLASScratchSize;
+
+	if (bSlowPath)
+	{
+		UE_LOG(LogRayTracingDynamicGeometryManager, Warning, TEXT("Required scratch buffer size (%u) is larger than 2GB. Consider reducing dynamic geometry complexity or number of primitives updated per frame."), BLASScratchSize);
+	}
+
 	if (BLASScratchSize > 0)
 	{
 		FRDGBufferDesc ScratchBufferDesc;
 		ScratchBufferDesc.Usage = EBufferUsageFlags::RayTracingScratch | EBufferUsageFlags::StructuredBuffer;
 		ScratchBufferDesc.BytesPerElement = ScratchAlignment;
-		ScratchBufferDesc.NumElements = FMath::DivideAndRoundUp(BLASScratchSize, ScratchAlignment);
+
+		// Even when using the slow path we need to allocate scratch buffer in RDG to get correct pass order.
+		ScratchBufferDesc.NumElements = bSlowPath ? 1 : FMath::DivideAndRoundUp(BLASScratchSize, ScratchAlignment);
 
 		OutDynamicGeometryScratchBuffer = GraphBuilder.CreateBuffer(ScratchBufferDesc, TEXT("DynamicGeometry.BLASSharedScratchBuffer"));
 	}
@@ -939,7 +952,7 @@ void FRayTracingDynamicGeometryUpdateManager::AddDynamicGeometryUpdatePass(FRDGB
 		PassParams->DynamicGeometryScratchBuffer = OutDynamicGeometryScratchBuffer;
 
 		GraphBuilder.AddPass(RDG_EVENT_NAME("RayTracingDynamicUpdateBuild"), PassParams, ComputePassFlags | ERDGPassFlags::NeverCull,
-			[this, PassParams](FRHICommandList& RHICmdList)
+			[this, PassParams, bSlowPath](FRHICommandList& RHICmdList)
 			{
 				// Can't use parallel command list because we have to make sure we are not building BVH data
 				// on the same RTGeometry on multiple threads at the same time. Ideally move the build
@@ -947,10 +960,17 @@ void FRayTracingDynamicGeometryUpdateManager::AddDynamicGeometryUpdatePass(FRDGB
 				// with other build requests in the engine (see UE-106982)
 				SCOPED_DRAW_EVENT(RHICmdList, Build);
 
-				FRHIBufferRange ScratchBufferRange;
-				ScratchBufferRange.Buffer = PassParams->DynamicGeometryScratchBuffer ? PassParams->DynamicGeometryScratchBuffer->GetRHI() : nullptr;
-				ScratchBufferRange.Offset = 0;
-				RHICmdList.BuildAccelerationStructures(BuildParams, ScratchBufferRange);
+				if (!bSlowPath)
+				{
+					FRHIBufferRange ScratchBufferRange;
+					ScratchBufferRange.Buffer = PassParams->DynamicGeometryScratchBuffer ? PassParams->DynamicGeometryScratchBuffer->GetRHI() : nullptr;
+					ScratchBufferRange.Offset = 0;
+					RHICmdList.BuildAccelerationStructures(BuildParams, ScratchBufferRange);
+				}
+				else
+				{
+					RHICmdList.BuildAccelerationStructures(BuildParams);
+				}
 			});
 	}
 
