[Stochastic Lighting] Only bind resources relevant to the specific MaterialSource to avoid bugs due to unintended data usage.

#Synced-CL 50921461

diff --git a/Engine/Source/Runtime/Renderer/Private/StochasticLighting/StochasticLighting.cpp b/Engine/Source/Runtime/Renderer/Private/StochasticLighting/StochasticLighting.cpp
--- a/Engine/Source/Runtime/Renderer/Private/StochasticLighting/StochasticLighting.cpp
+++ b/Engine/Source/Runtime/Renderer/Private/StochasticLighting/StochasticLighting.cpp
@@ -463,9 +463,24 @@ void StochasticLighting::FContext::Run(const FRunConfig& RunConfig)
 
 	FStochasticLightingTileClassificationMarkCS::FParameters* PassParameters = GraphBuilder.AllocParameters<FStochasticLightingTileClassificationMarkCS::FParameters>();
 	PassParameters->View = View.ViewUniformBuffer;
-	PassParameters->SceneTextures = GetSceneTextureParameters(GraphBuilder, SceneTextures.UniformBuffer);
-	PassParameters->SceneTexturesStruct = SceneTextures.UniformBuffer;
-	PassParameters->FrontLayerTranslucencyGBufferParameters = FrontLayerTranslucencyGBuffer;
+	if (MaterialSource == EMaterialSource::GBuffer)
+	{
+		PassParameters->SceneTextures = GetSceneTextureParameters(GraphBuilder, SceneTextures.UniformBuffer);
+		PassParameters->SceneTexturesStruct = SceneTextures.UniformBuffer;
+	}
+	else if (MaterialSource == EMaterialSource::HairStrands)
+	{
+		PassParameters->HairStrands = HairStrands::BindHairStrandsViewUniformParameters(View);
+		PassParameters->SceneTextures.GBufferVelocityTexture = (*SceneTextures.UniformBuffer)->GBufferVelocityTexture;
+	}
+	else if (MaterialSource == EMaterialSource::FrontLayerGBuffer)
+	{
+		PassParameters->FrontLayerTranslucencyGBufferParameters = FrontLayerTranslucencyGBuffer;
+	}
+	else
+	{
+		unimplemented();
+	}
 	PassParameters->Substrate = Substrate::BindSubstrateGlobalUniformParameters(View);
 	PassParameters->DepthHistoryTexture = DepthHistoryTexture;
 	PassParameters->NormalAndShadingInfoHistory = NormalAndShadingInfoHistory;
@@ -495,7 +510,6 @@ void StochasticLighting::FContext::Run(const FRunConfig& RunConfig)
 	PassParameters->StochasticLightingStateFrameIndex = StateFrameIndex;
 	PassParameters->OutputDownsampledDepthAndNormal = RunConfig.bOutputDownsampledDepthAndNormal ? 1 : 0;
 	PassParameters->ForwardLightStruct = View.ForwardLightingResources.ForwardLightUniformBuffer;
-	PassParameters->HairStrands = HairStrands::BindHairStrandsViewUniformParameters(View);
 	PassParameters->BlueNoise = BlueNoiseUniformBuffer;
 
 	if (bTileClassifySubstrate)
