Fix Stencil LOD dithering on mobile (when r.StencilForLODDither=1).

Non-dithering primitives could overwrite the stencil dither pattern written in DitheredStencilPrePass before the PrePass. This caused bad dither transitions for non-instanced primitives.

Fix the issue by not writing to STENCIL_SANDBOX_MASK for non-dithering primitives. This preserves the dither pattern, as well as other stencil bits (RECEIVE_DECAL, lighting channels, etc.).

#tests Editor preview mode + console

#Synced-CL 50723010

diff --git a/Engine/Source/Runtime/Renderer/Private/DepthRendering.cpp b/Engine/Source/Runtime/Renderer/Private/DepthRendering.cpp
--- a/Engine/Source/Runtime/Renderer/Private/DepthRendering.cpp
+++ b/Engine/Source/Runtime/Renderer/Private/DepthRendering.cpp
@@ -798,7 +798,8 @@ bool FDepthPassMeshProcessor::Process(
 				const FMaterial & Material,
 				FMaterialShadingModelField ShadingModels,
 				bool bUsesDeferredShading,
-				TOptional<uint32> DitheringFlag);
+				TOptional<uint32> DitheringFlag,
+				bool bKeepStencilDither = false);
 
 			// *Don't* get shading models from MaterialResource since it's for a default material
 			FMaterialShadingModelField ShadingModels = MeshBatch.MaterialRenderProxy->GetIncompleteMaterialWithFallback(ERHIFeatureLevel::ES3_1).GetShadingModels();
@@ -809,7 +810,8 @@ bool FDepthPassMeshProcessor::Process(
 				MaterialResource,
 				ShadingModels,
 				bUsesDeferredShading,
-				DitheringFlag);
+				DitheringFlag,
+				true);
 		}
 		else if (DitheringFlag.IsSet())
 		{
diff --git a/Engine/Source/Runtime/Renderer/Private/MobileBasePass.cpp b/Engine/Source/Runtime/Renderer/Private/MobileBasePass.cpp
--- a/Engine/Source/Runtime/Renderer/Private/MobileBasePass.cpp
+++ b/Engine/Source/Runtime/Renderer/Private/MobileBasePass.cpp
@@ -87,9 +87,11 @@ void SetMobileBasePassDepthState(
 	const FMaterial& Material, 
 	FMaterialShadingModelField ShadingModels, 
 	bool bUsesDeferredShading,
-	TOptional<uint32> DitheringFlag)
+	TOptional<uint32> DitheringFlag,
+	bool bKeepStencilDither = false)
 {
 	uint8 StencilValue = 0u;
+	
 	if (DitheringFlag.IsSet())
 	{
 		StencilValue = DitheringFlag.GetValue();
@@ -101,6 +103,16 @@ void SetMobileBasePassDepthState(
 			STENCIL_SANDBOX_MASK, 0xFF
 		>::GetRHI());
 	}
+	else if (bKeepStencilDither && IsStencilForLODDitherEnabled(Material.GetShaderPlatform()))
+	{
+		// If the object is not transitioning and StencilLODDither is enabled, we need to preserve the dither pattern, 
+		// as well as the other stencil bits (RECEIVE_DECAL, lighting channels, etc.).
+		DrawRenderState.SetDepthStencilState(TStaticDepthStencilState<
+			true, CF_DepthNearOrEqual,
+			true, CF_Always, SO_Keep, SO_Keep, SO_Replace,
+			false, CF_Always, SO_Keep, SO_Keep, SO_Keep,
+			0x00, (0xff & ~STENCIL_SANDBOX_MASK) >::GetRHI());
+	}
 	else
 	{
 		DrawRenderState.SetDepthStencilState(TStaticDepthStencilState<
diff --git a/Engine/Source/Runtime/Renderer/Private/VelocityRendering.cpp b/Engine/Source/Runtime/Renderer/Private/VelocityRendering.cpp
--- a/Engine/Source/Runtime/Renderer/Private/VelocityRendering.cpp
+++ b/Engine/Source/Runtime/Renderer/Private/VelocityRendering.cpp
@@ -1238,7 +1238,8 @@ bool FVelocityMeshProcessor::Process(
 			const FMaterial& Material, 
 			FMaterialShadingModelField ShadingModels, 
 			bool bUsesDeferredShading,
-			TOptional<uint32> DitheringFlag);
+			TOptional<uint32> DitheringFlag, 
+			bool bKeepStencilDither = false);
 
 		// *Don't* get shading models from MaterialResource since it's for a default material
 		FMaterialShadingModelField ShadingModels = MeshBatch.MaterialRenderProxy->GetIncompleteMaterialWithFallback(ERHIFeatureLevel::ES3_1).GetShadingModels();
@@ -1249,7 +1250,8 @@ bool FVelocityMeshProcessor::Process(
 			MaterialResource, 
 			ShadingModels, 
 			bUsesDeferredShading, 
-			GetOptionalDitheringFlag(ViewIfDynamicMeshCommand, MeshBatch, StaticMeshId));
+			GetOptionalDitheringFlag(ViewIfDynamicMeshCommand, MeshBatch, StaticMeshId),
+			true);
 	}
 
 	FMeshMaterialShaderElementData ShaderElementData;
