Compile fix on Mac

#Synced-CL 50723010

diff --git a/Engine/Source/Runtime/Renderer/Private/DepthRendering.cpp b/Engine/Source/Runtime/Renderer/Private/DepthRendering.cpp
--- a/Engine/Source/Runtime/Renderer/Private/DepthRendering.cpp
+++ b/Engine/Source/Runtime/Renderer/Private/DepthRendering.cpp
@@ -799,7 +799,7 @@ bool FDepthPassMeshProcessor::Process(
 				FMaterialShadingModelField ShadingModels,
 				bool bUsesDeferredShading,
 				TOptional<uint32> DitheringFlag,
-				bool bKeepStencilDither = false);
+				bool bKeepStencilDither);
 
 			// *Don't* get shading models from MaterialResource since it's for a default material
 			FMaterialShadingModelField ShadingModels = MeshBatch.MaterialRenderProxy->GetIncompleteMaterialWithFallback(ERHIFeatureLevel::ES3_1).GetShadingModels();
diff --git a/Engine/Source/Runtime/Renderer/Private/MobileBasePass.cpp b/Engine/Source/Runtime/Renderer/Private/MobileBasePass.cpp
--- a/Engine/Source/Runtime/Renderer/Private/MobileBasePass.cpp
+++ b/Engine/Source/Runtime/Renderer/Private/MobileBasePass.cpp
@@ -88,7 +88,7 @@ void SetMobileBasePassDepthState(
 	FMaterialShadingModelField ShadingModels, 
 	bool bUsesDeferredShading,
 	TOptional<uint32> DitheringFlag,
-	bool bKeepStencilDither = false)
+	bool bKeepStencilDither)
 {
 	uint8 StencilValue = 0u;
 	
@@ -498,7 +498,7 @@ void MobileBasePass::SetOpaqueRenderState(FMeshPassProcessorRenderState& DrawRen
 {
 	if (bCanUseDepthStencil)
 	{
-		SetMobileBasePassDepthState(DrawRenderState, PrimitiveSceneProxy, Material, ShadingModels, bUsesDeferredShading, TOptional<uint32>());
+		SetMobileBasePassDepthState(DrawRenderState, PrimitiveSceneProxy, Material, ShadingModels, bUsesDeferredShading, TOptional<uint32>(), false);
 	}
 	else
 	{
diff --git a/Engine/Source/Runtime/Renderer/Private/VelocityRendering.cpp b/Engine/Source/Runtime/Renderer/Private/VelocityRendering.cpp
--- a/Engine/Source/Runtime/Renderer/Private/VelocityRendering.cpp
+++ b/Engine/Source/Runtime/Renderer/Private/VelocityRendering.cpp
@@ -1239,7 +1239,7 @@ bool FVelocityMeshProcessor::Process(
 			FMaterialShadingModelField ShadingModels, 
 			bool bUsesDeferredShading,
 			TOptional<uint32> DitheringFlag, 
-			bool bKeepStencilDither = false);
+			bool bKeepStencilDither);
 
 		// *Don't* get shading models from MaterialResource since it's for a default material
 		FMaterialShadingModelField ShadingModels = MeshBatch.MaterialRenderProxy->GetIncompleteMaterialWithFallback(ERHIFeatureLevel::ES3_1).GetShadingModels();
