Move StochasticLighting TileClassificationMark to async compute when Lumen is using IrradianceFieldGather async compute. Saves .51ms on Switch 2 at 1062p.

#Synced-CL 50792079

diff --git a/Engine/Source/Runtime/Renderer/Private/StochasticLighting/StochasticLighting.cpp b/Engine/Source/Runtime/Renderer/Private/StochasticLighting/StochasticLighting.cpp
--- a/Engine/Source/Runtime/Renderer/Private/StochasticLighting/StochasticLighting.cpp
+++ b/Engine/Source/Runtime/Renderer/Private/StochasticLighting/StochasticLighting.cpp
@@ -20,7 +20,7 @@ static TAutoConsoleVariable<int32> CVarStochasticLightingAsyncCompute(
 	1,
 	TEXT("Whether to run the tile classification mark pass on async compute:\n")
 	TEXT("  0 - Disabled\n")
-	TEXT("  1 - Enable only when MegaLights sample generation is async\n")
+	TEXT("  1 - Enable only when downstream features (MegaLights, Lumen) are async compute\n")
 	TEXT("  2 - Enabled always"),
 	ECVF_Scalability | ECVF_RenderThreadSafe);
 
@@ -789,17 +789,6 @@ bool FDeferredShadingSceneRenderer::RequiresStochasticLightingPass()
  */
 void FDeferredShadingSceneRenderer::StochasticLightingTileClassificationMark(FRDGBuilder& GraphBuilder, FLumenSceneFrameTemporaries& FrameTemporaries, const FSceneTextures& SceneTextures, bool bAsyncMegaLightsGenerateSamples)
 {
-	ERDGPassFlags ComputePassFlags = ERDGPassFlags::Compute;
-	const int32 AsyncComputeMode = FMath::Clamp(CVarStochasticLightingAsyncCompute.GetValueOnRenderThread(), 0, 2);
-	if (AsyncComputeMode == 1)
-	{
-		ComputePassFlags = bAsyncMegaLightsGenerateSamples ? ERDGPassFlags::AsyncCompute : ERDGPassFlags::Compute;
-	}
-	else if (AsyncComputeMode == 2)
-	{
-		ComputePassFlags = ERDGPassFlags::AsyncCompute;
-	}
-
 	if (Views.IsEmpty())
 	{
 		return;
@@ -832,6 +821,28 @@ void FDeferredShadingSceneRenderer::StochasticLightingTileClassificationMark(FRD
 				|| (bTileClassifyMegaLights && MegaLightsDownsampleFactor == FIntPoint(2, 2));
 
 			StochasticLighting::FRunConfig& RunConfig = RunConfigs.Emplace_GetRef(View);
+
+			ERDGPassFlags ComputePassFlags = ERDGPassFlags::Compute;
+			{
+				const int32 AsyncComputeMode = FMath::Clamp(CVarStochasticLightingAsyncCompute.GetValueOnRenderThread(), 0, 2);
+				if (AsyncComputeMode == 1)
+				{
+					if (bTileClassifyMegaLights)
+					{
+						ComputePassFlags = bAsyncMegaLightsGenerateSamples ? ERDGPassFlags::AsyncCompute : ERDGPassFlags::Compute;
+					}
+					else if (bLumenIrradianceFieldGather)
+					{
+						const bool bLumenFinalGatherUseAsyncCompute = LumenDiffuseIndirect::UseAsyncCompute(ViewFamily, ViewPipelineState.DiffuseIndirectMethod);
+						ComputePassFlags = bLumenFinalGatherUseAsyncCompute ? ERDGPassFlags::AsyncCompute : ERDGPassFlags::Compute;
+					}
+				}
+				else if (AsyncComputeMode == 2)
+				{
+					ComputePassFlags = ERDGPassFlags::AsyncCompute;
+				}
+			}
+
 			RunConfig.ComputePassFlags = ComputePassFlags;
 			RunConfig.ReflectionsMethod = ViewPipelineState.ReflectionsMethod;
 			RunConfig.bCopyDepthAndNormal = bCopyDepthAndNormal;
