Lumen Irradiance Field Gather skips Closest HZB preparation. Saves .1ms on Switch 2 docked.

#Synced-CL 50792079

diff --git a/Engine/Source/Runtime/Renderer/Private/DeferredShadingRenderer.cpp b/Engine/Source/Runtime/Renderer/Private/DeferredShadingRenderer.cpp
--- a/Engine/Source/Runtime/Renderer/Private/DeferredShadingRenderer.cpp
+++ b/Engine/Source/Runtime/Renderer/Private/DeferredShadingRenderer.cpp
@@ -1330,6 +1330,12 @@ void FDeferredShadingSceneRenderer::CommitFinalPipelineState()
 			const bool bHasSSGI = ViewPipelineState[&FPerViewPipelineState::DiffuseIndirectMethod] == EDiffuseIndirectMethod::SSGI;
 			const bool bUseLumen = ViewPipelineState[&FPerViewPipelineState::DiffuseIndirectMethod] == EDiffuseIndirectMethod::Lumen 
 				|| ViewPipelineState[&FPerViewPipelineState::ReflectionsMethod] == EReflectionsMethod::Lumen;
+
+			const ELumenFinalGatherMethod LumenFinalGatherMethod = Lumen::GetFinalGatherMethod(ViewFamily, ShaderPlatform);
+			// Only HZB traversal needs Closest HZB
+			const bool bLumenNeedsClosestHZB = (ViewPipelineState[&FPerViewPipelineState::DiffuseIndirectMethod] == EDiffuseIndirectMethod::Lumen && LumenFinalGatherMethod == ELumenFinalGatherMethod::ScreenProbeGather)
+				|| ViewPipelineState[&FPerViewPipelineState::ReflectionsMethod] == EReflectionsMethod::Lumen;
+
 			const bool bHasFirstPersonSelfShadow = ShouldRenderFirstPersonSelfShadow(ViewFamily);
 
 			// Requires FurthestHZB
@@ -1341,7 +1347,7 @@ void FDeferredShadingSceneRenderer::CommitFinalPipelineState()
 				bHasSSGI || bUseLumen);
 
 			ViewPipelineState.Set(&FPerViewPipelineState::bClosestHZB, 
-				bHasSSGI || bUseLumen || bHasFirstPersonSelfShadow || MegaLights::IsUsingClosestHZB(ViewFamily));
+				bHasSSGI || bLumenNeedsClosestHZB || bHasFirstPersonSelfShadow || MegaLights::IsUsingClosestHZB(ViewFamily));
 		}
 	}
 
