[Stochastic Lighting] Fix ReflectionsMethod not being set in StochasticLighting::FRunConfig for StochasticLightingFrontLayerTranslucencyTileClassificationMark.
- doesn't really affect anything at the moment but still important to avoid issues later on.

#Synced-CL 50994078

diff --git a/Engine/Source/Runtime/Renderer/Private/StochasticLighting/StochasticLighting.cpp b/Engine/Source/Runtime/Renderer/Private/StochasticLighting/StochasticLighting.cpp
--- a/Engine/Source/Runtime/Renderer/Private/StochasticLighting/StochasticLighting.cpp
+++ b/Engine/Source/Runtime/Renderer/Private/StochasticLighting/StochasticLighting.cpp
@@ -916,6 +916,7 @@ void FDeferredShadingSceneRenderer::StochasticLightingFrontLayerTranslucencyTile
 
 		StochasticLighting::FRunConfig& RunConfig = RunConfigs.Emplace_GetRef(View);
 		RunConfig.ComputePassFlags = ComputePassFlags;
+		RunConfig.ReflectionsMethod = ReflectionsMethod;
 		RunConfig.bCopyDepthAndNormal = bCopyDepthAndNormal;
 		RunConfig.bDownsampleDepthAndNormal2x1 = bDownsampleDepthAndNormal2x1;
 		RunConfig.bDownsampleDepthAndNormal2x2 = bDownsampleDepthAndNormal2x2;
