Material Cache, render composition in batches when exceeding the max number of nanite views or layers

#Synced-CL 51014304

diff --git a/Engine/Source/Runtime/Renderer/Private/MaterialCache/MaterialCacheRenderer.cpp b/Engine/Source/Runtime/Renderer/Private/MaterialCache/MaterialCacheRenderer.cpp
--- a/Engine/Source/Runtime/Renderer/Private/MaterialCache/MaterialCacheRenderer.cpp
+++ b/Engine/Source/Runtime/Renderer/Private/MaterialCache/MaterialCacheRenderer.cpp
@@ -630,9 +630,14 @@ static void GetMaterialCacheDefaultMaterials(const FPrimitiveSceneProxy* Proxy,
 	}
 }
 
-static void MaterialCacheAllocateAndBatchPages(FSceneRendererBase* Renderer, FRDGBuilder& GraphBuilder, FMaterialCacheSceneExtension& SceneExtension, FMaterialCacheRenderData& RenderData)
+static uint32 MaterialCacheAllocateAndBatchPages(FSceneRendererBase* Renderer, FRDGBuilder& GraphBuilder, FMaterialCacheSceneExtension& SceneExtension, FMaterialCacheRenderData& RenderData, int32 BatchStart)
 {
-	for (int32 EntryIndex = 0; EntryIndex < RenderData.Bucket->PendingEntries.Num(); EntryIndex++)
+	int32 ProducedPages = 0;
+	
+	// Each batch is limited by the layers for hardware pages, and nanite views for nanite pages
+	const int32 MaxProducedPages = FMath::Min(GRHIGlobals.MaxTextureArrayLayers, NANITE_MAX_VIEWS_PER_CULL_RASTERIZE_PASS);
+	
+	for (int32 EntryIndex = BatchStart; EntryIndex < RenderData.Bucket->PendingEntries.Num(); EntryIndex++)
 	{
 		FMaterialCachePendingEntry& Entry = RenderData.Bucket->PendingEntries[EntryIndex];
 
@@ -658,6 +663,21 @@ static void MaterialCacheAllocateAndBatchPages(FSceneRendererBase* Renderer, FRD
 		{
 			continue;
 		}
+		
+		// Shouldn't happen
+		if (Entry.Pages.Num() > MaxProducedPages)
+		{
+			checkf(false, TEXT("Entry page count exceeded internal limit"));
+			continue;
+		}
+		
+		// Out of pages? Resume at the current index
+		if (ProducedPages + Entry.Pages.Num() > MaxProducedPages)
+		{
+			return EntryIndex;
+		}
+		
+		ProducedPages += Entry.Pages.Num();
 
 		for (FMaterialCachePendingPageEntry& Page : Entry.Pages)
 		{
@@ -728,6 +748,9 @@ static void MaterialCacheAllocateAndBatchPages(FSceneRendererBase* Renderer, FRD
 		CreatePageIndirectionBuffer(GraphBuilder, LayerRenderData.Nanite.GenericCSBatch);
 		CreatePageIndirectionBuffer(GraphBuilder, LayerRenderData.VertexInvariant.GenericCSBatch);
 	}
+	
+	// Last batch completed
+	return RenderData.Bucket->PendingEntries.Num();
 }
 
 static FIntPoint GetMaterialCacheTileSize()
@@ -1810,19 +1833,23 @@ void MaterialCacheRenderPages(FRDGBuilder& GraphBuilder)
 	for (auto&& [Guid, Bucket] : SceneExtension.TagBuckets)
 	{
 		RDG_EVENT_SCOPE(GraphBuilder, "MaterialCache");
+		
+		for (int32_t Offset = 0; Offset < Bucket.PendingEntries.Num();)
+		{
+			// Create render data on the graph's lifetime
+			FMaterialCacheRenderData& RenderData = *GraphBuilder.AllocObject<FMaterialCacheRenderData>();
+			RenderData.Bucket = &Bucket;
 
-		// Create render data on the graph's lifetime
-		FMaterialCacheRenderData& RenderData = *GraphBuilder.AllocObject<FMaterialCacheRenderData>();
-		RenderData.Bucket = &Bucket;
-
-		// First, allocate and batch all pages
-		MaterialCacheAllocateAndBatchPages(Renderer, GraphBuilder, SceneExtension, RenderData);
+			// First, allocate and batch all pages
+			Offset = MaterialCacheAllocateAndBatchPages(Renderer, GraphBuilder, SceneExtension, RenderData, Offset);
 
-		// Then, render them with their allotted layers
-		if (!RenderData.ABuffer.Pages.IsEmpty())
-		{
-			MaterialCacheRenderLayers(Renderer, GraphBuilder, SceneExtension, RenderData);
+			// Then, render them with their allotted layers
+			if (!RenderData.ABuffer.Pages.IsEmpty())
+			{
+				MaterialCacheRenderLayers(Renderer, GraphBuilder, SceneExtension, RenderData);
+			}
 		}
+
 	}
 
 	SceneExtension.TagBuckets.Empty();
