Fix crash in FFrontLayerTranslucencyGBufferMeshProcessor::CollectPSOInitializers.
- get ShaderPlatform from SceneTexturesConfig instead of FScene*.

#Synced-CL 51016705

#Synced-CL 51016705

diff --git a/Engine/Source/Runtime/Renderer/Private/FrontLayerTranslucency.cpp b/Engine/Source/Runtime/Renderer/Private/FrontLayerTranslucency.cpp
--- a/Engine/Source/Runtime/Renderer/Private/FrontLayerTranslucency.cpp
+++ b/Engine/Source/Runtime/Renderer/Private/FrontLayerTranslucency.cpp
@@ -307,7 +307,7 @@ void FFrontLayerTranslucencyGBufferMeshProcessor::CollectPSOInitializers(const F
 		return;
 	}
 
-	const bool bMegaLightsSupported = ShouldCompileMegaLightsShaders(Scene->GetShaderPlatform());
+	const bool bMegaLightsSupported = ShouldCompileMegaLightsShaders(SceneTexturesConfig.ShaderPlatform);
 
 	FGraphicsPipelineRenderTargetsInfo RenderTargetsInfo;
 	RenderTargetsInfo.NumSamples = 1;
