Support precreated SceneProxy when adding them to the scene.

#Synced-CL 51009600

diff --git a/Engine/Source/Runtime/Renderer/Private/RendererScene.cpp b/Engine/Source/Runtime/Renderer/Private/RendererScene.cpp
--- a/Engine/Source/Runtime/Renderer/Private/RendererScene.cpp
+++ b/Engine/Source/Runtime/Renderer/Private/RendererScene.cpp
@@ -1438,9 +1438,8 @@ void FScene::BatchAddPrimitivesInternal(TArrayView<T*> InPrimitives)
 
 		FPrimitiveSceneProxy* PrimitiveSceneProxy  = nullptr;
 
-		if (Primitive->GetPrimitiveComponentInterface())
+		if (Primitive->GetPrimitiveComponentInterface() && !Primitive->GetSceneProxy())
 		{
-			checkf(!Primitive->GetSceneProxy(), TEXT("Primitive has already been added to the scene!"));
 			PrimitiveSceneProxy = Primitive->GetPrimitiveComponentInterface()->CreateSceneProxy();
 			check(SceneData.SceneProxy == PrimitiveSceneProxy); // CreateSceneProxy has access to the shared SceneData and should set it properly
 		}
