[Stochastic Lighting] Don't compile unecessary permutations of FStochasticLightingTileClassificationMarkCS based on whether Lumen or MegaLights are supported.

#Synced-CL 51016705

diff --git a/Engine/Source/Runtime/Renderer/Private/StochasticLighting/StochasticLighting.cpp b/Engine/Source/Runtime/Renderer/Private/StochasticLighting/StochasticLighting.cpp
--- a/Engine/Source/Runtime/Renderer/Private/StochasticLighting/StochasticLighting.cpp
+++ b/Engine/Source/Runtime/Renderer/Private/StochasticLighting/StochasticLighting.cpp
@@ -240,6 +240,16 @@ class FStochasticLightingTileClassificationMarkCS : public FGlobalShader
 			return false;
 		}
 
+		if (PermutationVector.Get<FTileClassifyLumen>() && !DoesPlatformSupportLumenGI(Parameters.Platform))
+		{
+			return false;
+		}
+
+		if (PermutationVector.Get<FTileClassifyMegaLights>() && !MegaLights::ShouldCompileShaders(Parameters.Platform))
+		{
+			return false;
+		}
+
 		return DoesPlatformSupportLumenGI(Parameters.Platform) || MegaLights::ShouldCompileShaders(Parameters.Platform);
 	}
 
