Add TONEMAPPING_METHOD permutation to CombineLUTs to try to reduce shader complexity

Also wrapped filmic HDR logic with !OUTPUT_DEVICE_SRGB to avoid that code.

#Synced-CL 51003583

diff --git a/Engine/Shaders/Private/PostProcessCombineLUTs.usf b/Engine/Shaders/Private/PostProcessCombineLUTs.usf
--- a/Engine/Shaders/Private/PostProcessCombineLUTs.usf
+++ b/Engine/Shaders/Private/PostProcessCombineLUTs.usf
@@ -319,7 +319,6 @@ float3 ComputeInvOutputACES(float3 ODTColor, float3x3 LimitingRgbToXYZ, float3x3
 /** Values of output device (matching C++'s ETonemappingMethod) */
 #define TONEMAPPING_METHOD_Filmic 0
 #define TONEMAPPING_METHOD_ACES 1
-uint TonemappingMethod;
 
 #define TONEMAPPING_SDR_INVEOTF_SRGB 0
 #define TONEMAPPING_SDR_INVEOTF_BT1886 1
@@ -327,44 +326,43 @@ uint TonemappingSDRInvEOTF;
 
 float3 ApplyToneCurve(float3 ColorAP1, float PeakLuminance, float3x3 AP1_2_Output)
 {
-	float3 OutputGamutColor;
-	if (TonemappingMethod == TONEMAPPING_METHOD_ACES)
+#if TONEMAPPING_METHOD == TONEMAPPING_METHOD_ACES
+	float3 GradedColor = ComputeGradedColor(ColorAP1);
+	float3 OutputGamutColor = ComputeOutputACES(GradedColor, PeakLuminance, GetLimitingRgbToXYZ(), GetLimitingXYZToRgb(), AP1_2_Output) * LinearToNitsScaleInverse;
+#elif TONEMAPPING_METHOD == TONEMAPPING_METHOD_Filmic
+	float3 FilmColorNoGamma = ComputeFilmColorNoGamma(ColorAP1);
+
+	// Apply "gamma" curve adjustment.
+	float3 FilmColor = pow( max(0, FilmColorNoGamma), InverseGamma.y );
+
+	// FIXME: Workaround for UE-29935, pushing all colors with a 0 component to black output
+	// Default parameters seem to cancel out (sRGB->XYZ->AP1->XYZ->sRGB), so should be okay for a temp fix
+	float3 OutputGamutColor = WorkingColorSpace.bIsSRGB && (OutputGamut == TONEMAPPER_GAMUT_sRGB_D65) ?
+		FilmColor :
+		mul( AP1_2_Output, mul( (float3x3)WorkingColorSpace.ToAP1, FilmColor ) );
+
+#if !OUTPUT_DEVICE_SRGB
+	// "Automatic HDR" using ACES to remap filmic.
+	if (PeakLuminance > LinearToNitsScale)
 	{
-		float3 GradedColor = ComputeGradedColor(ColorAP1);
-		OutputGamutColor = ComputeOutputACES(GradedColor, PeakLuminance, GetLimitingRgbToXYZ(), GetLimitingXYZToRgb(), AP1_2_Output) * LinearToNitsScaleInverse;
-	}
-	// TonemappingMethod == TONEMAPPING_METHOD_Filmic
-	else
-	{
-		float3 FilmColorNoGamma = ComputeFilmColorNoGamma(ColorAP1);
-
-		// Apply "gamma" curve adjustment.
-		float3 FilmColor = pow( max(0, FilmColorNoGamma), InverseGamma.y );
-
-		// FIXME: Workaround for UE-29935, pushing all colors with a 0 component to black output
-		// Default parameters seem to cancel out (sRGB->XYZ->AP1->XYZ->sRGB), so should be okay for a temp fix
-		OutputGamutColor = WorkingColorSpace.bIsSRGB && (OutputGamut == TONEMAPPER_GAMUT_sRGB_D65) ?
-			FilmColor :
-			mul( AP1_2_Output, mul( (float3x3)WorkingColorSpace.ToAP1, FilmColor ) );
-
-		// "Automatic HDR" using ACES to remap filmic.
-		if (PeakLuminance > LinearToNitsScale)
+		// Rec. 2100 specifies a dim environment similar to various viewing standards, e.g. BT.2035, for consuming Rec. 709 w/BT.1886.
+		// sRGB is meant to be viewed in a bright office, so its EOTF leaves values a bit inflated in comparison to BT.1886.
+		// This adjustment converts from "bright" linear space to "dim" linear space.
+		if (TonemappingSDRInvEOTF == TONEMAPPING_SDR_INVEOTF_SRGB)
 		{
-			// Rec. 2100 specifies a dim environment similar to various viewing standards, e.g. BT.2035, for consuming Rec. 709 w/BT.1886.
-			// sRGB is meant to be viewed in a bright office, so its EOTF leaves values a bit inflated in comparison to BT.1886.
-			// This adjustment converts from "bright" linear space to "dim" linear space.
-			if (TonemappingSDRInvEOTF == TONEMAPPING_SDR_INVEOTF_SRGB)
-			{
-				// Apply inverse of sRGB EOTF to generate nonlinear values that would get sent to an SDR display.
-				OutputGamutColor = LinearToSrgb(OutputGamutColor);
-				// Apply BT.1886 EOTF to linearize the values into "dim" linear space.
-				OutputGamutColor = pow(OutputGamutColor, Bt1886Gamma);
-			}
-
-			float3 GradedColor = ComputeInvOutputACES(OutputGamutColor * LinearToNitsScale, (float3x3)WorkingColorSpace.ToXYZ, (float3x3)WorkingColorSpace.FromXYZ);
-			OutputGamutColor = ComputeOutputACES(GradedColor, PeakLuminance, GetLimitingRgbToXYZ(), GetLimitingXYZToRgb(), AP1_2_Output) * LinearToNitsScaleInverse;
+			// Apply inverse of sRGB EOTF to generate nonlinear values that would get sent to an SDR display.
+			OutputGamutColor = LinearToSrgb(OutputGamutColor);
+			// Apply BT.1886 EOTF to linearize the values into "dim" linear space.
+			OutputGamutColor = pow(OutputGamutColor, Bt1886Gamma);
 		}
+
+		float3 GradedColor = ComputeInvOutputACES(OutputGamutColor * LinearToNitsScale, (float3x3)WorkingColorSpace.ToXYZ, (float3x3)WorkingColorSpace.FromXYZ);
+		OutputGamutColor = ComputeOutputACES(GradedColor, PeakLuminance, GetLimitingRgbToXYZ(), GetLimitingXYZToRgb(), AP1_2_Output) * LinearToNitsScaleInverse;
 	}
+#endif
+#else
+#error Unknown TONEMAPPING_METHOD
+#endif
 
 	return OutputGamutColor;
 }
diff --git a/Engine/Source/Runtime/Renderer/Private/PostProcess/PostProcessCombineLUTs.cpp b/Engine/Source/Runtime/Renderer/Private/PostProcess/PostProcessCombineLUTs.cpp
--- a/Engine/Source/Runtime/Renderer/Private/PostProcess/PostProcessCombineLUTs.cpp
+++ b/Engine/Source/Runtime/Renderer/Private/PostProcess/PostProcessCombineLUTs.cpp
@@ -184,7 +184,6 @@ BEGIN_SHADER_PARAMETER_STRUCT(FCombineLUTParameters, )
 	SHADER_PARAMETER(float, BlueCorrection)
 	SHADER_PARAMETER(float, ExpandGamut)
 	SHADER_PARAMETER(float, ToneCurveAmount)
-	SHADER_PARAMETER(uint32, TonemappingMethod)
 	SHADER_PARAMETER(uint32, TonemappingSDRInvEOTF)
 	SHADER_PARAMETER(float, FilmSlope)
 	SHADER_PARAMETER(float, FilmToe)
@@ -239,15 +238,21 @@ struct FCachedLUTSettings
 	FCombineLUTParameters Parameters;
 	FWorkingColorSpaceShaderParameters WorkingColorSpaceShaderParameters;
 	int32 AcesVersion = 0;
+	ETonemappingMethod TonemappingMethod{};
 
-	bool UpdateCachedValues(FRDGBuilder& GraphBuilder, const FViewInfo& View, const FTexture* const* Textures, const float* Weights, uint32 BlendCount)
+	bool UpdateCachedValues(FRDGBuilder& GraphBuilder, const FViewInfo& View, const FPostProcessSettings& Settings, const FTexture* const* Textures, const float* Weights, uint32 BlendCount)
 	{
 		bool bHasChanged = false;
-		GetCombineLUTParameters(GraphBuilder, View, Textures, Weights, BlendCount, bHasChanged);
+		GetCombineLUTParameters(GraphBuilder, View, Settings, Textures, Weights, BlendCount, bHasChanged);
 		UPDATE_CACHE_SETTINGS(UniqueID, View.State ? View.State->GetViewKey() : 0, bHasChanged);
 		UPDATE_CACHE_SETTINGS(ShaderPlatform, View.GetShaderPlatform(), bHasChanged);
 		static const auto CVarAcesVersion = IConsoleManager::Get().FindTConsoleVariableDataInt(TEXT("r.HDR.Aces.Version"));
 		UPDATE_CACHE_SETTINGS(AcesVersion, CVarAcesVersion->GetValueOnRenderThread(), bHasChanged);
+		const int32 TonemappingMethodOverride = CVarLUTTonemappingMethod.GetValueOnRenderThread();
+		const ETonemappingMethod NextTonemappingMethod = (TonemappingMethodOverride >= 0) ?
+			static_cast<ETonemappingMethod>(TonemappingMethodOverride) :
+			Settings.TonemappingMethod;
+		UPDATE_CACHE_SETTINGS(TonemappingMethod, NextTonemappingMethod, bHasChanged);
 
 		const FWorkingColorSpaceShaderParameters* InWorkingColorSpaceShaderParameters = reinterpret_cast<const FWorkingColorSpaceShaderParameters*>(GDefaultWorkingColorSpaceUniformBuffer.GetContents());
 		if (InWorkingColorSpaceShaderParameters)
@@ -266,6 +271,7 @@ struct FCachedLUTSettings
 	void GetCombineLUTParameters(
 		FRDGBuilder& GraphBuilder,
 		const FViewInfo& View,
+		const FPostProcessSettings& Settings,
 		const FTexture* const* Textures,
 		const float* Weights,
 		uint32 BlendCount,
@@ -274,12 +280,6 @@ struct FCachedLUTSettings
 		check(Textures);
 		check(Weights);
 
-		const FSceneViewFamily& ViewFamily = *(View.Family);
-		const FPostProcessSettings& Settings = (
-			ViewFamily.EngineShowFlags.ColorGrading && CVarColorGrading.GetValueOnRenderThread() != 0)
-			? View.FinalPostProcessSettings
-			: FPostProcessSettings::GetDefault();
-
 		for (uint32 BlendIndex = 0; BlendIndex < BlendCount; ++BlendIndex)
 		{
 			// Neutral texture occupies the first slot and doesn't actually need to be set.
@@ -297,6 +297,7 @@ struct FCachedLUTSettings
 
 		Parameters.WorkingColorSpace = GDefaultWorkingColorSpaceUniformBuffer.GetUniformBufferRef();
 
+		const FSceneViewFamily& ViewFamily = *(View.Family);
 		FTonemapperOutputDeviceParameters TonemapperOutputDeviceParameters = GetTonemapperOutputDeviceParameters(ViewFamily);
 		UPDATE_CACHE_SETTINGS(Parameters.OutputDevice.InverseGamma, TonemapperOutputDeviceParameters.InverseGamma, bHasChanged);
 		UPDATE_CACHE_SETTINGS(Parameters.OutputDevice.OutputDevice, TonemapperOutputDeviceParameters.OutputDevice, bHasChanged);
@@ -387,12 +388,6 @@ struct FCachedLUTSettings
 		UPDATE_CACHE_SETTINGS(Parameters.ExpandGamut, Settings.ExpandGamut, bHasChanged);
 		UPDATE_CACHE_SETTINGS(Parameters.ToneCurveAmount, Settings.ToneCurveAmount, bHasChanged);
 
-		const int32 TonemappingMethodOverride = CVarLUTTonemappingMethod.GetValueOnRenderThread();
-		const uint32 TonemappingMethod = (TonemappingMethodOverride >= 0) ?
-			static_cast<uint32>(TonemappingMethodOverride) :
-			static_cast<uint32>(Settings.TonemappingMethod);
-		UPDATE_CACHE_SETTINGS(Parameters.TonemappingMethod, TonemappingMethod, bHasChanged);
-
 		const int32 TonemappingSDRInvEOTFOverride = CVarLUTTonemappingSDRInvEOTF.GetValueOnRenderThread();
 		const uint32 TonemappingSDRInvEOTF = (TonemappingSDRInvEOTFOverride >= 0) ?
 			static_cast<uint32>(TonemappingSDRInvEOTFOverride) :
@@ -415,7 +410,8 @@ public:
 	class FBlendCount : SHADER_PERMUTATION_RANGE_INT("BLENDCOUNT", 1, 5);
 	class FOutputDeviceSRGB : SHADER_PERMUTATION_BOOL("OUTPUT_DEVICE_SRGB");
 	class FUseACES2 : SHADER_PERMUTATION_BOOL("USE_ACES_2");
-	using FPermutationDomain = TShaderPermutationDomain<FBlendCount, FOutputDeviceSRGB, FUseACES2>;
+	class FTonemappingMethod : SHADER_PERMUTATION_ENUM_CLASS("TONEMAPPING_METHOD", ETonemappingMethod);
+	using FPermutationDomain = TShaderPermutationDomain<FBlendCount, FOutputDeviceSRGB, FUseACES2, FTonemappingMethod>;
 
 	DECLARE_GLOBAL_SHADER(FLUTBlenderCS);
 	SHADER_USE_PARAMETER_STRUCT(FLUTBlenderCS, FGlobalShader);
@@ -572,7 +568,12 @@ FRDGTextureRef AddCombineLUTPass(FRDGBuilder& GraphBuilder, const FViewInfo& Vie
 
 	View.SetValidTonemappingLUT();
 
-	const bool bHasChanged = CachedLUTSettings.UpdateCachedValues(GraphBuilder, View, LocalTextures, LocalWeights, LocalCount);
+	const FPostProcessSettings& Settings = (
+		ViewFamily.EngineShowFlags.ColorGrading && CVarColorGrading.GetValueOnRenderThread() != 0)
+		? View.FinalPostProcessSettings
+		: FPostProcessSettings::GetDefault();
+
+	const bool bHasChanged = CachedLUTSettings.UpdateCachedValues(GraphBuilder, View, Settings, LocalTextures, LocalWeights, LocalCount);
 	if (!bHasChanged && OutputTexture && CVarLUTUpdateEveryFrame.GetValueOnRenderThread() == 0)
 	{
 		return OutputTexture;
@@ -591,6 +592,7 @@ FRDGTextureRef AddCombineLUTPass(FRDGBuilder& GraphBuilder, const FViewInfo& Vie
 	FLUTBlenderCS::FPermutationDomain PermutationVector;
 	PermutationVector.Set<FLUTBlenderCS::FBlendCount>(LocalCount);
 	PermutationVector.Set<FLUTBlenderCS::FUseACES2>(CachedLUTSettings.AcesVersion > 1);
+	PermutationVector.Set<FLUTBlenderCS::FTonemappingMethod>(CachedLUTSettings.TonemappingMethod);
 
 	FLUTBlenderCS::FParameters* PassParameters = GraphBuilder.AllocParameters<FLUTBlenderCS::FParameters>();
 	PassParameters->CombineLUT = CachedLUTSettings.Parameters;
