* Fix hair strands custom depth which were using previous frame data when rendering custom  data before base pass.
* CustomDepth is now faster by reusing primary visiblity rasterization output
* Unify HitProxy, CustomDepth, SelectionOutline into a single shader.

#Synced-CL 51018413

diff --git a/Engine/Shaders/Private/HairStrands/HairStrandsHitProxy.usf b/Engine/Shaders/Private/HairStrands/HairStrandsHitProxy.usf
--- a/Engine/Shaders/Private/HairStrands/HairStrandsHitProxy.usf
+++ b/Engine/Shaders/Private/HairStrands/HairStrandsHitProxy.usf
@@ -5,96 +5,73 @@ float3 DecodeVelocityFromTexture(float4 EncodedV);
 #include "HairStrandsVisibilityCommon.ush" 
 
 ///////////////////////////////////////////////////////////////////////////////////////////////////
+// Used for HitProxy, Selection, and CustomDepth
 
-#if SHADER_HITPROXY_ID
+#if SHADER_EMIT_DEPTH
 
+uint							 bDiscardIfZero;
 float 							 CoverageThreshold;
+float2							 InvViewportResolution;
 uint							 MaxMaterialCount;
-Buffer<uint>					 MaterialIdToHitProxyIdBuffer;
+Buffer<uint>					 ActiveMaterialIdBuffer;
 Texture2D<uint>					 VisNodeIndex;
 StructuredBuffer<FPackedHairVis> VisNodeData;
 Texture2D<float>				 CoverageTexture;
+Texture2D<float>				 HairOnlyDepthTexture;
 
 void EmitPS(
+	#if PERMUTATION_TILE
 	in FScreenVertexOutput In,
+	#else
+	in float4 InPosition : SV_POSITION,
+	#endif
+	#if PERMUTATION_COLOR
 	out float4 OutHitProxyId : SV_Target0,
+	#endif
 	out float OutDepth : SV_Depth)
 {
+	#if PERMUTATION_COLOR
 	OutHitProxyId = 0;
+	#endif
 	OutDepth = 0;
+
+	bool bDiscard = true;
+
+	#if PERMUTATION_TILE
 	const uint2 PixelCoord = (uint2)In.Position.xy;
+	const bool bValid = true;
+	#else
+	const float2 UV = InPosition.xy * InvViewportResolution;
+	const uint2 PixelCoord = UV * View.ViewSizeAndInvSize.xy;
+	const bool bValid = HairOnlyDepthTexture.Load(uint3(PixelCoord, 0)) > 0;
+	#endif
 
-	// Take the first sample
-	const FNodeDesc NodeDesc = DecodeNodeDesc(VisNodeIndex.Load(uint3(PixelCoord,0)));
+	const FNodeDesc NodeDesc = DecodeNodeDesc(VisNodeIndex.Load(uint3(PixelCoord, 0)));
 	const float Coverage = CoverageTexture.Load(uint3(PixelCoord, 0));
-	if (NodeDesc.Count > 0 && Coverage >= CoverageThreshold)
+	if (bValid && NodeDesc.Count > 0 && Coverage >= CoverageThreshold)
 	{
+		// Take the first sample
 		const uint SampleIndex = NodeDesc.Offset + 0;
 		const FHairVis InNode  = UnpackHairVis(VisNodeData[SampleIndex]);
 		const uint MaterialId  = clamp(InNode.MaterialId, 0, MaxMaterialCount-1);
-		const uint MaterialHitProxyId = MaterialIdToHitProxyIdBuffer[MaterialId];
+		const uint MaterialIdValue = ActiveMaterialIdBuffer[MaterialId];
 
-		float4 HitProxyId = float4(
-			float((MaterialHitProxyId >> 16) & 0xFF) / 255.0f,
-			float((MaterialHitProxyId >> 8) & 0xFF) / 255.0f,
-			float((MaterialHitProxyId >> 0) & 0xFF) / 255.0f,
+		#if PERMUTATION_COLOR
+		OutHitProxyId = float4(
+			float((MaterialIdValue >> 16) & 0xFF) / 255.0f,
+			float((MaterialIdValue >> 8) & 0xFF) / 255.0f,
+			float((MaterialIdValue >> 0) & 0xFF) / 255.0f,
 			0);
-		OutHitProxyId = HitProxyId;
+		#endif
 		OutDepth = InNode.Depth;
-	}
-	else
-	{
-		discard;
-	}
-}
-
-#endif // SHADER_HITPROXY_ID
-
-///////////////////////////////////////////////////////////////////////////////////////////////////
-
-#if SHADER_SELECTION
 
-float							 CoverageThreshold;
-float2							 InvViewportResolution;
-uint							 MaxMaterialCount;
-Buffer<uint>					 SelectionMaterialIdBuffer;
-Texture2D<uint>					 VisNodeIndex;
-StructuredBuffer<FPackedHairVis> VisNodeData;
-Texture2D<float>				 CoverageTexture;
-Texture2D<float>				 HairOnlyDepthTexture;
-
-void EmitPS(
-	in float4 InPosition : SV_POSITION,
-	out float OutDepth : SV_Depth)
-{
-	OutDepth = 0;
-	const float2 UV = InPosition.xy * InvViewportResolution;
-	const uint2 PixelCoord = UV * View.ViewSizeAndInvSize.xy;
-	const bool bValid = HairOnlyDepthTexture.Load(uint3(PixelCoord, 0)) > 0;
-
-	bool bIsSelected = false;
-	if (bValid)
-	{
-		// Take the first sample
-		const FNodeDesc NodeDesc = DecodeNodeDesc(VisNodeIndex.Load(uint3(PixelCoord, 0)));
-		if (NodeDesc.Count > 0)
-		{
-			const float Coverage = CoverageTexture.Load(uint3(PixelCoord, 0));
-			const uint SampleIndex = NodeDesc.Offset + 0;
-			const FHairVis InNode = UnpackHairVis(VisNodeData[SampleIndex]);
-			const uint MaterialId = clamp(InNode.MaterialId, 0, MaxMaterialCount - 1);
-			if (SelectionMaterialIdBuffer[MaterialId] > 0 && Coverage >= CoverageThreshold)
-			{
-				bIsSelected = true;
-				OutDepth = InNode.Depth;
-			}
-		}
+		bDiscard = bDiscardIfZero > 0 ? MaterialIdValue == 0 : false;
 	}
 
-	if (!bIsSelected)
+	if (bDiscard)
 	{
 		discard;
 	}
 }
 
-#endif // SHADER_SELECTION
\ No newline at end of file
+#endif // SHADER_EMIT_DEPTH
diff --git a/Engine/Source/Runtime/Renderer/Private/HairStrands/HairStrandsData.h b/Engine/Source/Runtime/Renderer/Private/HairStrands/HairStrandsData.h
--- a/Engine/Source/Runtime/Renderer/Private/HairStrands/HairStrandsData.h
+++ b/Engine/Source/Runtime/Renderer/Private/HairStrands/HairStrandsData.h
@@ -487,6 +487,7 @@ namespace HairStrands
 	bool HasPositionsChanged(FRDGBuilder& GraphBuilder, const FScene& Scene, const FViewInfo& View);
 	void DrawHitProxies(FRDGBuilder& GraphBuilder, const FScene& Scene, const FViewInfo& View, FInstanceCullingManager& InstanceCullingManager, FRDGTextureRef HitProxyTexture, FRDGTextureRef HitProxyDepthTexture);
 	void DrawEditorSelection(FRDGBuilder& GraphBuilder, const FViewInfo& View, const FIntRect& ViewportRect, FRDGTextureRef SelectionDepthTexture);
+	void DrawCustomDepth(FRDGBuilder& GraphBuilder, const FViewInfo& View, const FIntRect& ViewportRect, FRDGTextureRef CustomDepthTexture);
 
 	// Mesh batch helpers
 	FHairGroupPublicData* GetHairData(const FMeshBatch* In);
diff --git a/Engine/Source/Runtime/Renderer/Private/HairStrands/HairStrandsRendering.cpp b/Engine/Source/Runtime/Renderer/Private/HairStrands/HairStrandsRendering.cpp
--- a/Engine/Source/Runtime/Renderer/Private/HairStrands/HairStrandsRendering.cpp
+++ b/Engine/Source/Runtime/Renderer/Private/HairStrands/HairStrandsRendering.cpp
@@ -238,6 +238,15 @@ void RenderHairBasePass(
 			View.HairStrandsViewData.UniformBuffer = InternalCreateHairStrandsViewUniformBuffer(GraphBuilder, nullptr);
 			View.HairStrandsViewData.bIsValid = false;
 		}
+
+		// Late composition of hair into to custom depth (avoid extra rasterization by reusing primary raster outputs)
+		if (SceneTextures.CustomDepth.Depth && EnumHasAllFlags(SceneTextures.SetupMode, ESceneTextureSetupMode::CustomDepth))
+		{
+			if (HairStrands::HasViewHairStrandsData(View))
+			{
+				HairStrands::DrawCustomDepth(GraphBuilder, View, View.ViewRect, SceneTextures.CustomDepth.Depth);
+			}
+		}
 	}
 }
 
diff --git a/Engine/Source/Runtime/Renderer/Private/HairStrands/HairStrandsVisibility.cpp b/Engine/Source/Runtime/Renderer/Private/HairStrands/HairStrandsVisibility.cpp
--- a/Engine/Source/Runtime/Renderer/Private/HairStrands/HairStrandsVisibility.cpp
+++ b/Engine/Source/Runtime/Renderer/Private/HairStrands/HairStrandsVisibility.cpp
@@ -114,7 +114,7 @@ static int32 GHairStrands_InvalidationPosition_Debug = 0;
 static FAutoConsoleVariableRef CVarHairStrands_InvalidationPosition_Debug(TEXT("r.HairStrands.PathTracing.InvalidationDebug"), GHairStrands_InvalidationPosition_Debug, TEXT("Enable bounding box drawing for groom element causing path tracer invalidation"));
 
 static float GHairStrands_Selection_CoverageThreshold = 0.0f;
-static FAutoConsoleVariableRef CVarHairStrands_Selection_CoverageThreshold(TEXT("r.HairStrands.Selection.CoverageThreshold"), GHairStrands_Selection_CoverageThreshold, TEXT("Coverage threshold for making hair strands outline selection finer"));
+static FAutoConsoleVariableRef CVarHairStrands_Selection_CoverageThreshold(TEXT("r.HairStrands.Selection.CoverageThreshold"), GHairStrands_Selection_CoverageThreshold, TEXT("Coverage threshold for making hair strands outline selection and custom depth finer"));
 
 static int32 GHairStrandsUseHWRaster = 0;
 static FAutoConsoleVariableRef CVarHairStrandsUseHWRaster(TEXT("r.HairStrands.Visibility.UseHWRaster"), GHairStrandsUseHWRaster, TEXT("Toggles the hardware rasterizer for hair strands visibility rendering."));
@@ -4082,23 +4082,30 @@ static FRasterComputeOutput AddVisibilityComputeRasterPass(
 }
 
 ///////////////////////////////////////////////////////////////////////////////////////////////////
-// Selection outline
+// Selection outline, HitProxyId, CustomDepth
 
-class FHairStrandsEmitSelectionPS : public FGlobalShader
+class FHairStrandsEmitDepthPS : public FGlobalShader
 {
-	DECLARE_GLOBAL_SHADER(FHairStrandsEmitSelectionPS);
-	SHADER_USE_PARAMETER_STRUCT(FHairStrandsEmitSelectionPS, FGlobalShader);
+	DECLARE_GLOBAL_SHADER(FHairStrandsEmitDepthPS);
+	SHADER_USE_PARAMETER_STRUCT(FHairStrandsEmitDepthPS, FGlobalShader);
+
+	class FColor : SHADER_PERMUTATION_BOOL("PERMUTATION_COLOR");
+	class FTile : SHADER_PERMUTATION_BOOL("PERMUTATION_TILE");
+	using FPermutationDomain = TShaderPermutationDomain<FColor, FTile>;
 
 	BEGIN_SHADER_PARAMETER_STRUCT(FParameters, )
 		SHADER_PARAMETER(float, CoverageThreshold)
 		SHADER_PARAMETER(FVector2f, InvViewportResolution)
 		SHADER_PARAMETER(uint32, MaxMaterialCount)
+		SHADER_PARAMETER(uint32, bHasValidHairOnlyDepth)
+		SHADER_PARAMETER(uint32, bDiscardIfZero)
 		SHADER_PARAMETER_STRUCT_REF(FViewUniformShaderParameters, View)
 		SHADER_PARAMETER_RDG_TEXTURE(Texture2D, VisNodeIndex)
 		SHADER_PARAMETER_RDG_TEXTURE(Texture2D<float>, CoverageTexture)
 		SHADER_PARAMETER_RDG_TEXTURE(Texture2D<float>, HairOnlyDepthTexture)
 		SHADER_PARAMETER_RDG_BUFFER_SRV(StructuredBuffer<FPackedHairVis>, VisNodeData)
-		SHADER_PARAMETER_RDG_BUFFER_SRV(Buffer<uint>, SelectionMaterialIdBuffer)
+		SHADER_PARAMETER_RDG_BUFFER_SRV(Buffer<uint>, ActiveMaterialIdBuffer)
+		SHADER_PARAMETER_STRUCT_INCLUDE(FHairStrandsTilePassVS::FParameters, TileData)
 		RENDER_TARGET_BINDING_SLOTS()
 	END_SHADER_PARAMETER_STRUCT()
 
@@ -4110,172 +4117,148 @@ class FHairStrandsEmitSelectionPS : public FGlobalShader
 	static void ModifyCompilationEnvironment(const FGlobalShaderPermutationParameters& Parameters, FShaderCompilerEnvironment& OutEnvironment)
 	{
 		FGlobalShader::ModifyCompilationEnvironment(Parameters, OutEnvironment);
-		OutEnvironment.SetDefine(TEXT("SHADER_SELECTION"), 1);
+		OutEnvironment.SetDefine(TEXT("SHADER_EMIT_DEPTH"), 1);
 	}
 };
-IMPLEMENT_GLOBAL_SHADER(FHairStrandsEmitSelectionPS, "/Engine/Private/HairStrands/HairStrandsHitProxy.usf", "EmitPS", SF_Pixel);
+IMPLEMENT_GLOBAL_SHADER(FHairStrandsEmitDepthPS, "/Engine/Private/HairStrands/HairStrandsHitProxy.usf", "EmitPS", SF_Pixel);
+
+enum class EHairStrandsEmitDepthMode : uint8
+{
+	Selection,
+	HitProxy,
+	CustomDepth,
+};
 
-void AddHairStrandsSelectionOutlinePass(
+static void AddHairStrandsEmitDepthPass(
 	FRDGBuilder& GraphBuilder,
 	const FViewInfo& View,
 	const FIntRect& ViewportRect,
 	const FHairStrandsMacroGroupDatas& MacroGroupDatas,
-	FHairStrandsTiles TileData,
+	const EHairStrandsEmitDepthMode Mode,
+	const FHairStrandsTiles& TileData,
 	FRDGTextureRef VisNodeIndex,
 	FRDGBufferRef VisNodeData,
 	FRDGTextureRef CoverageTexture,
 	FRDGTextureRef HairOnlyDepthTexture,
-	FRDGTextureRef SelectionDepthTexture)
+	FRDGTextureRef OutHitProxyTexture,
+	FRDGTextureRef OutEmitDepthTexture)
 {
 	if (View.HairStrandsMeshElements.Num() == 0 || !VisNodeData)
 	{
 		return;
 	}
 
-#if WITH_EDITOR
 	// Create mapping table between MaterialId and BatchId
-	TArray<uint32> SelectionMaterialId;
-	SelectionMaterialId.Init(0u, View.HairStrandsMeshElements.Num());
+	TArray<uint32> ActiveMaterialId;
+	ActiveMaterialId.Init(0u, View.HairStrandsMeshElements.Num());
 	for (const FHairStrandsMacroGroupData& MacroGroupData : MacroGroupDatas)
 	{
 		for (const FHairStrandsMacroGroupData::PrimitiveInfo& PrimitiveInfo : MacroGroupData.PrimitivesInfos)
 		{
-			if (SelectionMaterialId.IsValidIndex(PrimitiveInfo.MaterialId) && PrimitiveInfo.PrimitiveSceneProxy)
+			if (ActiveMaterialId.IsValidIndex(PrimitiveInfo.MaterialId) && PrimitiveInfo.PrimitiveSceneProxy)
 			{
-				const uint32 bSelected = PrimitiveInfo.PrimitiveSceneProxy->IsSelected() ? 1u : 0u;
-				SelectionMaterialId[PrimitiveInfo.MaterialId] = bSelected;
+				uint32 Value = 0u;
+				switch (Mode)
+				{
+					case EHairStrandsEmitDepthMode::Selection   : Value = PrimitiveInfo.PrimitiveSceneProxy->IsSelected() ? 1u : 0u; break;
+					case EHairStrandsEmitDepthMode::CustomDepth : Value = PrimitiveInfo.PrimitiveSceneProxy->ShouldRenderCustomDepth() ? 1u : 0u; break;
+					case EHairStrandsEmitDepthMode::HitProxy    : Value = PrimitiveInfo.Mesh ? PrimitiveInfo.Mesh->BatchHitProxyId.GetColor().DWColor() : 0u; break;
+				}
+				ActiveMaterialId[PrimitiveInfo.MaterialId] = Value;
 			}
 		}
 	}
-
-	const FHairStrandsTiles::ETileType TileType = FHairStrandsTiles::ETileType::HairAll;
-
-	FRDGBufferRef SelectionMaterialIdBuffer = CreateUploadBuffer(GraphBuilder, TEXT("Hair.MaterialIdToHitProxyIdBuffer"), sizeof(uint32), SelectionMaterialId.Num(), SelectionMaterialId.GetData(), sizeof(uint32) * SelectionMaterialId.Num());
-	auto* PassParameters = GraphBuilder.AllocParameters<FHairStrandsEmitSelectionPS::FParameters>();
+	
+	FRDGBufferRef ActiveMaterialIdBuffer = CreateUploadBuffer(GraphBuilder, TEXT("Hair.ActiveMaterialIdBuffer"), sizeof(uint32), ActiveMaterialId.Num(), ActiveMaterialId.GetData(), sizeof(uint32) * ActiveMaterialId.Num());
+	auto* PassParameters = GraphBuilder.AllocParameters<FHairStrandsEmitDepthPS::FParameters>();
 	PassParameters->View = View.ViewUniformBuffer;
 	PassParameters->CoverageThreshold = FMath::Clamp(GHairStrands_Selection_CoverageThreshold, 0.f, 1.f);
-	PassParameters->MaxMaterialCount = SelectionMaterialId.Num();
+	PassParameters->MaxMaterialCount = ActiveMaterialId.Num();
 	PassParameters->InvViewportResolution = FVector2f(1.f/ViewportRect.Width(), 1.f/ViewportRect.Height());
+	PassParameters->bDiscardIfZero = Mode != EHairStrandsEmitDepthMode::HitProxy ? 1u : 0u;
 	PassParameters->VisNodeIndex = VisNodeIndex;
 	PassParameters->VisNodeData = GraphBuilder.CreateSRV(VisNodeData);
 	PassParameters->CoverageTexture = CoverageTexture;
 	PassParameters->HairOnlyDepthTexture = HairOnlyDepthTexture;
-	PassParameters->SelectionMaterialIdBuffer = GraphBuilder.CreateSRV(SelectionMaterialIdBuffer, PF_R32_UINT);
-	PassParameters->RenderTargets.DepthStencil = FDepthStencilBinding(SelectionDepthTexture, ERenderTargetLoadAction::ELoad, FExclusiveDepthStencil::DepthWrite_StencilWrite);
-
-	const FIntRect Viewport = View.ViewRect;
-	auto PixelShader = View.ShaderMap->GetShader<FHairStrandsEmitSelectionPS>();
-
-	// We don't use tile rendering for hair selection, because the outline buffer is unscaled, and does not match the visibility buffer which is unscaled
-	const uint32 StencilRef = 3;
-	FPixelShaderUtils::AddFullscreenPass(
-		GraphBuilder,
-		View.ShaderMap,
-		RDG_EVENT_NAME("HairStrands::EmitSelection"),
-		PixelShader,
-		PassParameters,
-		ViewportRect,
-		TStaticBlendState<>::GetRHI(),
-		TStaticRasterizerState<>::GetRHI(),
-		TStaticDepthStencilState<true, CF_DepthNearOrEqual, true, CF_Always, SO_Keep, SO_Keep, SO_Replace>::GetRHI(), 
-		StencilRef);
-
-#endif
-}
-
-///////////////////////////////////////////////////////////////////////////////////////////////////
-// HitProxyId
-
-class FHairStrandsEmitHitProxyIdPS : public FGlobalShader
-{
-	DECLARE_GLOBAL_SHADER(FHairStrandsEmitHitProxyIdPS);
-	SHADER_USE_PARAMETER_STRUCT(FHairStrandsEmitHitProxyIdPS, FGlobalShader);
-
-	BEGIN_SHADER_PARAMETER_STRUCT(FParameters, )
-		SHADER_PARAMETER(float, CoverageThreshold)
-		SHADER_PARAMETER(uint32, MaxMaterialCount)
-		SHADER_PARAMETER_STRUCT_REF(FViewUniformShaderParameters, View)
-		SHADER_PARAMETER_RDG_TEXTURE(Texture2D, VisNodeIndex)
-		SHADER_PARAMETER_RDG_BUFFER_SRV(StructuredBuffer<FPackedHairVis>, VisNodeData)
-		SHADER_PARAMETER_RDG_BUFFER_SRV(Buffer<uint>, MaterialIdToHitProxyIdBuffer)
-		SHADER_PARAMETER_RDG_TEXTURE(Texture2D<float>, CoverageTexture)
-		SHADER_PARAMETER_STRUCT_INCLUDE(FHairStrandsTilePassVS::FParameters, TileData)
-		RENDER_TARGET_BINDING_SLOTS()
-	END_SHADER_PARAMETER_STRUCT()
-
-	static bool ShouldCompilePermutation(const FGlobalShaderPermutationParameters& Parameters)
+	PassParameters->ActiveMaterialIdBuffer = GraphBuilder.CreateSRV(ActiveMaterialIdBuffer, PF_R32_UINT);
+	if (PassParameters->HairOnlyDepthTexture == nullptr)
 	{
-		return IsHairStrandsSupported(EHairStrandsShaderType::Strands, Parameters.Platform);
+		PassParameters->HairOnlyDepthTexture = GSystemTextures.GetDepthDummy(GraphBuilder);
 	}
-
-	static void ModifyCompilationEnvironment(const FGlobalShaderPermutationParameters& Parameters, FShaderCompilerEnvironment& OutEnvironment)
+	
+	bool bUseTile = false;
+	FRHIDepthStencilState* DepthStencilState = nullptr;
+	uint32 StencilRef = 0;
+	const TCHAR* PassName = TEXT("Unknown");
+	switch (Mode)
 	{
-		FGlobalShader::ModifyCompilationEnvironment(Parameters, OutEnvironment);
-		OutEnvironment.SetDefine(TEXT("SHADER_HITPROXY_ID"), 1);
+		case EHairStrandsEmitDepthMode::Selection:
+		{
+			PassParameters->RenderTargets.DepthStencil = FDepthStencilBinding(OutEmitDepthTexture, ERenderTargetLoadAction::ELoad, FExclusiveDepthStencil::DepthWrite_StencilWrite);
+			DepthStencilState = TStaticDepthStencilState<true, CF_DepthNearOrEqual, true, CF_Always, SO_Keep, SO_Keep, SO_Replace>::GetRHI(); 
+			StencilRef = 3;
+			PassName = TEXT("Selection");
+			bUseTile = false;
+		}
+		break;
+		case EHairStrandsEmitDepthMode::CustomDepth:
+		{
+			PassParameters->RenderTargets.DepthStencil = FDepthStencilBinding(OutEmitDepthTexture, ERenderTargetLoadAction::ELoad, FExclusiveDepthStencil::DepthWrite_StencilNop);
+			DepthStencilState = TStaticDepthStencilState<true, CF_DepthNearOrEqual, false>::GetRHI(); 
+			StencilRef = 0xFF;
+			PassName = TEXT("CustomDepth");
+			bUseTile = true;
+		}
+		break;
+		case EHairStrandsEmitDepthMode::HitProxy:
+		{
+			PassParameters->RenderTargets[0] = FRenderTargetBinding(OutHitProxyTexture, ERenderTargetLoadAction::ELoad);
+			PassParameters->RenderTargets.DepthStencil = FDepthStencilBinding(OutEmitDepthTexture, ERenderTargetLoadAction::ELoad, FExclusiveDepthStencil::DepthWrite_StencilWrite);
+			DepthStencilState = TStaticDepthStencilState<true, CF_DepthNearOrEqual>::GetRHI();
+			StencilRef = 0xFF;
+			PassName = TEXT("HitProxy");
+			bUseTile = true;
+		}
+		break;
+		default: check(false);
 	}
-};
-IMPLEMENT_GLOBAL_SHADER(FHairStrandsEmitHitProxyIdPS, "/Engine/Private/HairStrands/HairStrandsHitProxy.usf", "EmitPS", SF_Pixel);
 
-void AddHairStrandsHitProxyIdPass(
-	FRDGBuilder& GraphBuilder,
-	const FScene& Scene,
-	const FViewInfo& View,
-	const FHairStrandsMacroGroupDatas& MacroGroupDatas,
-	FHairStrandsTiles TileData,
-	FRDGTextureRef VisNodeIndex,
-	FRDGBufferRef VisNodeData,
-	FRDGTextureRef CoverageTexture,
-	FRDGTextureRef HitProxyTexture,
-	FRDGTextureRef HitProxyDepthTexture)
-{
-#if WITH_EDITOR
-	if (View.HairStrandsMeshElements.Num() == 0 || !VisNodeData)
+	const FHairStrandsTiles::ETileType TileType = FHairStrandsTiles::ETileType::HairAll;
+	if (bUseTile)
 	{
-		return;
+		PassParameters->TileData = GetHairStrandsTileParameters(View, TileData, TileType);
+		check(TileData.IsValid());
 	}
 
-	// Create mapping table between MaterialId and BatchId
-	TArray<uint32> MaterialIdToHitProxyId;
-	MaterialIdToHitProxyId.Init(0u, View.HairStrandsMeshElements.Num());
-	for (const FHairStrandsMacroGroupData& MacroGroupData : MacroGroupDatas)
+	FHairStrandsEmitDepthPS::FPermutationDomain PermutationVector;
+	PermutationVector.Set<FHairStrandsEmitDepthPS::FColor>(Mode == EHairStrandsEmitDepthMode::HitProxy);
+	PermutationVector.Set<FHairStrandsEmitDepthPS::FTile>(bUseTile);
+	auto PixelShader = View.ShaderMap->GetShader<FHairStrandsEmitDepthPS>(PermutationVector);
+	if (!bUseTile)
 	{
-		for (const FHairStrandsMacroGroupData::PrimitiveInfo& PrimitiveInfo : MacroGroupData.PrimitivesInfos)
-		{
-			if (MaterialIdToHitProxyId.IsValidIndex(PrimitiveInfo.MaterialId))
-			{
-				if (const FMeshBatch* MeshBatch = PrimitiveInfo.Mesh)
-				{
-					const uint32 HitColor = MeshBatch->BatchHitProxyId.GetColor().DWColor();
-					MaterialIdToHitProxyId[PrimitiveInfo.MaterialId] = HitColor;
-				}
-			}
-		}
+		// We don't use tile rendering for hair selection, because the outline buffer is unscaled, and does not match the visibility buffer which is unscaled
+		FPixelShaderUtils::AddFullscreenPass(
+			GraphBuilder,
+			View.ShaderMap,
+			RDG_EVENT_NAME("HairStrands::EmitDepth(%s)", PassName),
+			PixelShader,
+			PassParameters,
+			ViewportRect,
+			TStaticBlendState<>::GetRHI(),
+			TStaticRasterizerState<>::GetRHI(),
+			DepthStencilState,
+			StencilRef);
 	}
+	else
+	{
+		const FIntRect Viewport = View.ViewRect;
+		TShaderMapRef<FHairStrandsTilePassVS> TileVertexShader(View.ShaderMap);
 
-	const FHairStrandsTiles::ETileType TileType = FHairStrandsTiles::ETileType::HairAll;
-
-	FRDGBufferRef MaterialIdToHitProxyIdBuffer = CreateUploadBuffer(GraphBuilder, TEXT("Hair.MaterialIdToHitProxyIdBuffer"), sizeof(uint32), MaterialIdToHitProxyId.Num(), MaterialIdToHitProxyId.GetData(), sizeof(uint32) * MaterialIdToHitProxyId.Num());
-	auto* PassParameters = GraphBuilder.AllocParameters<FHairStrandsEmitHitProxyIdPS::FParameters>();
-	PassParameters->View = View.ViewUniformBuffer;
-	PassParameters->CoverageThreshold = FMath::Clamp(GHairStrands_Selection_CoverageThreshold, 0.f, 1.f);
-	PassParameters->MaxMaterialCount = MaterialIdToHitProxyId.Num();
-	PassParameters->VisNodeIndex = VisNodeIndex;
-	PassParameters->VisNodeData = GraphBuilder.CreateSRV(VisNodeData);
-	PassParameters->MaterialIdToHitProxyIdBuffer = GraphBuilder.CreateSRV(MaterialIdToHitProxyIdBuffer, PF_R32_UINT);
-	PassParameters->CoverageTexture = CoverageTexture;
-	PassParameters->RenderTargets[0] = FRenderTargetBinding(HitProxyTexture, ERenderTargetLoadAction::ELoad);
-	PassParameters->RenderTargets.DepthStencil = FDepthStencilBinding(HitProxyDepthTexture, ERenderTargetLoadAction::ELoad, FExclusiveDepthStencil::DepthWrite_StencilWrite);
-	PassParameters->TileData = GetHairStrandsTileParameters(View, TileData, TileType);
-
-	const FIntRect Viewport = View.ViewRect;
-	auto PixelShader = View.ShaderMap->GetShader<FHairStrandsEmitHitProxyIdPS>();
-	TShaderMapRef<FHairStrandsTilePassVS> TileVertexShader(View.ShaderMap);
-
-	GraphBuilder.AddPass(
-		RDG_EVENT_NAME("HairStrands::EmitHitProxyId(Tile)"),
-		PassParameters,
-		ERDGPassFlags::Raster,
-		[PassParameters, TileVertexShader, PixelShader, Viewport, TileType](FRDGAsyncTask, FRHICommandList& RHICmdList)
+		GraphBuilder.AddPass(
+			RDG_EVENT_NAME("HairStrands::EmitDepth(%s)", PassName),
+			PassParameters,
+			ERDGPassFlags::Raster,
+			[PassParameters, TileVertexShader, PixelShader, Viewport, TileType, DepthStencilState](FRDGAsyncTask, FRHICommandList& RHICmdList)
 		{
 			FHairStrandsTilePassVS::FParameters ParametersVS = PassParameters->TileData;
 
@@ -4283,7 +4266,7 @@ void AddHairStrandsHitProxyIdPass(
 			RHICmdList.ApplyCachedRenderTargets(GraphicsPSOInit);
 			GraphicsPSOInit.BlendState = TStaticBlendState<>::GetRHI();
 			GraphicsPSOInit.RasterizerState = TStaticRasterizerState<>::GetRHI();
-			GraphicsPSOInit.DepthStencilState = TStaticDepthStencilState<true, CF_DepthNearOrEqual>::GetRHI();
+			GraphicsPSOInit.DepthStencilState = DepthStencilState;
 
 			GraphicsPSOInit.BoundShaderState.VertexDeclarationRHI = GFilterVertexDeclaration.VertexDeclarationRHI;
 			GraphicsPSOInit.BoundShaderState.VertexShaderRHI = TileVertexShader.GetVertexShader();
@@ -4298,7 +4281,7 @@ void AddHairStrandsHitProxyIdPass(
 			RHICmdList.SetStreamSource(0, nullptr, 0);
 			RHICmdList.DrawPrimitiveIndirect(PassParameters->TileData.TileIndirectBuffer->GetRHI(), FHairStrandsTiles::GetIndirectDrawArgOffset(TileType));
 		});
-#endif
+	}
 }
 
 ///////////////////////////////////////////////////////////////////////////////////////////////////
@@ -4386,20 +4369,42 @@ static void AddHairStrandsHasPositionChangedPass(
 namespace HairStrands
 {
 
+// Draw hair strands custom depth value
+void DrawCustomDepth(FRDGBuilder& GraphBuilder, const FViewInfo& View, const FIntRect& ViewportRect, FRDGTextureRef CustomDepthTexture)
+{
+	AddHairStrandsEmitDepthPass(
+		GraphBuilder,
+		View,
+		ViewportRect,
+		View.HairStrandsViewData.MacroGroupDatas,
+		EHairStrandsEmitDepthMode::CustomDepth,
+		View.HairStrandsViewData.VisibilityData.TileData,
+		View.HairStrandsViewData.VisibilityData.NodeIndex,
+		View.HairStrandsViewData.VisibilityData.NodeVisData,
+		View.HairStrandsViewData.VisibilityData.CoverageTexture,
+		View.HairStrandsViewData.VisibilityData.HairOnlyDepthTexture,
+		nullptr,
+		CustomDepthTexture);
+}
+
 // Draw hair strands depth value for outline selection
 void DrawEditorSelection(FRDGBuilder& GraphBuilder, const FViewInfo& View, const FIntRect& ViewportRect, FRDGTextureRef SelectionDepthTexture)
 {
-	AddHairStrandsSelectionOutlinePass(
+	#if WITH_EDITOR
+	AddHairStrandsEmitDepthPass(
 		GraphBuilder,
 		View,
 		ViewportRect,
 		View.HairStrandsViewData.MacroGroupDatas,
+		EHairStrandsEmitDepthMode::Selection,
 		View.HairStrandsViewData.VisibilityData.TileData,
 		View.HairStrandsViewData.VisibilityData.NodeIndex,
 		View.HairStrandsViewData.VisibilityData.NodeVisData,
 		View.HairStrandsViewData.VisibilityData.CoverageTexture,
 		View.HairStrandsViewData.VisibilityData.HairOnlyDepthTexture,
+		nullptr,
 		SelectionDepthTexture);
+	#endif
 }
 
 // Draw hair strands hit proxy values
@@ -4582,7 +4587,19 @@ void DrawHitProxies(
 	}
 
 	// Make this tile dispatch
-	AddHairStrandsHitProxyIdPass(GraphBuilder, Scene, View, MacroGroupDatas, TileData, VisNodeIndex, VisNodeData, CoverageTexture, HitProxyTexture, HitProxyDepthTexture);
+	AddHairStrandsEmitDepthPass(
+		GraphBuilder,
+		View,
+		View.ViewRect,
+		MacroGroupDatas,
+		EHairStrandsEmitDepthMode::HitProxy,
+		TileData,
+		VisNodeIndex,
+		VisNodeData,
+		CoverageTexture,
+		View.HairStrandsViewData.VisibilityData.HairOnlyDepthTexture,
+		HitProxyTexture,
+		HitProxyDepthTexture);
 }
 
 // Check if any simulated/skinned-bound groom has its positions updated (e.g. for invalidating the path-tracer accumulation)
diff --git a/Engine/Source/Runtime/Renderer/Private/SceneVisibility.cpp b/Engine/Source/Runtime/Renderer/Private/SceneVisibility.cpp
--- a/Engine/Source/Runtime/Renderer/Private/SceneVisibility.cpp
+++ b/Engine/Source/Runtime/Renderer/Private/SceneVisibility.cpp
@@ -1554,6 +1554,9 @@ void FRelevancePacket::ComputeRelevance(FDynamicPrimitiveIndexList& DynamicPrimi
 		const bool bEditorSelectionRelevance = ViewRelevance.bEditorStaticSelectionRelevance;
 		const bool bTranslucentRelevance = ViewRelevance.HasTranslucency();
 		const bool bHairStrandsRelevance = bHairStrandsEnabled && ViewRelevance.bHairStrands;
+		// Filter out hair strands geometry, as custom depth is handled in a custom manner. This is only for hair strands, hair cards/meshes use the regular path. 
+		// Use !bRenderInMainPass as a discriminant for filtering only hair strands geometry.
+		const bool bIsCompatibleWithCustomDepth = !(bHairStrandsRelevance && !ViewRelevance.bRenderInMainPass);
 
 		if (View.bIsReflectionCapture && !PrimitiveSceneProxy->IsVisibleInReflectionCaptures())
 		{
@@ -1785,7 +1788,7 @@ void FRelevancePacket::ComputeRelevance(FDynamicPrimitiveIndexList& DynamicPrimi
 									DrawCommandPacket.AddCommandsForMesh(PrimitiveIndex, PrimitiveSceneInfo, StaticMeshRelevance, StaticMesh, CullingPayloadFlags, Scene, bCanCache, EMeshPass::AnisotropyPass);
 								}
 
-								if (ViewRelevance.bRenderCustomDepth)
+								if (ViewRelevance.bRenderCustomDepth && bIsCompatibleWithCustomDepth)
 								{
 									DrawCommandPacket.AddCommandsForMesh(PrimitiveIndex, PrimitiveSceneInfo, StaticMeshRelevance, StaticMesh, CullingPayloadFlags, Scene, bCanCache, EMeshPass::CustomDepth);
 								}
@@ -2133,7 +2136,7 @@ void FRelevancePacket::ComputeRelevance(FDynamicPrimitiveIndexList& DynamicPrimi
 		bHasSingleLayerWaterMaterial |= ViewRelevance.bUsesSingleLayerWaterMaterial;
 		bUsesSecondStageDepthPass |= ViewRelevance.bRenderInSecondStageDepthPass && ShadingPath!=EShadingPath::Mobile;
 
-		if (ViewRelevance.bRenderCustomDepth)
+		if (ViewRelevance.bRenderCustomDepth && bIsCompatibleWithCustomDepth)
 		{
 			bHasCustomDepthPrimitives = true;
 			CustomDepthStencilValues.Add(PrimitiveSceneInfo->Proxy->GetCustomDepthStencilValue());
@@ -2453,6 +2456,9 @@ static void ComputeDynamicMeshRelevance(
 {
 	const int32 NumElements = MeshBatch.Mesh->Elements.Num();
 	const bool bIsTranslucentClippedDepthEnabled = NeedTSRTranslucentClippedDepthPass(View);
+	// Filter out hair strands geometry, as custom depth is handled in a custom manner. This is only for hair strands, hair cards/meshes use the regular path. 
+	// Use !bRenderInMainPass as a discriminant for filtering only hair strands geometry.
+	const bool bIsCompatibleWithCustomDepth = !(ViewRelevance.bHairStrands && !ViewRelevance.bRenderInMainPass);
 
 	if (ViewRelevance.bDrawRelevance && (ViewRelevance.bRenderInMainPass || ViewRelevance.bRenderCustomDepth || ViewRelevance.bRenderInDepthPass))
 	{
@@ -2490,7 +2496,7 @@ static void ComputeDynamicMeshRelevance(
 				View.NumVisibleDynamicMeshElements[EMeshPass::MobileBasePassCSM] += NumElements;
 			}
 
-			if (ViewRelevance.bRenderCustomDepth)
+			if (ViewRelevance.bRenderCustomDepth && bIsCompatibleWithCustomDepth)
 			{
 				PassMask.Set(EMeshPass::CustomDepth);
 				View.NumVisibleDynamicMeshElements[EMeshPass::CustomDepth] += NumElements;
