Fix HV not using the correct VSM function for compiling or not VSM permutation.

#Synced-CL 51018413

diff --git a/Engine/Source/Runtime/Renderer/Private/HeterogeneousVolumes/HeterogeneousVolumesHardwareRayTracing.cpp b/Engine/Source/Runtime/Renderer/Private/HeterogeneousVolumes/HeterogeneousVolumesHardwareRayTracing.cpp
--- a/Engine/Source/Runtime/Renderer/Private/HeterogeneousVolumes/HeterogeneousVolumesHardwareRayTracing.cpp
+++ b/Engine/Source/Runtime/Renderer/Private/HeterogeneousVolumes/HeterogeneousVolumesHardwareRayTracing.cpp
@@ -392,8 +392,7 @@ class FRenderLightingCacheWithPreshadingRGS : public FGlobalShader
 	{
 		FMaterialShader::ModifyCompilationEnvironment(Parameters, OutEnvironment);
 
-		bool bSupportVirtualShadowMap = IsFeatureLevelSupported(Parameters.Platform, ERHIFeatureLevel::SM5);
-		if (bSupportVirtualShadowMap)
+		if (DoesPlatformSupportVirtualShadowMaps(Parameters.Platform))
 		{
 			OutEnvironment.SetDefine(TEXT("VIRTUAL_SHADOW_MAP"), 1);
 			FVirtualShadowMapArray::SetShaderDefines(OutEnvironment);
@@ -475,8 +474,7 @@ class FRenderSingleScatteringWithPreshadingRGS : public FGlobalShader
 	{
 		FMaterialShader::ModifyCompilationEnvironment(Parameters, OutEnvironment);
 
-		bool bSupportVirtualShadowMap = IsFeatureLevelSupported(Parameters.Platform, ERHIFeatureLevel::SM5);
-		if (bSupportVirtualShadowMap)
+		if (DoesPlatformSupportVirtualShadowMaps(Parameters.Platform))
 		{
 			OutEnvironment.SetDefine(TEXT("VIRTUAL_SHADOW_MAP"), 1);
 			FVirtualShadowMapArray::SetShaderDefines(OutEnvironment);
diff --git a/Engine/Source/Runtime/Renderer/Private/HeterogeneousVolumes/HeterogeneousVolumesLiveShadingPipeline.cpp b/Engine/Source/Runtime/Renderer/Private/HeterogeneousVolumes/HeterogeneousVolumesLiveShadingPipeline.cpp
--- a/Engine/Source/Runtime/Renderer/Private/HeterogeneousVolumes/HeterogeneousVolumesLiveShadingPipeline.cpp
+++ b/Engine/Source/Runtime/Renderer/Private/HeterogeneousVolumes/HeterogeneousVolumesLiveShadingPipeline.cpp
@@ -418,8 +418,7 @@ class FRenderLightingCacheWithLiveShadingCS : public FMeshMaterialShader
 		OutEnvironment.SetDefine(TEXT("THREADGROUP_SIZE_2D"), GetThreadGroupSize2D());
 		OutEnvironment.SetDefine(TEXT("THREADGROUP_SIZE_3D"), GetThreadGroupSize3D());
 
-		bool bSupportVirtualShadowMap = IsFeatureLevelSupported(Parameters.Platform, ERHIFeatureLevel::SM5);
-		if (bSupportVirtualShadowMap)
+		if (DoesPlatformSupportVirtualShadowMaps(Parameters.Platform))
 		{
 			OutEnvironment.SetDefine(TEXT("VIRTUAL_SHADOW_MAP"), 1);
 			FVirtualShadowMapArray::SetShaderDefines(OutEnvironment);
@@ -622,8 +621,7 @@ class FRenderSingleScatteringWithLiveShadingCS : public FMeshMaterialShader
 		OutEnvironment.SetDefine(TEXT("THREADGROUP_SIZE_2D"), GetThreadGroupSize2D());
 		OutEnvironment.SetDefine(TEXT("FOG_MATERIALBLENDING_OVERRIDE"), 1);
 
-		bool bSupportVirtualShadowMap = IsFeatureLevelSupported(Parameters.Platform, ERHIFeatureLevel::SM5);
-		if (bSupportVirtualShadowMap)
+		if (DoesPlatformSupportVirtualShadowMaps(Parameters.Platform))
 		{
 			OutEnvironment.SetDefine(TEXT("VIRTUAL_SHADOW_MAP"), 1);
 			FVirtualShadowMapArray::SetShaderDefines(OutEnvironment);
@@ -3459,8 +3457,7 @@ class FRenderVolumetricShadowMaskCS : public FGlobalShader
 		OutEnvironment.SetDefine(TEXT("THREADGROUP_SIZE_1D"), GetThreadGroupSize1D());
 		OutEnvironment.SetDefine(TEXT("THREADGROUP_SIZE_2D"), GetThreadGroupSize2D());
 
-		bool bSupportVirtualShadowMap = IsFeatureLevelSupported(Parameters.Platform, ERHIFeatureLevel::SM5);
-		if (bSupportVirtualShadowMap)
+		if (DoesPlatformSupportVirtualShadowMaps(Parameters.Platform))
 		{
 			OutEnvironment.SetDefine(TEXT("VIRTUAL_SHADOW_MAP"), 1);
 			FVirtualShadowMapArray::SetShaderDefines(OutEnvironment);
diff --git a/Engine/Source/Runtime/Renderer/Private/HeterogeneousVolumes/HeterogeneousVolumesPreshadingPipeline.cpp b/Engine/Source/Runtime/Renderer/Private/HeterogeneousVolumes/HeterogeneousVolumesPreshadingPipeline.cpp
--- a/Engine/Source/Runtime/Renderer/Private/HeterogeneousVolumes/HeterogeneousVolumesPreshadingPipeline.cpp
+++ b/Engine/Source/Runtime/Renderer/Private/HeterogeneousVolumes/HeterogeneousVolumesPreshadingPipeline.cpp
@@ -136,8 +136,7 @@ class FRenderLightingCacheWithPreshadingCS : public FGlobalShader
 		OutEnvironment.SetDefine(TEXT("THREADGROUP_SIZE_2D"), GetThreadGroupSize2D());
 		OutEnvironment.SetDefine(TEXT("THREADGROUP_SIZE_3D"), GetThreadGroupSize3D());
 
-		bool bSupportVirtualShadowMap = IsFeatureLevelSupported(Parameters.Platform, ERHIFeatureLevel::SM5);
-		if (bSupportVirtualShadowMap)
+		if (DoesPlatformSupportVirtualShadowMaps(Parameters.Platform))
 		{
 			OutEnvironment.SetDefine(TEXT("VIRTUAL_SHADOW_MAP"), 1);
 			FVirtualShadowMapArray::SetShaderDefines(OutEnvironment);
@@ -283,8 +282,7 @@ class FRenderSingleScatteringWithPreshadingCS : public FGlobalShader
 		OutEnvironment.SetDefine(TEXT("THREADGROUP_SIZE_1D"), GetThreadGroupSize1D());
 		OutEnvironment.SetDefine(TEXT("THREADGROUP_SIZE_2D"), GetThreadGroupSize2D());
 
-		bool bSupportVirtualShadowMap = IsFeatureLevelSupported(Parameters.Platform, ERHIFeatureLevel::SM5);
-		if (bSupportVirtualShadowMap)
+		if (DoesPlatformSupportVirtualShadowMaps(Parameters.Platform))
 		{
 			OutEnvironment.SetDefine(TEXT("VIRTUAL_SHADOW_MAP"), 1);
 			FVirtualShadowMapArray::SetShaderDefines(OutEnvironment);
