Fix VSM code being included/compiled on SM5 platforms while not allowed.

#Synced-CL 51018413

diff --git a/Engine/Source/Runtime/Renderer/Private/BasePassRendering.cpp b/Engine/Source/Runtime/Renderer/Private/BasePassRendering.cpp
--- a/Engine/Source/Runtime/Renderer/Private/BasePassRendering.cpp
+++ b/Engine/Source/Runtime/Renderer/Private/BasePassRendering.cpp
@@ -980,7 +980,7 @@ void ModifyBasePassCSPSCompilationEnvironment(const FMeshMaterialShaderPermutati
 	const bool bEnableSkyLight = bMaterialIsLit && (bProjectEnableSkyLight || bShaderPlatformUsesForward || bMaterialIsTranslucent || bMaterialIsSingleLayerWater);
 	OutEnvironment.SetDefine(TEXT("ENABLE_SKY_LIGHT"), bEnableSkyLight);
 
-	const bool bSupportVirtualShadowMap = IsFeatureLevelSupported(Parameters.Platform, ERHIFeatureLevel::SM5);
+	const bool bSupportVirtualShadowMap = DoesPlatformSupportVirtualShadowMaps(Parameters.Platform);
 	if (bSupportVirtualShadowMap && (bMaterialIsTranslucent || bMaterialIsSingleLayerWater))
 	{
 		FVirtualShadowMapArray::SetShaderDefines(OutEnvironment);
