Substrate - Toon - unified diffuse ramp sampling.
Specular is still evaluted per lights.
Sky lights is still separated.

The feature is hidden behind a cvar.
The shadow can optionally be added to the unifed diffuse ramp sampling.

#tests Enable via cvar r.Substrate.Experimental.ToonUnifiedDiffuse and add multiple lights. This got tested in enginetest.
#rnx
#Synced-CL 51090036

#Synced-CL 51090288

diff --git a/Engine/Shaders/Private/DeferredLightPixelShaders.usf b/Engine/Shaders/Private/DeferredLightPixelShaders.usf
--- a/Engine/Shaders/Private/DeferredLightPixelShaders.usf
+++ b/Engine/Shaders/Private/DeferredLightPixelShaders.usf
@@ -28,6 +28,9 @@
   #endif
 #endif
 
+#define SUBSTRATE_OUTPUT_TOON_DATA (SUBSTRATE_ENABLED && SUBSTRATE_TILETYPE >= 2 && SUBSTRATE_EXPERIMENTAL_TOON_UNIFIED_DIFFUSE)
+#define SUBSTRATE_ENABLE_TOON_DIFFUSE_RAMP (SUBSTRATE_OUTPUT_TOON_DATA ? 0 : 1)
+
 #if SUBSTRATE_GBUFFER_FORMAT==0 || USE_HAIR_LIGHTING
 #define SUBSTRATE_LOAD_FROM_MATERIALCONTAINER 0	// Load from a specified BSDF, not from MaterialBuffer.
 #define SUBSTRATE_INLINE_SHADING 1				// Inline shading for normal setup if needed.
@@ -258,6 +261,9 @@ void DeferredLightPixelMain(
 #endif
 	float4 SVPos			: SV_POSITION,
 	out float4 OutColor		: SV_Target0
+#if SUBSTRATE_OUTPUT_TOON_DATA
+	, out float4 OutToonData : SV_Target1
+#endif
 #if SUBSTRATE_OPAQUE_ROUGH_REFRACTION_ENABLED
 	, out float3 OutOpaqueRoughRefractionSceneColor : SV_Target1
 	, out float3 OutSubSurfaceSceneColor : SV_Target2
@@ -266,6 +272,9 @@ void DeferredLightPixelMain(
 {
 	const float2 PixelPos = SVPos.xy;
 	OutColor = 0;
+#if SUBSTRATE_OUTPUT_TOON_DATA
+	OutToonData = float4(0.0f, 0.0f, 0.0f, -1.0f);
+#endif
 #if SUBSTRATE_OPAQUE_ROUGH_REFRACTION_ENABLED
 	OutOpaqueRoughRefractionSceneColor = 0;
 	OutSubSurfaceSceneColor = 0;
@@ -364,6 +373,12 @@ void DeferredLightPixelMain(
 			EstimatedCost += SubstrateLighting.EstimatedCost;
 
 			OutColor += SubstrateLighting.SceneColor;
+		#if SUBSTRATE_OUTPUT_TOON_DATA
+			// We do not need to test for toon BSDF as it is a single BSDF and always will work correctly when needed.
+			{
+				OutToonData = SubstrateLighting.ToonData;
+			}
+		#endif
 		#if SUBSTRATE_OPAQUE_ROUGH_REFRACTION_ENABLED
 			OutOpaqueRoughRefractionSceneColor += SubstrateLighting.OpaqueRoughRefractionSceneColor;
 			OutSubSurfaceSceneColor += SubstrateLighting.SubSurfaceSceneColor;
@@ -377,6 +392,9 @@ void DeferredLightPixelMain(
 	// Doing the same to match the legacy output.
 	const float4 CostWithColorCode = ConvertEstimatedCostToColor(EstimatedCost);
 	OutColor = CostWithColorCode + CostWithColorCode;
+	#if SUBSTRATE_OUTPUT_TOON_DATA
+	OutToonData = float4(0.0f, 0.0f, 0.0f, -1.0f);
+	#endif
 	#if SUBSTRATE_OPAQUE_ROUGH_REFRACTION_ENABLED
 	OutOpaqueRoughRefractionSceneColor = 0.0f;
 	OutSubSurfaceSceneColor = 0.0f;
@@ -422,6 +440,9 @@ void DeferredLightPixelMain(
 	// A:Non Specular SceneColor Luminance
 	// So we need PreExposure for both color and alpha
 	OutColor.rgba *= GetExposure();
+#if SUBSTRATE_OUTPUT_TOON_DATA
+	OutToonData.rgb *= GetExposure();
+#endif
 #if SUBSTRATE_OPAQUE_ROUGH_REFRACTION_ENABLED 
 	// Idem
 	OutOpaqueRoughRefractionSceneColor *= GetExposure();
@@ -622,3 +643,72 @@ void DeferredLightPixelMain(
 #endif // SUBSTRATE_ENABLED
 }
 #endif
+
+////////////////////////////////////////////////////////////////////////////////////////////////////
+
+#ifdef DeferredLightApplyToonDiffusePS
+
+Texture2D<float4> ToonDataTexture;
+
+void DeferredLightApplyToonDiffusePS(
+	in float4 SVPos : SV_POSITION,
+	out float4 OutColor : SV_Target0
+)
+{
+	const float2 PixelPos = SVPos.xy;
+	OutColor = 0;
+
+
+	FInputParams InputParams = (FInputParams)0;
+	InputParams.PixelPos = SVPos.xy;
+	InputParams.ScreenPosition = SvPositionToScreenPosition(SVPos);
+	InputParams.ScreenUV = InputParams.ScreenPosition.xy / InputParams.ScreenPosition.w * View.ScreenPositionScaleBias.xy + View.ScreenPositionScaleBias.wz;
+	InputParams.ScreenVector = 0;
+
+#if SUBSTRATE_ENABLED==1
+
+#if SUBSTRATE_LOAD_FROM_MATERIALCONTAINER
+	FSubstrateAddressing SubstrateAddressing = GetSubstratePixelDataByteOffset(PixelPos, uint2(View.BufferSizeAndInvSize.xy), Substrate.MaxBytesPerPixel);
+	FSubstratePixelHeader SubstratePixelHeader = UnpackSubstrateHeaderIn(Substrate.MaterialTextureArray, SubstrateAddressing, Substrate.TopLayerTexture);
+	const bool bMaterialIsLit = SubstratePixelHeader.ClosureCount > 0;	// This test is also enough to exclude sky pixels
+#else
+	const bool bForceUnlitOrDefaultLit = SUBSTRATE_FASTPATH;
+	FScreenSpaceData ScreenSpaceData = GetScreenSpaceData(InputParams.ScreenUV, true, bForceUnlitOrDefaultLit);
+	const FGBufferData GBuffer = ScreenSpaceData.GBuffer;
+	const bool bMaterialIsLit = GBuffer.ShadingModelID != SHADINGMODELID_UNLIT;
+#endif
+
+	#if SUBSTRATE_LOAD_FROM_MATERIALCONTAINER
+	if (SubstratePixelHeader.GetComplexPathMode() == COMPLEXPATH_MODE_TOON)	// TODO turn off that test This test is in fact no necessary since toon is single BSDF and always min will work and no leak cross pixel
+	{
+		FSubstrateBSDF BSDF = UnpackSubstrateBSDF(MaterialBuffer, SubstrateAddressing, SubstratePixelHeader);
+		float ToonProfileId = TOON_PROFILEID(BSDF);
+		const float2 ToonPatternUV = TOON_PATTERNUVS(BSDF);
+	#else
+	if (GBuffer.ShadingModelID == SHADINGMODELID_SUBSTRATE_TOON)
+	{
+		float ToonProfileId = 0.0f;
+		float2 ToonPatternUV = 0.0;
+		UnpackToonCustomData(PackRGBA8(GBuffer.CustomData), ToonPatternUV, ToonProfileId);
+	#endif
+
+		float4 ToonData = ToonDataTexture[SVPos.xy];
+
+		FToonProfileEvaluateSettings ToonEvalSettings = GetToonProfileEvaluateSettings(ToonProfileId);
+		float ToonProfileDiffuseDither = EvaluateToonProfileDiffuseDither(ToonProfileId, ToonPatternUV, ToonEvalSettings);
+
+		const float NoL = ToonData.a;
+		const float3 ToonDiffuse = ToonData.rgb * EvaluateToonProfileDiffuseRamp(ToonProfileId, NoL + ToonProfileDiffuseDither);
+		OutColor = float4(ToonDiffuse, Luminance(ToonDiffuse));
+	}
+	else
+	{
+		clip(-1);
+	}
+
+#endif
+
+}
+
+
+#endif // DeferredLightApplyToonDiffusePS
diff --git a/Engine/Shaders/Private/LightAccumulator.ush b/Engine/Shaders/Private/LightAccumulator.ush
--- a/Engine/Shaders/Private/LightAccumulator.ush
+++ b/Engine/Shaders/Private/LightAccumulator.ush
@@ -181,6 +181,7 @@ FDeferredLightingSplit LightAccumulator_GetResultSplit(FLightAccumulator In)
 struct FSubstrateDeferredLighting
 {
 	float4 SceneColor;						// Untouched scene luminance, alpha channel may contain scatterable light luminance.
+	float4 ToonData;
 #if SUBSTRATE_OPAQUE_ROUGH_REFRACTION_ENABLED
 	float3 OpaqueRoughRefractionSceneColor;	// Luminance that is going to be blurier by the top layer roughness
 	float3 SubSurfaceSceneColor;			// Scene color that should go through the sub surface post process. There is not alpha: not needed because it is separated from specular already.
@@ -197,6 +198,7 @@ struct FSubstrateDeferredLighting
 FSubstrateDeferredLighting GetInitialisedSubstrateDeferredLighting()
 {
 	FSubstrateDeferredLighting Result = (FSubstrateDeferredLighting)0;
+	Result.ToonData = float4(0.0f, 0.0f, 0.0f, -1.0f);
 	return Result;
 }
 
diff --git a/Engine/Shaders/Private/ShadingEnergyConservationTable.usf b/Engine/Shaders/Private/ShadingEnergyConservationTable.usf
--- a/Engine/Shaders/Private/ShadingEnergyConservationTable.usf
+++ b/Engine/Shaders/Private/ShadingEnergyConservationTable.usf
@@ -73,7 +73,7 @@ void BuildEnergyTableCS(uint3 DispatchThreadId : SV_DispatchThreadID)
 		const float3 L = reflect(-V, H);
 		if (L.z > 0)
 		{
-			const float a2 = Alpha * Alpha;
+			const float a2 = Alpha.x * Alpha.x;	// no anisotropy
 			const float NoH = H.z;
 			const float NoL = L.z;
 			Weight  = D_GGX(a2, NoH);
diff --git a/Engine/Shaders/Private/Substrate/SubstrateDeferredLighting.ush b/Engine/Shaders/Private/Substrate/SubstrateDeferredLighting.ush
--- a/Engine/Shaders/Private/Substrate/SubstrateDeferredLighting.ush
+++ b/Engine/Shaders/Private/Substrate/SubstrateDeferredLighting.ush
@@ -199,6 +199,19 @@ FSubstrateDeferredLighting SubstrateDeferredLighting(
 		#endif
 				;
 
+		#if SUBSTRATE_OUTPUT_TOON_DATA
+			const bool bIsToon = BSDF_GETTYPE(BSDFContext.BSDF) == SUBSTRATE_BSDF_TYPE_TOON;
+			if (bIsToon)
+			{
+				const float3 ToonDiffuseLuminance = DiffuseLuminance * CommonMultiplier;
+				if (any(ToonDiffuseLuminance > 0))
+				{
+					SubstrateLighting.ToonData = float4(ToonDiffuseLuminance, BSDFEvaluate.ToonData) ;
+					DiffuseLuminance = 0.0f;
+				}
+			}
+		#endif
+
 			FLightAccumulator Out = (FLightAccumulator)0;
 			LightAccumulator_AddSplit(Out, DiffuseLuminance, SpecularLuminance, DiffuseLuminance, CommonMultiplier, BSDFEvaluate.bPostProcessSubsurface);
 			AccumulateSubstrateDeferredLighting(SubstrateLighting, Out, BSDFEvaluate.bPostProcessSubsurface, BSDF_GETISTOPLAYER(BSDF));
diff --git a/Engine/Shaders/Private/Substrate/SubstrateEvaluation.ush b/Engine/Shaders/Private/Substrate/SubstrateEvaluation.ush
--- a/Engine/Shaders/Private/Substrate/SubstrateEvaluation.ush
+++ b/Engine/Shaders/Private/Substrate/SubstrateEvaluation.ush
@@ -57,6 +57,10 @@
 #define SUBSTRATE_DIFFUSE_COLOR_BOOST 0
 #endif
 
+#ifndef SUBSTRATE_ENABLE_TOON_DIFFUSE_RAMP
+#define SUBSTRATE_ENABLE_TOON_DIFFUSE_RAMP 1 
+#endif
+
 #ifndef SUBSTRATE_USE_FULLYSIMPLIFIED_MATERIAL
 #define SUBSTRATE_USE_FULLYSIMPLIFIED_MATERIAL 0
 #endif
@@ -416,6 +420,9 @@ struct FSubstrateEvaluateResult
 	float3 TransmittanceAlongN;			// Transmittance to the next layer along the surface normal. This is used later to compute the transmittance towards the light.
 	bool   bPostProcessSubsurface;		// True if we need to separate the subsurface light contribution for the screen space diffusion process.
 	bool   bApplyProjectedSolidAngle;	// True if the Saturate(NoL) factor should be applied or not over the value of the BSDF
+
+	// This is NoL maybe weighted by the shadow for toon shading model only.
+	float  ToonData;
 };
 
 // PUNCTUAL corresponds to an evaluation (from in & out directions).
@@ -1400,16 +1407,17 @@ FSubstrateEvaluateResult SubstrateEvaluateBSDFCommon(FSubstrateBSDFContext BSDFC
 				Specular = pow(saturate(BSDFContext.Context.NoH), max(1.0f, 1.0 + RoughnessExponent2 * SpecularExponentScale));					// NoH based
 			#endif
 			}
-			// TOON_TODO add anisotropy support
-			// TOON_TODO energy conservation
-			// TOON_TODO experiment with GGX alike lob that is normalize according to peak lobe?
-			// TOON_TODO Area light, the roughness expand trick here does not make sense to use. MRP?
 			
+			const bool bApplyDiffuseRamp	= SUBSTRATE_ENABLE_TOON_DIFFUSE_RAMP;
+			const bool bApplyDiffuseShadow	= !SUBSTRATE_ENABLE_TOON_DIFFUSE_RAMP && ToonEvalSettings.bUnifiedDiffuseIncludeShadow;
+
+			Sample.ToonData					= AreaLight.NoL * (bApplyDiffuseShadow ? 1.0f : TransmittanceShadow);
 
 			// Diffuse
-			Sample.DiffuseColor				= DirectionalAlbedo_SpecularTransmission * DiffuseColor * EvaluateToonProfileDiffuseRamp(ToonProfileId, BSDFContext.Context.NoL + ToonProfileDiffuseDither); // AreaLight.NoL to handle capsule shape?
+			const float3 DiffuseRamp		= bApplyDiffuseRamp ? EvaluateToonProfileDiffuseRamp(ToonProfileId, BSDFContext.Context.NoL + ToonProfileDiffuseDither) : 1.0f;
+			Sample.DiffuseColor				= DirectionalAlbedo_SpecularTransmission * DiffuseColor * DiffuseRamp;
 			Sample.DiffusePathValue			= DirectionalAlbedo_SpecularTransmission * Diffuse_Lambert(Sample.DiffuseColor);
-			Sample.IntegratedDiffuseValue	= DirectionalAlbedo_SpecularTransmission * TransmittanceShadow * AreaLight.Falloff * AreaLight.FalloffColor * Sample.DiffusePathValue;
+			Sample.IntegratedDiffuseValue	= DirectionalAlbedo_SpecularTransmission * (bApplyDiffuseShadow ? TransmittanceShadow : 1.0f) * AreaLight.Falloff * AreaLight.FalloffColor * Sample.DiffusePathValue; //  * AreaLight.NoL for smoother edges? But then transmittance is blocked. So this is left in the hand of the artist.
 
 			const float3 SpecularRamp = SpecularColor * MSScale * EvaluateToonProfileSpecularRamp(ToonProfileId, saturate(dot(Specular, 1.0f/3.0f) + ToonProfileSpecularDither));
 			Sample.SpecularPathValue = SpecularRamp;
diff --git a/Engine/Shaders/Private/ToonProfile.usf b/Engine/Shaders/Private/ToonProfile.usf
--- a/Engine/Shaders/Private/ToonProfile.usf
+++ b/Engine/Shaders/Private/ToonProfile.usf
@@ -24,6 +24,7 @@ float DiffuseIndirectScale;
 float SpecularIndirectScale;
 float bSpecularUsesGGX;
 float bEnergyConservation;
+float bUnifiedDiffuseIncludeShadow;
 float SpecularIndirectRampRepetition;
 
 RWTexture2DArray<float4>	ToonProfileTexture;
@@ -77,7 +78,7 @@ void ToonProfileUpdateAtlasTileCS(uint3 DispatchThreadId : SV_DispatchThreadID)
 			}
 			else if (PixelPos.x == TOON_PROFILE_EVAL_PARAMS2_X)
 			{
-				ToonProfileTexture[uint3(PixelPos, ToonProfileIndex)] = float4(bSpecularUsesGGX, bEnergyConservation, 0.0f, 0.0f);
+				ToonProfileTexture[uint3(PixelPos, ToonProfileIndex)] = float4(bSpecularUsesGGX, bEnergyConservation, bUnifiedDiffuseIncludeShadow, 0.0f);
 			}
 			else if (PixelPos.x == TOON_PROFILE_SHADOW_PARAMS_X)
 			{
diff --git a/Engine/Shaders/Private/ToonProfileCommon.ush b/Engine/Shaders/Private/ToonProfileCommon.ush
--- a/Engine/Shaders/Private/ToonProfileCommon.ush
+++ b/Engine/Shaders/Private/ToonProfileCommon.ush
@@ -46,6 +46,7 @@ struct FToonProfileEvaluateSettings
 	float SpecularDitherPatternSize;
 	float bSpecularUsesGGX;
 	float bEnergyConservation;
+	float bUnifiedDiffuseIncludeShadow;
 };
 
 FToonProfileEvaluateSettings GetToonProfileEvaluateSettings(in float ToonProfileId)
@@ -70,6 +71,7 @@ FToonProfileEvaluateSettings GetToonProfileEvaluateSettings(in float ToonProfile
 
 	Ret.bSpecularUsesGGX = Data2.x;
 	Ret.bEnergyConservation = Data2.y;
+	Ret.bUnifiedDiffuseIncludeShadow = Data2.z;
 
 	return Ret;
 }
diff --git a/Engine/Source/Runtime/Engine/Classes/Engine/ToonProfile.h b/Engine/Source/Runtime/Engine/Classes/Engine/ToonProfile.h
--- a/Engine/Source/Runtime/Engine/Classes/Engine/ToonProfile.h
+++ b/Engine/Source/Runtime/Engine/Classes/Engine/ToonProfile.h
@@ -140,6 +140,12 @@ struct FToonProfileStruct
 	UPROPERTY(Category = "Experimental", EditAnywhere, meta = (DisplayName = "Energy Conservation"))
 	uint8 bEnergyConservation : 1;
 
+	/**
+	 * Experimental. Requires r.Substrate.Experimental.ToonUnifiedDiffuse=1.
+	 */
+	UPROPERTY(Category = "Experimental", EditAnywhere, meta = (DisplayName = "Unified Diffuse Include Shadow"))
+	uint8 bUnifiedDiffuseIncludeShadow : 1;
+
 	FToonProfileStruct();
 
 	void Invalidate()
diff --git a/Engine/Source/Runtime/Engine/Private/Rendering/ToonProfile.cpp b/Engine/Source/Runtime/Engine/Private/Rendering/ToonProfile.cpp
--- a/Engine/Source/Runtime/Engine/Private/Rendering/ToonProfile.cpp
+++ b/Engine/Source/Runtime/Engine/Private/Rendering/ToonProfile.cpp
@@ -83,6 +83,7 @@ public:
 		SHADER_PARAMETER(float, DiffuseIndirectScale)
 		SHADER_PARAMETER(float, SpecularIndirectScale)
 		SHADER_PARAMETER(float, bSpecularUsesGGX)
+		SHADER_PARAMETER(float, bUnifiedDiffuseIncludeShadow)
 		SHADER_PARAMETER(float, bEnergyConservation)
 		SHADER_PARAMETER(float, SpecularIndirectRampRepetition)
 
@@ -418,6 +419,7 @@ IPooledRenderTarget* FToonProfileTextureManager::GetAtlasTexture(FRDGBuilder& Gr
 			PassParameters->DiffuseIndirectScale = Data.DiffuseIndirectScale;
 			PassParameters->SpecularIndirectScale = Data.SpecularIndirectScale;
 			PassParameters->bSpecularUsesGGX = Data.bSpecularUsesGGX ? 1.0f : 0.0f;
+			PassParameters->bUnifiedDiffuseIncludeShadow = Data.bUnifiedDiffuseIncludeShadow ? 1.0f : 0.0f;
 			PassParameters->bEnergyConservation = Data.bEnergyConservation ? 1.0f : 0.0f;
 			PassParameters->SpecularIndirectRampRepetition = Data.SpecularIndirectRampRepetition;
 
@@ -557,6 +559,7 @@ FToonProfileStruct::FToonProfileStruct()
 	SpecularIndirectRampRepetition = 0.0f;
 
 	bSpecularUsesGGX = 0;
+	bUnifiedDiffuseIncludeShadow = 0;
 	bEnergyConservation = 0;
 
 	AddPointToCurve(ShadowRamp, 0.00, 0.0);
diff --git a/Engine/Source/Runtime/Engine/Private/ShaderCompiler/ShaderCompiler.cpp b/Engine/Source/Runtime/Engine/Private/ShaderCompiler/ShaderCompiler.cpp
--- a/Engine/Source/Runtime/Engine/Private/ShaderCompiler/ShaderCompiler.cpp
+++ b/Engine/Source/Runtime/Engine/Private/ShaderCompiler/ShaderCompiler.cpp
@@ -4140,12 +4140,17 @@ void GlobalBeginCompileShader(
 		const bool IsStochasticLightingEnabled = bSubstrate && Substrate::IsStochasticLightingEnabled(Target.GetPlatform());
 		SET_SHADER_DEFINE(Input.Environment, SUBSTRATE_STOCHASTIC_LIGHTING_ENABLED, IsStochasticLightingEnabled ? 1 : 0);
 
-
 		{
 			static IConsoleVariable* CVar = IConsoleManager::Get().FindConsoleVariable(TEXT("r.Substrate.Experimental.ToonReflectionQuantizationEnabled"));
 			const bool bExperimentalToonReflectionQuantizationEnabled = bSubstrate && CVar && CVar->GetInt() > 0 ? 1 : 0;
 			SET_SHADER_DEFINE(Input.Environment, SUBSTRATE_EXPERIMENTAL_TOON_REFLECTION_QUANTIZATION, bExperimentalToonReflectionQuantizationEnabled ? 1 : 0);
 		}
+
+		{
+			static IConsoleVariable* CVar = IConsoleManager::Get().FindConsoleVariable(TEXT("r.Substrate.Experimental.ToonUnifiedDiffuse"));
+			const bool bExperimentalToonUnifiedDiffuseEnabled = bSubstrate && CVar && CVar->GetInt() > 0 ? 1 : 0;
+			SET_SHADER_DEFINE(Input.Environment, SUBSTRATE_EXPERIMENTAL_TOON_UNIFIED_DIFFUSE, bExperimentalToonUnifiedDiffuseEnabled ? 1 : 0);
+		}
 	}
 
 	{
diff --git a/Engine/Source/Runtime/RenderCore/Private/RenderUtils.cpp b/Engine/Source/Runtime/RenderCore/Private/RenderUtils.cpp
--- a/Engine/Source/Runtime/RenderCore/Private/RenderUtils.cpp
+++ b/Engine/Source/Runtime/RenderCore/Private/RenderUtils.cpp
@@ -2220,6 +2220,12 @@ static TAutoConsoleVariable<int32> CVarSubstrateExperimentalToonReflectionQuanti
 	TEXT("Enable experimental toon reflection quantization."),
 	ECVF_ReadOnly | ECVF_RenderThreadSafe);
 
+static TAutoConsoleVariable<int32> CVarSubstrateExperimentalToonUnifiedDiffuse(
+	TEXT("r.Substrate.Experimental.ToonUnifiedDiffuse"),
+	0,
+	TEXT("Enable experimental toon unified diffuse evaluation."),
+	ECVF_ReadOnly | ECVF_RenderThreadSafe);
+
 namespace Substrate
 {
 	bool IsSubstrateEnabled()
@@ -2509,6 +2515,13 @@ namespace Substrate
 		return IsSubstrateEnabled() && IsUsingDBuffers(InPlatform) && bDBufferPassSupported && CVarSubstrateDBufferPassPlatform.Get(InPlatform) > 0;
 	}
 
+	// Accessor to experimental variable static local here to not have have it exposed in public header RenderUtils.h for now.
+	static bool IsExperimentalToonUnifiedDiffuseEnabled()
+	{
+		static IConsoleVariable* CVar = IConsoleManager::Get().FindConsoleVariable(TEXT("r.Substrate.Experimental.ToonUnifiedDiffuse"));
+		return CVar && CVar->GetInt() > 0;
+	}
+
 	bool IsOpaqueRoughRefractionEnabled(EShaderPlatform InPlatform)
 	{
 		static FShaderPlatformCachedIniValue<int32> CVar(TEXT("r.Substrate.OpaqueMaterialRoughRefraction"));
@@ -2537,7 +2550,7 @@ namespace Substrate
 			}
 		}
 
-		return IsSubstrateEnabled() && CVar.Get(InPlatform) > 0 && !bIsSubstrateBlendableGBufferEnabled && bManyClosures;
+		return IsSubstrateEnabled() && CVar.Get(InPlatform) > 0 && !bIsSubstrateBlendableGBufferEnabled && bManyClosures && !IsExperimentalToonUnifiedDiffuseEnabled();
 	}
 
 	bool IsAdvancedVisualizationEnabled()
diff --git a/Engine/Source/Runtime/RenderCore/Private/Shader.cpp b/Engine/Source/Runtime/RenderCore/Private/Shader.cpp
--- a/Engine/Source/Runtime/RenderCore/Private/Shader.cpp
+++ b/Engine/Source/Runtime/RenderCore/Private/Shader.cpp
@@ -2572,6 +2572,15 @@ void ShaderMapAppendKey(EShaderPlatform Platform, FShaderKeyGenerator& KeyGen)
 				KeyGen.Append(TEXT("EXPTOONRQ"));
 			}
 		}
+
+		{
+			static IConsoleVariable* CVar = IConsoleManager::Get().FindConsoleVariable(TEXT("r.Substrate.Experimental.ToonUnifiedDiffuse"));
+			if (CVar && CVar->GetInt() > 0)
+			{
+				KeyGen.AppendSeparator();
+				KeyGen.Append(TEXT("EXPTOONUD"));
+			}
+		}
 	}
 
 	{
diff --git a/Engine/Source/Runtime/Renderer/Private/LightRendering.cpp b/Engine/Source/Runtime/Renderer/Private/LightRendering.cpp
--- a/Engine/Source/Runtime/Renderer/Private/LightRendering.cpp
+++ b/Engine/Source/Runtime/Renderer/Private/LightRendering.cpp
@@ -37,6 +37,7 @@
 #include "RHIResourceUtils.h"
 #include "RectLightSceneProxy.h"
 #include "Shadows/FirstPersonSelfShadow.h"
+#include "PixelShaderUtils.h"
 
 #include "PSOPrecacheValidation.h"
 
@@ -212,6 +213,13 @@ static FAutoConsoleVariableRef CVarDebugLightDiscardProp(
 );
 #endif // ENABLE_DEBUG_DISCARD_PROP
 
+// Accessor to experimental variable static local here to not have have it exposed in public header RenderUtils.h for now.
+static bool IsExperimentalToonUnifiedDiffuseEnabled()
+{
+	static IConsoleVariable* CVar = IConsoleManager::Get().FindConsoleVariable(TEXT("r.Substrate.Experimental.ToonUnifiedDiffuse"));
+	return CVar && CVar->GetInt() > 0;
+}
+
 #if RHI_RAYTRACING
 
 static bool ShouldRenderRayTracingShadowsForLightType(ELightComponentType LightType)
@@ -1217,6 +1225,38 @@ class FDeferredLightPS : public FGlobalShader
 
 IMPLEMENT_GLOBAL_SHADER(FDeferredLightPS, "/Engine/Private/DeferredLightPixelShaders.usf", "DeferredLightPixelMain", SF_Pixel);
 
+/** Shader used to compose toon diffuse lighting. */
+class FDeferredLightApplyToonDiffusePS : public FGlobalShader
+{
+	DECLARE_SHADER_TYPE(FDeferredLightApplyToonDiffusePS, Global)
+	SHADER_USE_PARAMETER_STRUCT(FDeferredLightApplyToonDiffusePS, FGlobalShader);
+
+	using FPermutationDomain = TShaderPermutationDomain<Substrate::FSubstrateTileType>;
+
+	BEGIN_SHADER_PARAMETER_STRUCT(FParameters, )
+		SHADER_PARAMETER_STRUCT_REF(FViewUniformShaderParameters, View)
+		SHADER_PARAMETER_RDG_UNIFORM_BUFFER(FSceneTextureUniformParameters, SceneTextures)
+		SHADER_PARAMETER_RDG_UNIFORM_BUFFER(FSubstrateGlobalUniformParameters, Substrate)
+		SHADER_PARAMETER_RDG_TEXTURE(Texture2D<float4>, ToonDataTexture)
+		RENDER_TARGET_BINDING_SLOTS()
+	END_SHADER_PARAMETER_STRUCT()
+
+public:
+
+	static bool ShouldCompilePermutation(const FGlobalShaderPermutationParameters& Parameters)
+	{
+		return IsFeatureLevelSupported(Parameters.Platform, ERHIFeatureLevel::SM5);
+	}
+
+	static void ModifyCompilationEnvironment(const FGlobalShaderPermutationParameters& Parameters, FShaderCompilerEnvironment& OutEnvironment)
+	{
+		FGlobalShader::ModifyCompilationEnvironment(Parameters, OutEnvironment);
+		OutEnvironment.SetDefine(TEXT("USE_HAIR_LIGHTING"), -1);
+	}
+};
+
+IMPLEMENT_GLOBAL_SHADER(FDeferredLightApplyToonDiffusePS, "/Engine/Private/DeferredLightPixelShaders.usf", "DeferredLightApplyToonDiffusePS", SF_Pixel);
+
 /** Shader used to visualize stationary light overlap. */
 class FDeferredLightOverlapPS : public FGlobalShader
 {
@@ -2459,6 +2499,40 @@ void FDeferredShadingSceneRenderer::RenderLights(
 				UnbatchedLightsPass(true/*bIsHairPass*/);
 			}
 		}
+
+		if (Substrate::IsSubstrateEnabled() && IsExperimentalToonUnifiedDiffuseEnabled())
+		{
+			for (int32 ViewIndex = 0, ViewCount = Views.Num(); ViewIndex < ViewCount; ++ViewIndex)
+			{
+				FViewInfo& View = Views[ViewIndex];
+
+				if (Substrate::GetSubstrateUsesTileType(View, ESubstrateTileType::EComplex) || Substrate::GetSubstrateUsesTileType(View, ESubstrateTileType::EComplexSpecial))
+				{
+					FDeferredLightApplyToonDiffusePS::FParameters* ToonPassParameters = GraphBuilder.AllocParameters<FDeferredLightApplyToonDiffusePS::FParameters>();
+					// PS - General parameters
+					ToonPassParameters->View = View.ViewUniformBuffer;
+					ToonPassParameters->SceneTextures = SceneTextures.UniformBuffer;
+					ToonPassParameters->Substrate = Substrate::BindSubstrateGlobalUniformParameters(View);
+
+					//	const FSubstrateSceneData& SubstrateSceneData = *View.GetSceneTextures().SubstrateSceneData;
+					ToonPassParameters->ToonDataTexture = SubstrateSceneData.SubstrateSceneToonData;
+
+					ToonPassParameters->RenderTargets[0] = FRenderTargetBinding(SceneTextures.Color.Target, ERenderTargetLoadAction::ELoad);
+
+					FDeferredLightApplyToonDiffusePS::FPermutationDomain ToonPermutationVector;
+					const ESubstrateTileType TileType = ESubstrateTileType::EComplex;											// TODO render "complex" and "complex special" tiles
+					ToonPermutationVector.Set<Substrate::FSubstrateTileType>(TileType);
+					TShaderMapRef<FDeferredLightApplyToonDiffusePS> PixelShader(View.ShaderMap, ToonPermutationVector);
+
+					FRHIBlendState* BlendState = TStaticBlendState<CW_RGBA, BO_Add, BF_One, BF_One, BO_Add, BF_One, BF_One>::GetRHI();
+
+					FPixelShaderUtils::AddFullscreenPass<FDeferredLightApplyToonDiffusePS>(
+						GraphBuilder, View.ShaderMap, RDG_EVENT_NAME("Light::ApplyToonDiffuse"), PixelShader, ToonPassParameters,
+						View.ViewRect, BlendState);
+					// A full screen pass for now. SUBSTRATE_TODO use EComplex tiles.
+				}
+			}
+		}
 	}
 
 	GraphBuilder.FlushSetupQueue();
@@ -2631,7 +2705,8 @@ static FDeferredLightPS::FParameters GetDeferredLightPSParameters(
 	bool bCloudShadow,
 	TRDGUniformBufferRef<FVirtualShadowMapUniformParameters> VirtualShadowMapUniformBuffer = nullptr,
 	FRDGTextureRef ShadowMaskBits = nullptr,
-	int32 VirtualShadowMapId = INDEX_NONE)
+	int32 VirtualShadowMapId = INDEX_NONE,
+	ESubstrateTileType SubstrateTileType = ESubstrateTileType::ECount)
 {
 	FDeferredLightPS::FParameters Out;
 
@@ -2674,12 +2749,21 @@ static FDeferredLightPS::FParameters GetDeferredLightPSParameters(
 
 	// PS - Render Targets
 	Out.RenderTargets[0] = FRenderTargetBinding(SceneColorTexture, ERenderTargetLoadAction::ELoad);
-	if (Substrate::IsOpaqueRoughRefractionEnabled(View.GetShaderPlatform()) && Substrate::UsesSubstrateMaterialBuffer(View.GetShaderPlatform()))
+	if (SubstrateTileType >= ESubstrateTileType::EComplex && IsExperimentalToonUnifiedDiffuseEnabled())
+	{
+		const FSubstrateSceneData& SubstrateSceneData = *View.GetSceneTextures().SubstrateSceneData;
+		Out.RenderTargets[1] = FRenderTargetBinding(SubstrateSceneData.SubstrateSceneToonData, ERenderTargetLoadAction::ELoad);
+
+		// As of today, OpaqueRoughRefraction cannot be used together with experimental toon features. This is reflected in the implementation of IsOpaqueRoughRefractionEnabled().
+		check(!Substrate::IsOpaqueRoughRefractionEnabled(View.GetShaderPlatform()));
+	}
+	else if (Substrate::IsOpaqueRoughRefractionEnabled(View.GetShaderPlatform()) && Substrate::UsesSubstrateMaterialBuffer(View.GetShaderPlatform()))
 	{
 		const FSubstrateSceneData& SubstrateSceneData = *View.GetSceneTextures().SubstrateSceneData;
 		Out.RenderTargets[1] = FRenderTargetBinding(SubstrateSceneData.SeparatedOpaqueRoughRefractionSceneColor, ERenderTargetLoadAction::ELoad);
 		Out.RenderTargets[2] = FRenderTargetBinding(SubstrateSceneData.SeparatedSubSurfaceSceneColor, ERenderTargetLoadAction::ELoad);
 	}
+
 	if (SceneDepthTexture)
 	{
 		Out.RenderTargets.DepthStencil = FDepthStencilBinding(SceneDepthTexture, ERenderTargetLoadAction::ELoad, ERenderTargetLoadAction::ELoad, FExclusiveDepthStencil::DepthRead_StencilWrite);
@@ -2703,7 +2787,16 @@ uint32 SetupLightGraphicsPSOState(
 	EShaderPlatform ShaderPlatform)
 {
 	GraphicsPSOInit.PrimitiveType = PT_TriangleList;
-	if (Substrate::IsOpaqueRoughRefractionEnabled(ShaderPlatform))
+	if (SubstrateTileMaterialType >= ESubstrateTileType::EComplex && IsExperimentalToonUnifiedDiffuseEnabled())
+	{
+		GraphicsPSOInit.BlendState = TStaticBlendState<
+			CW_RGBA, BO_Add, BF_One, BF_One, BO_Add, BF_One, BF_One,
+			CW_RGBA, BO_Add, BF_One, BF_One, BO_Max, BF_One, BF_One>::GetRHI();
+
+		// As of today, OpaqueRoughRefraction cannot be used together with experimental toon features. This is reflected in the implementation of IsOpaqueRoughRefractionEnabled().
+		check(!Substrate::IsOpaqueRoughRefractionEnabled(ShaderPlatform));
+	}
+	else if (Substrate::IsOpaqueRoughRefractionEnabled(ShaderPlatform))
 	{
 		GraphicsPSOInit.BlendState = TStaticBlendState<
 			CW_RGBA, BO_Add, BF_One, BF_One, BO_Add, BF_One, BF_One,
@@ -3019,7 +3112,7 @@ static void RenderLight(
 		// PS - Generatl parameters
 		PassParameters->PS = GetDeferredLightPSParameters(
 			GraphBuilder, Scene, View, LightSceneInfo, SceneTextures.Color.Target, SceneTextures.Depth.Target, SceneTextures.UniformBuffer, View.HairStrandsViewData.UniformBuffer, 
-			ScreenShadowMaskTexture, LightingChannelsTexture, bCloudShadow, VirtualShadowMapUniformBuffer, ShadowMaskBits, VirtualShadowMapId);
+			ScreenShadowMaskTexture, LightingChannelsTexture, bCloudShadow, VirtualShadowMapUniformBuffer, ShadowMaskBits, VirtualShadowMapId, SubstrateTileType);
 		// VS - General parameters
 		if (bIsRadial)
 		{
diff --git a/Engine/Source/Runtime/Renderer/Private/Substrate/Substrate.cpp b/Engine/Source/Runtime/Renderer/Private/Substrate/Substrate.cpp
--- a/Engine/Source/Runtime/Renderer/Private/Substrate/Substrate.cpp
+++ b/Engine/Source/Runtime/Renderer/Private/Substrate/Substrate.cpp
@@ -173,6 +173,13 @@ bool UsesSubstrateClosureCountFromMaterialData()
 	return CVarSubstrateUseClosureCountFromMaterial.GetValueOnRenderThread() > 0;
 }
 
+// Accessor to experimental variable static local here to not have have it exposed in public header RenderUtils.h for now.
+static bool IsExperimentalToonUnifiedDiffuseEnabled()
+{
+	static IConsoleVariable* CVar = IConsoleManager::Get().FindConsoleVariable(TEXT("r.Substrate.Experimental.ToonUnifiedDiffuse"));
+	return CVar && CVar->GetInt() > 0;
+}
+
 uint32 GetSubstrateMaxClosureCount(const FViewInfo& View)
 {
 	uint32 Out = 1;
@@ -603,6 +610,14 @@ void InitialiseSubstrateFrameSceneData(FRDGBuilder& GraphBuilder, FScene& Scene,
 			}
 		}
 
+		Out.SubstrateSceneToonData = nullptr;
+		if (IsExperimentalToonUnifiedDiffuseEnabled())
+		{
+			Out.SubstrateSceneToonData = GraphBuilder.CreateTexture(
+				FRDGTextureDesc::Create2D(SceneTextureExtent, PF_FloatRGBA, FClearValueBinding(FLinearColor(0.0f, 0.0f, 0.0f, -1.0f)), TexCreate_ShaderResource | TexCreate_UAV | TexCreate_RenderTargetable), TEXT("Substrate.SubstrateSceneToonData"));
+			AddClearRenderTargetPass(GraphBuilder, Out.SubstrateSceneToonData, Out.SubstrateSceneToonData->Desc.ClearValue.GetClearColor());
+		}
+
 		// Closure offsets
 		if (bNeedsMaterialBuffer && bNeedsClosureOffsets)
 		{
diff --git a/Engine/Source/Runtime/Renderer/Private/Substrate/Substrate.h b/Engine/Source/Runtime/Renderer/Private/Substrate/Substrate.h
--- a/Engine/Source/Runtime/Renderer/Private/Substrate/Substrate.h
+++ b/Engine/Source/Runtime/Renderer/Private/Substrate/Substrate.h
@@ -141,6 +141,9 @@ struct FSubstrateSceneData
 	// Used for Luminance that should go through opaque rough refraction (when under a top layer interface)
 	FRDGTextureRef SeparatedOpaqueRoughRefractionSceneColor = nullptr;
 
+	// Used for Toon unified diffuse.
+	FRDGTextureRef SubstrateSceneToonData = nullptr;
+
 	//Public facing minimal uniform data.
 	TRDGUniformBufferRef<FSubstratePublicGlobalUniformParameters> SubstratePublicGlobalUniformParameters{};
 };
