Fix translucency / single layer water (SLW) order when under water on mobile renderer
- Above water: render water then translucencies
- Under water: render translucencies then water (no need to copy SceneColor before translucencies when underwater)

#Synced-CL 51059040

diff --git a/Engine/Source/Runtime/Renderer/Private/MobileShadingRenderer.cpp b/Engine/Source/Runtime/Renderer/Private/MobileShadingRenderer.cpp
--- a/Engine/Source/Runtime/Renderer/Private/MobileShadingRenderer.cpp
+++ b/Engine/Source/Runtime/Renderer/Private/MobileShadingRenderer.cpp
@@ -2979,10 +2979,52 @@ void FMobileSceneRenderer::RenderDeferredMultiPass(FRDGBuilder& GraphBuilder, FS
 			RenderDistortion(GraphBuilder, View, SceneTextures.Color.Target, SceneTextures.Depth.Target);
 		}
 			
-		if (bShouldRenderSingleLayerWater 
-			&& View.bHasSingleLayerWaterMaterial
-			&& HasAnyDraw(View.ParallelMeshDrawCommandPasses[EMeshPass::SingleLayerWaterPass]))
+		const bool bRenderTranslucency = ShouldRenderTranslucency(StandardTranslucencyPass, Views) && ViewFamily.EngineShowFlags.Translucency && HasAnyDraw(View.ParallelMeshDrawCommandPasses[StandardTranslucencyMeshPass]); 
+		const bool bRenderSingleLayerWater = bShouldRenderSingleLayerWater && View.bHasSingleLayerWaterMaterial && HasAnyDraw(View.ParallelMeshDrawCommandPasses[EMeshPass::SingleLayerWaterPass]);
+		ETranslucencyView TranslucencyViewsToRender = bRenderTranslucency ? GetTranslucencyViews(Views) : ETranslucencyView::None;
+		const bool bIsCameraUnderWater = EnumHasAnyFlags(TranslucencyViewsToRender, ETranslucencyView::UnderWater);
+		
+		// - Above water : render water then translucencies
+		// - Under water : render translucencies then water
+		auto AddTranslucencyPass = [&]()
+		{
+			check(bRenderTranslucency);
+
+			// No need to copy SceneColor before translucencies when underwater
+			if (View.TranslucentPrimCount.UseSceneColorCopy(StandardTranslucencyPass) && (!bRenderSingleLayerWater || !bIsCameraUnderWater))
+			{
+				FRDGTextureRef SceneColorCopyTexture = AddCopySceneColorPass(GraphBuilder, Views, SceneTextures.Color, false, false);
+				if (SceneColorCopyTexture != nullptr)
+				{
+					BasePassParameters->SceneColorCopyTexture = SceneColorCopyTexture;
+				}
+				if (SceneDepthCopyTexture != nullptr)
+				{
+					BasePassParameters->SceneDepthCopyTexture = SceneDepthCopyTexture;
+				}
+
+				PassParameters->MobileBasePass = GraphBuilder.CreateUniformBuffer(BasePassParameters);
+			}
+
+			GraphBuilder.AddPass(
+				RDG_EVENT_NAME("Translucency"),
+				PassParameters,
+				ERDGPassFlags::Raster,
+				[this, &View, PassParameters](FRHICommandList& RHICmdList)
+				{
+					RenderTranslucency(RHICmdList, View, Views, StandardTranslucencyPass, StandardTranslucencyMeshPass, &PassParameters->InstanceCullingDrawParams);
+				});
+		};
+
+		if (bIsCameraUnderWater && bRenderTranslucency)
 		{
+			AddTranslucencyPass();
+		}
+			
+		if (bRenderSingleLayerWater)
+		{
+			RDG_EVENT_SCOPE(GraphBuilder, "SingleLayerWater");
+
 			auto* SLWPassParameters = GraphBuilder.AllocParameters<FMobileRenderPassParameters>();
 			*SLWPassParameters = *PassParameters;
 			SLWPassParameters->RenderTargets[0] = BasePassRenderTargets[0];
@@ -3021,34 +3063,10 @@ void FMobileSceneRenderer::RenderDeferredMultiPass(FRDGBuilder& GraphBuilder, FS
 					}
 				});
 		}
-			
-		if (ShouldRenderTranslucency(StandardTranslucencyPass, Views) 
-			&& ViewFamily.EngineShowFlags.Translucency
-			&& HasAnyDraw(View.ParallelMeshDrawCommandPasses[StandardTranslucencyMeshPass]))
+
+		if (!bIsCameraUnderWater && bRenderTranslucency)
 		{
-			if (View.TranslucentPrimCount.UseSceneColorCopy(StandardTranslucencyPass))
-			{
-				FRDGTextureRef SceneColorCopyTexture = AddCopySceneColorPass(GraphBuilder, Views, SceneTextures.Color, false, false);
-				if (SceneColorCopyTexture != nullptr)
-				{
-					BasePassParameters->SceneColorCopyTexture = SceneColorCopyTexture;
-				}
-				if (SceneDepthCopyTexture != nullptr)
-				{
-					BasePassParameters->SceneDepthCopyTexture = SceneDepthCopyTexture;
-				}
-				
-				PassParameters->MobileBasePass = GraphBuilder.CreateUniformBuffer(BasePassParameters);
-			}
-			
-			GraphBuilder.AddPass(
-				RDG_EVENT_NAME("Translucency"),
-				PassParameters,
-				ERDGPassFlags::Raster,
-				[this, &View, PassParameters](FRHICommandList& RHICmdList)
-				{
-					RenderTranslucency(RHICmdList, View, Views, StandardTranslucencyPass, StandardTranslucencyMeshPass, &PassParameters->InstanceCullingDrawParams);
-				});
+			AddTranslucencyPass();
 		}
 	}
 }
