fix per-object shadow size not being correctly clamped to the max shadow size

#Synced-CL 51004502

diff --git a/Engine/Source/Runtime/Renderer/Private/ShadowSetup.cpp b/Engine/Source/Runtime/Renderer/Private/ShadowSetup.cpp
--- a/Engine/Source/Runtime/Renderer/Private/ShadowSetup.cpp
+++ b/Engine/Source/Runtime/Renderer/Private/ShadowSetup.cpp
@@ -3595,11 +3595,12 @@ void FSceneRenderer::CreatePerObjectProjectedShadow(
 		{
 			// Round down to the nearest power of two so that resolution changes are always doubling or halving the resolution, which increases filtering stability
 			// Use the max resolution if the desired resolution is larger than that
-			const int32 SizeX = MaxDesiredResolution >= MaxShadowResolution ? MaxShadowResolution : 
-				bCanUseShadowDepthAtlas ? 
-				(1 << (FMath::CeilLogTwo(MaxDesiredResolution) - 1)) : 
-				(FMath::RoundUpToPowerOfTwo(MaxDesiredResolution) - SHADOW_BORDER * 2); // see AllocateAtlasedShadowDepthTargets, we round up target size there
-			
+			int32 SizeX = bCanUseShadowDepthAtlas	? (1 << (FMath::CeilLogTwo(MaxDesiredResolution) - 1))
+													: (FMath::RoundUpToPowerOfTwo(MaxDesiredResolution) - SHADOW_BORDER * 2); // see AllocateAtlasedShadowDepthTargets, we round up target size there
+
+			SizeX = FMath::Min<int32>(SizeX, MaxShadowResolution);
+
+
 			if (bOpaque && bCreateOpaqueObjectShadow && (bOpaqueShadowIsVisibleThisFrame || bShadowIsPotentiallyVisibleNextFrame))
 			{
 				// Create a projected shadow for this interaction's shadow.
