Virtual texture async update changes.
Remove parallel feedback and gather tasks because we get most threading saving using a single async update task.
Going wide on work inside that leads to increased waits more often than not. Previously the default was to disable that path anyway.
This allows some removal of the task management code, and the PageUpdate lock and buffer that was needed to avoid race conditions during the parallel FGatherRequestsTask.
Added support for time slicing of the Gather and Submit  phases of async updating.
This is controlled by new cvars: .VT.MaxGatherTimePerFrame and r.VT.MaxSubmitTimePerFrame.
These default to 0 (no time slicing) but values of 1.5ms seem to work well on heavy virtual texture content.
Time slicing for the gather phase requires the pages to be sorted first so that virtual textures with larger screen coverage are more likely to make it through the time slicer.
Time slicing for the submit phase still has to make sure that root pages are updated even after the time interval is complete.
Also added r.VT.AsyncPageRequestTask.EarlyOut which is an experimental early out from the async task at the point that we wait on it on the render thread.
This may help reduces render thread stalls in the very worst cases where time slicing didn't work well due to task preemption.
#Synced-CL 51089442

diff --git a/Engine/Source/Runtime/Engine/Private/VT/VirtualTextureScalability.cpp b/Engine/Source/Runtime/Engine/Private/VT/VirtualTextureScalability.cpp
--- a/Engine/Source/Runtime/Engine/Private/VT/VirtualTextureScalability.cpp
+++ b/Engine/Source/Runtime/Engine/Private/VT/VirtualTextureScalability.cpp
@@ -69,6 +69,22 @@ namespace VirtualTextureScalability
 	);
 #endif
 
+	static TAutoConsoleVariable<float> CVarVTMaxGatherTimePerFrame(
+		TEXT("r.VT.MaxGatherTimePerFrame"),
+		0,
+		TEXT("Time slice limit in milliseconds used to limit time spent gathering virtual texture requests per frame.\n")
+		TEXT("0 (default) disables."),
+		ECVF_RenderThreadSafe | ECVF_Scalability
+	);
+
+	static TAutoConsoleVariable<float> CVarVTMaxSubmitTimePerFrame(
+		TEXT("r.VT.MaxSubmitTimePerFrame"),
+		0,
+		TEXT("Time slice limit in milliseconds used to limit time spent submitting virtual texture requests per frame.\n")
+		TEXT("0 (default) disables."),
+		ECVF_RenderThreadSafe | ECVF_Scalability
+	);
+
 	static TAutoConsoleVariable<int32> CVarVTMaxReleasedPerFrame(
 		TEXT("r.VT.MaxReleasedPerFrame"),
 		5,
@@ -222,6 +238,28 @@ namespace VirtualTextureScalability
 #endif
 	}
 
+	float GetMaxGatherTimePerFrame()
+	{
+#if WITH_EDITOR
+		if (GIsEditor)
+		{
+			return 0.0f;
+		}
+#endif
+		return CVarVTMaxGatherTimePerFrame.GetValueOnAnyThread() * 0.001f;
+	}
+
+	float GetMaxSubmitTimePerFrame()
+	{
+#if WITH_EDITOR
+		if (GIsEditor)
+		{
+			return 0.0f;
+		}
+#endif
+		return CVarVTMaxSubmitTimePerFrame.GetValueOnAnyThread() * 0.001f;
+	}
+
 	int32 GetMaxAllocatedVTReleasedPerFrame()
 	{
 #if WITH_EDITOR
diff --git a/Engine/Source/Runtime/Engine/Public/VT/VirtualTextureScalability.h b/Engine/Source/Runtime/Engine/Public/VT/VirtualTextureScalability.h
--- a/Engine/Source/Runtime/Engine/Public/VT/VirtualTextureScalability.h
+++ b/Engine/Source/Runtime/Engine/Public/VT/VirtualTextureScalability.h
@@ -14,6 +14,10 @@ namespace VirtualTextureScalability
 	ENGINE_API int32 GetMaxPagesProducedPerFrame();
 	/** Get max update rate of already mapped virtual texture pages. */
 	ENGINE_API int32 GetMaxContinuousUpdatesPerFrame();
+	/** Get max time in seconds to spend in a single time slice of the gather page requests task. */
+	ENGINE_API float GetMaxGatherTimePerFrame();
+	/** Get max time in seconds to spend in a single time slice of the submit page requests task. */
+	ENGINE_API float GetMaxSubmitTimePerFrame();
 	/** Get max allocated virtual textures to release per frame. */
 	ENGINE_API int32 GetMaxAllocatedVTReleasedPerFrame();
 	/** Get the size of the virtual texture feedback tiles. */
diff --git a/Engine/Source/Runtime/Renderer/Private/VT/TexturePagePool.h b/Engine/Source/Runtime/Renderer/Private/VT/TexturePagePool.h
--- a/Engine/Source/Runtime/Renderer/Private/VT/TexturePagePool.h
+++ b/Engine/Source/Runtime/Renderer/Private/VT/TexturePagePool.h
@@ -28,8 +28,6 @@ public:
 
 	void Initialize(uint32 InNumPages);
 
-	FCriticalSection& GetLock() { return CriticalSection; }
-
 	uint32 GetNumPages() const { return FMath::Max(NumPages, NumReservedPages) - NumReservedPages; }
 	uint32 GetNumLockedPages() const { return GetNumPages() - FreeHeap.Num(); }
 	uint32 GetNumMappedPages() const { return NumPagesMapped; }
@@ -225,8 +223,6 @@ private:
 		AddMappingToList(FreeHeadIndex, Index);
 	}
 
-	FCriticalSection CriticalSection;
-
 	FBinaryHeap<uint32, uint16> FreeHeap;
 
 	FHashTable PageHash;
diff --git a/Engine/Source/Runtime/Renderer/Private/VT/UniquePageList.h b/Engine/Source/Runtime/Renderer/Private/VT/UniquePageList.h
--- a/Engine/Source/Runtime/Renderer/Private/VT/UniquePageList.h
+++ b/Engine/Source/Runtime/Renderer/Private/VT/UniquePageList.h
@@ -16,11 +16,13 @@ public:
 	void	Add( uint32 Page, uint32 Count );
 
 	uint32	GetNum() const					{ return NumPages; }
-	uint32	GetPage( uint32 Index ) const	{ return Pages[ Index ]; }
-	uint32	GetCount( uint32 Index ) const	{ return Counts[ Index ]; }
+	uint32	GetPage(uint32 Index) const		{ checkSlow(Index < NumPages); const uint32 SortedIndex = bSorted ? SortOrder[Index] : Index; checkSlow(SortedIndex < NumPages); return Pages[SortedIndex]; }
+	uint32	GetCount(uint32 Index) const	{ checkSlow(Index < NumPages); const uint32 SortedIndex = bSorted ? SortOrder[Index] : Index; checkSlow(SortedIndex < NumPages); return Counts[SortedIndex]; }
 
 	void	MergePages(const FUniquePageList* RESTRICT Other);
 
+	void	SortPages();
+
 private:
 	enum
 	{
@@ -29,15 +31,18 @@ private:
 	};
 
 	bool bInitialized;
+	bool bSorted;
 	uint32 NumPages;
 	uint32 MaxNumCollisions;
 	uint16 HashIndices[HashSize];
 	uint32 Pages[ MaxUniquePages ];
 	uint16 Counts[ MaxUniquePages ];
+	uint16 SortOrder[ MaxUniquePages ];
 };
 
 inline FUniquePageList::FUniquePageList()
 	: bInitialized( false )
+	, bSorted( false )
 	, NumPages( 0 )
 	, MaxNumCollisions( 0 )
 {}
@@ -78,6 +83,7 @@ inline void FUniquePageList::Add( uint32 Page, uint32 Count )
 		HashIndex = (HashIndex + 1u) & (HashSize - 1u);
 		++NumCollisions;
 	}
+	bSorted = false;
 #if DO_GUARD_SLOW
 	MaxNumCollisions = FMath::Max(MaxNumCollisions, NumCollisions);
 #endif // DO_GUARD_SLOW
@@ -90,3 +96,14 @@ inline void FUniquePageList::MergePages(const FUniquePageList* RESTRICT Other)
 		Add(Other->Pages[Index], Other->Counts[Index]);
 	}
 }
+
+void FUniquePageList::SortPages()
+{
+	for (uint32 Index = 0; Index < NumPages; ++Index)
+	{
+		SortOrder[Index] = Index;
+	}
+	Algo::Sort(MakeArrayView(SortOrder, NumPages), [this](uint32 A, uint32 B) { return Counts[A] == Counts[B] ? A < B : Counts[A] > Counts[B]; });
+
+	bSorted = true;
+}
diff --git a/Engine/Source/Runtime/Renderer/Private/VT/VirtualTextureSystem.cpp b/Engine/Source/Runtime/Renderer/Private/VT/VirtualTextureSystem.cpp
--- a/Engine/Source/Runtime/Renderer/Private/VT/VirtualTextureSystem.cpp
+++ b/Engine/Source/Runtime/Renderer/Private/VT/VirtualTextureSystem.cpp
@@ -41,16 +41,12 @@ DECLARE_CYCLE_STAT(TEXT("Submit Requests"), STAT_ProcessRequests_Submit, STATGRO
 DECLARE_CYCLE_STAT(TEXT("Map Requests"), STAT_ProcessRequests_Map, STATGROUP_VirtualTexturing);
 DECLARE_CYCLE_STAT(TEXT("Map New VTs"), STAT_ProcessRequests_MapNew, STATGROUP_VirtualTexturing);
 DECLARE_CYCLE_STAT(TEXT("Finalize Requests"), STAT_ProcessRequests_Finalize, STATGROUP_VirtualTexturing);
-DECLARE_CYCLE_STAT(TEXT("Merge Unique Pages"), STAT_ProcessRequests_MergePages, STATGROUP_VirtualTexturing);
-DECLARE_CYCLE_STAT(TEXT("Merge Requests"), STAT_ProcessRequests_MergeRequests, STATGROUP_VirtualTexturing);
-DECLARE_CYCLE_STAT(TEXT("Submit Tasks"), STAT_ProcessRequests_SubmitTasks, STATGROUP_VirtualTexturing);
 DECLARE_CYCLE_STAT(TEXT("Wait Tasks"), STAT_ProcessRequests_WaitTasks, STATGROUP_VirtualTexturing);
 
 DECLARE_CYCLE_STAT(TEXT("Queue Adaptive Requests"), STAT_ProcessRequests_QueueAdaptiveRequests, STATGROUP_VirtualTexturing);
 DECLARE_CYCLE_STAT(TEXT("Finalize Adaptive Requests"), STAT_ProcessRequests_UpdateAdaptiveAllocations, STATGROUP_VirtualTexturing);
 
 DECLARE_CYCLE_STAT(TEXT("Feedback Map"), STAT_FeedbackMap, STATGROUP_VirtualTexturing);
-DECLARE_CYCLE_STAT(TEXT("Feedback Analysis"), STAT_FeedbackAnalysis, STATGROUP_VirtualTexturing);
 DECLARE_CYCLE_STAT(TEXT("Page Table Updates"), STAT_PageTableUpdates, STATGROUP_VirtualTexturing);
 DECLARE_CYCLE_STAT(TEXT("Flush Cache"), STAT_FlushCache, STATGROUP_VirtualTexturing);
 DECLARE_CYCLE_STAT(TEXT("Update Residency Tracking"), STAT_ResidencyTracking, STATGROUP_VirtualTexturing);
@@ -90,24 +86,6 @@ static TAutoConsoleVariable<int32> CVarVTEnableFeedbackProduce(
 	TEXT("Enable page loading triggered by GPU virtual texture feedback."),
 	ECVF_RenderThreadSafe
 );
-static TAutoConsoleVariable<int32> CVarVTParallelFeedbackTasks(
-	TEXT("r.VT.ParallelFeedbackTasks"),
-	0,
-	TEXT("Use worker threads for virtual texture feedback tasks."),
-	ECVF_RenderThreadSafe
-);
-static TAutoConsoleVariable<int32> CVarVTNumFeedbackTasks(
-	TEXT("r.VT.NumFeedbackTasks"),
-	1,
-	TEXT("Number of tasks to create to read virtual texture feedback."),
-	ECVF_RenderThreadSafe
-);
-static TAutoConsoleVariable<int32> CVarVTNumGatherTasks(
-	TEXT("r.VT.NumGatherTasks"),
-	1,
-	TEXT("Number of tasks to create to combine virtual texture feedback."),
-	ECVF_RenderThreadSafe
-);
 static TAutoConsoleVariable<int32> CVarVTEnablePlayback(
 	TEXT("r.VT.EnablePlayback"),
 	1,
@@ -121,12 +99,6 @@ static TAutoConsoleVariable<float> CVarVTPlaybackMipBias(
 	TEXT("This is in addition to the automatic bias applied for the difference between capture and playback resolutions."),
 	ECVF_RenderThreadSafe
 );
-static TAutoConsoleVariable<int32> CVarVTPageUpdateFlushCount(
-	TEXT("r.VT.PageUpdateFlushCount"),
-	8,
-	TEXT("Number of page updates to buffer before attempting to flush by taking a lock."),
-	ECVF_RenderThreadSafe
-);
 static TAutoConsoleVariable<int32> CVarVTForceContinuousUpdate(
 	TEXT("r.VT.ForceContinuousUpdate"),
 	0,
@@ -176,35 +148,69 @@ static TAutoConsoleVariable<int32> CVarVTAsyncPageRequestTask(
 	TEXT("Performs VT page requests on an async task."),
 	ECVF_RenderThreadSafe
 );
+static TAutoConsoleVariable<float> CVarVTAsyncPageRequestTaskEarlyOut(
+	TEXT("r.VT.AsyncPageRequestTask.EarlyOut"),
+	0.0f,
+	TEXT("Allow early out of non essential parts of async task instead of waiting for full completion.")
+	TEXT("0 (default) of less disables. A value above 0 enables and is the minimum time in milliseconds before we allow the early out."),
+	ECVF_RenderThreadSafe
+);
 static TAutoConsoleVariable<int32> CVarVTWaitBeforeUpdate(
 	TEXT("r.VT.WaitBeforeUpdate"),
 	0,
 	TEXT("Wait on the SceneRenderBuilder before kicking VT async update."),
 	ECVF_RenderThreadSafe
 );
-
-TAutoConsoleVariable<int32> CVarVTSortTileRequestsByPriority(
+static TAutoConsoleVariable<int32> CVarVTSortTileRequestsByPriority(
 	TEXT("r.VT.SortTileRequestsByPriority"),
 	1,
 	TEXT("Sorts the list of tile requests by priority before submitting. This allows priority requests to get processed before others in case of throttling"),
 	ECVF_RenderThreadSafe
 );
 
+static bool ShouldSyncProduceLockedTiles()
+{
+#if WITH_EDITOR
+	if (GIsEditor)
+	{
+		// Always force loading of locked root pages in editor to be sure that any build steps give best quality.
+		return true;
+	}
+#endif
+	// If we return false then we may render with a VT before the root pages are mapped.
+	// When that happens the shader switches to using the single color fallback value instead of sampling the VT.
+	// When the root page is finally mapped we will return to normal high quality VT sampling.
+	// todo[VT]: Make root pages always resident so that we never need to load sync the root pages.
+	return CVarVTSyncProduceLockedTiles.GetValueOnRenderThread() != 0;
+}
+
+static float GetMinSubmitAsyncEarlyOutTime()
+{
+#if WITH_EDITOR
+	if (GIsEditor)
+	{
+		return 0.0f;
+	}
+#endif
+	return FMath::Max(0.0f, CVarVTAsyncPageRequestTaskEarlyOut.GetValueOnRenderThread()) * 0.001f;
+}
+
 FVirtualTextureUpdateSettings::FVirtualTextureUpdateSettings()
 {
 	bEnableFeedback = CVarVTEnableFeedback.GetValueOnRenderThread() != 0;
 	bEnableFeedbackProduce = CVarVTEnableFeedbackProduce.GetValueOnRenderThread() != 0;
 	bEnablePlayback = CVarVTEnablePlayback.GetValueOnRenderThread() != 0;
-	bParallelFeedbackTasks = CVarVTParallelFeedbackTasks.GetValueOnRenderThread() != 0;
 	bForceContinuousUpdate = CVarVTForceContinuousUpdate.GetValueOnRenderThread() != 0;
-	bForceSyncPageUpdate = false;
-	NumFeedbackTasks = CVarVTNumFeedbackTasks.GetValueOnRenderThread();
-	NumGatherTasks = CVarVTNumGatherTasks.GetValueOnRenderThread();
-	MaxGatherPagesBeforeFlush = CVarVTPageUpdateFlushCount.GetValueOnRenderThread();
+	bSyncProduceAllTiles = false;
+	bSyncProduceLockedTiles = ShouldSyncProduceLockedTiles();
+	bEnableAsyncTasks = CVarVTAsyncPageRequestTask.GetValueOnRenderThread() != 0;
 	MaxRVTPageUploads = VirtualTextureScalability::GetMaxUploadsPerFrame();
 	MaxSVTPageUploads = VirtualTextureScalability::GetMaxUploadsPerFrameForStreamingVT();
 	MaxPagesProduced = VirtualTextureScalability::GetMaxPagesProducedPerFrame();
 	MaxContinuousUpdates = VirtualTextureScalability::GetMaxContinuousUpdatesPerFrame();
+	MaxGatherTime = VirtualTextureScalability::GetMaxGatherTimePerFrame();
+	MaxSubmitTime = VirtualTextureScalability::GetMaxSubmitTimePerFrame();
+	MinSubmitAsyncEarlyOutTime = GetMinSubmitAsyncEarlyOutTime();
 }
 
 FVirtualTextureUpdateSettings::FVirtualTextureUpdateSettings(FViewFamilyInfo const& InViewFamily)
@@ -217,23 +223,6 @@ FVirtualTextureUpdateSettings::FVirtualTextureUpdateSettings(FViewFamilyInfo con
 	PlaybackViewSize = InViewFamily.SceneTexturesConfig.Extent;
 }
 
-static bool ShouldSyncProduceLockedTiles()
-{
-#if WITH_EDITOR
-	if (GIsEditor)
-	{
-		// Always force loading of locked root pages in editor to be sure that any build steps give best quality.
-		return true;
-	}
-#endif
-	// If we return false then we may render with a VT before the root pages are mapped.
-	// When that happens the shader switches to using the single color fallback value instead of sampling the VT.
-	// When the root page is finally mapped we will return to normal high quality VT sampling.
-	// todo[VT]: Make root pages always resident so that we never need to load sync the root pages.
-	return CVarVTSyncProduceLockedTiles.GetValueOnRenderThread() != 0;
-}
-
-
 static FORCEINLINE uint32 EncodePage(uint32 ID, uint32 vLevel, uint32 vTileX, uint32 vTileY)
 {
 	const uint32 vLevelPlus1 = vLevel + 1u;
@@ -246,114 +235,31 @@ static FORCEINLINE uint32 EncodePage(uint32 ID, uint32 vLevel, uint32 vTileX, ui
 	return Page;
 }
 
-struct FPageUpdateBuffer
-{
-	static const uint32 PageCapacity = 128u;
-	uint16 PhysicalAddresses[PageCapacity];
-	uint32 PrevPhysicalAddress = ~0u;
-	uint32 NumPages = 0u;
-	uint32 NumPageUpdates = 0u;
-	uint32 WorkingSetSize = 0u;
-};
-
 struct FFeedbackAnalysisParameters
 {
-	FVirtualTextureSystem* System = nullptr;
 	const FUintPoint* FeedbackBuffer = nullptr;
 	FUniquePageList* UniquePageList = nullptr;
 	uint32 FeedbackSize = 0u;
 };
 
-class FFeedbackAnalysisTask
-{
-public:
-	explicit FFeedbackAnalysisTask(const FFeedbackAnalysisParameters& InParams) : Parameters(InParams) {}
-
-	FFeedbackAnalysisParameters Parameters;
-
-	static void DoTask(FFeedbackAnalysisParameters& InParams)
-	{
-		TRACE_CPUPROFILER_EVENT_SCOPE(FVirtualTextureSystem::FeedbackAnalysisTask);
-		InParams.UniquePageList->Initialize();
-		InParams.System->FeedbackAnalysisTask(InParams);
-	}
-
-	void DoTask(ENamedThreads::Type CurrentThread, const FGraphEventRef& MyCompletionGraphEvent)
-	{
-		FTaskTagScope TaskTagScope(ETaskTag::EParallelRenderingThread);
-		DoTask(Parameters);
-	}
-
-	static ESubsequentsMode::Type GetSubsequentsMode() { return ESubsequentsMode::TrackSubsequents; }
-	ENamedThreads::Type GetDesiredThread() { return ENamedThreads::AnyNormalThreadNormalTask; }
-	FORCEINLINE TStatId GetStatId() const { RETURN_QUICK_DECLARE_CYCLE_STAT(FFeedbackAnalysisTask, STATGROUP_VirtualTexturing); }
-};
-
 struct FAddRequestedTilesParameters
 {
-	FVirtualTextureSystem* System = nullptr;
 	uint32 ViewLevelBias = 0u;
 	const uint32* RequestBuffer = nullptr;
 	uint32 NumRequests = 0u;
 	FUniquePageList* UniquePageList = nullptr;
 };
 
-class FAddRequestedTilesTask
-{
-public:
-	explicit FAddRequestedTilesTask(const FAddRequestedTilesParameters& InParams) : Parameters(InParams) {}
-
-	FAddRequestedTilesParameters Parameters;
-
-	static void DoTask(FAddRequestedTilesParameters& InParams)
-	{
-		InParams.UniquePageList->Initialize();
-		InParams.System->AddRequestedTilesTask(InParams);
-	}
-
-	void DoTask(ENamedThreads::Type CurrentThread, const FGraphEventRef& MyCompletionGraphEvent)
-	{
-		FTaskTagScope TaskTagScope(ETaskTag::EParallelRenderingThread);
-		DoTask(Parameters);
-	}
-
-	static ESubsequentsMode::Type GetSubsequentsMode() { return ESubsequentsMode::TrackSubsequents; }
-	ENamedThreads::Type GetDesiredThread() { return ENamedThreads::AnyHiPriThreadHiPriTask; }
-	FORCEINLINE TStatId GetStatId() const { RETURN_QUICK_DECLARE_CYCLE_STAT(FAddRequestedTilesTask, STATGROUP_VirtualTexturing); }
-};
-
 struct FGatherRequestsParameters
 {
-	FVirtualTextureSystem* System = nullptr;
 	const FUniquePageList* UniquePageList = nullptr;
-	FPageUpdateBuffer* PageUpdateBuffers = nullptr;
 	FUniqueRequestList* RequestList = nullptr;
-	uint32 PageUpdateFlushCount = 0u;
 	uint32 PageStartIndex = 0u;
 	uint32 NumPages = 0u;
 	uint32 FrameRequested;
 	bool bForceContinuousUpdate = false;
 	bool bEnableLoadRequests = false;
-};
-
-class FGatherRequestsTask
-{
-public:
-	explicit FGatherRequestsTask(const FGatherRequestsParameters& InParams) : Parameters(InParams) {}
-
-	FGatherRequestsParameters Parameters;
-
-	void DoTask(ENamedThreads::Type CurrentThread, const FGraphEventRef& MyCompletionGraphEvent)
-	{
-		FTaskTagScope TaskTagScope(ETaskTag::EParallelRenderingThread);
-		TRACE_CPUPROFILER_EVENT_SCOPE(FVirtualTextureSystem::GatherRequestsTask);
-		Parameters.RequestList->Initialize();
-		Parameters.System->GatherRequestsTask(Parameters);
-	}
-
-	static ESubsequentsMode::Type GetSubsequentsMode() { return ESubsequentsMode::TrackSubsequents; }
-	ENamedThreads::Type GetDesiredThread() { return ENamedThreads::AnyHiPriThreadHiPriTask; }
-	FORCEINLINE TStatId GetStatId() const { RETURN_QUICK_DECLARE_CYCLE_STAT(FGatherRequestsTask, STATGROUP_VirtualTexturing); }
+	float TimeBudget = 0.0f;
 };
 
 static FVirtualTextureSystem* GVirtualTextureSystem = nullptr;
@@ -1385,6 +1291,8 @@ void FVirtualTextureSystem::LoadPendingTiles(FRDGBuilder& GraphBuilder, ERHIFeat
 	if (PackedTiles.Num() > 0)
 	{
 		FVirtualTextureUpdateSettings Settings;
+		Settings.EnableThrottling(false);
+		
 		FConcurrentLinearBulkObjectAllocator Allocator;
 
 		FUniquePageList* UniquePageList = Allocator.Create<FUniquePageList>();
@@ -1396,11 +1304,12 @@ void FVirtualTextureSystem::LoadPendingTiles(FRDGBuilder& GraphBuilder, ERHIFeat
 
 		FUniqueRequestList* RequestList = Allocator.Create<FUniqueRequestList>(Allocator);
 		RequestList->Initialize();
+
 		GatherLockedTileRequests(RequestList);
 		GatherRequests(RequestList, UniquePageList, Frame, Allocator, Settings);
 		// No need to sort requests, since we're submitting all of them here (no throttling)
 		AllocateResources(GraphBuilder);
-		SubmitRequests(GraphBuilder.RHICmdList, FeatureLevel, Allocator, Settings, RequestList, false);
+		SubmitRequests(GraphBuilder.RHICmdList, FeatureLevel, Allocator, Settings, RequestList);
 		FinalizeRequests(GraphBuilder, nullptr);
 
 		// Swap any failed tile locks to be picked up by next BeginUpdate()
@@ -1542,6 +1451,8 @@ void FVirtualTextureSystem::RequestTilesForRegionInternal(const IAllocatedVirtua
 
 void FVirtualTextureSystem::FeedbackAnalysisTask(const FFeedbackAnalysisParameters& Parameters)
 {
+	TRACE_CPUPROFILER_EVENT_SCOPE(FVirtualTextureSystem::FeedbackAnalysisTask);
+
 	FUniquePageList* RESTRICT RequestedPageList = Parameters.UniquePageList;
 	const FUintPoint* RESTRICT Buffer = Parameters.FeedbackBuffer;
 	const uint32 BufferSize = Parameters.FeedbackSize;
@@ -1577,150 +1488,44 @@ void FVirtualTextureSystem::FeedbackAnalysisTask(const FFeedbackAnalysisParamete
 
 void FVirtualTextureSystem::GatherRequests(FUniqueRequestList* MergedRequestList, const FUniquePageList* UniquePageList, uint32 FrameRequested, FConcurrentLinearBulkObjectAllocator& Allocator, FVirtualTextureUpdateSettings const& Settings)
 {
-	const uint32 MaxNumGatherTasks = FMath::Clamp((uint32)Settings.NumGatherTasks, 1u, MaxNumTasks);
-	const uint32 PageUpdateFlushCount = FMath::Min<uint32>(Settings.MaxGatherPagesBeforeFlush, FPageUpdateBuffer::PageCapacity);
-
-	FGatherRequestsParameters GatherRequestsParameters[MaxNumTasks];
-	uint32 NumGatherTasks = 0u;
-	{
-		const uint32 MinNumPagesPerTask = 64u;
-		const uint32 NumPagesPerTask = FMath::Max(FMath::DivideAndRoundUp(UniquePageList->GetNum(), MaxNumGatherTasks), MinNumPagesPerTask);
-		const uint32 NumPages = UniquePageList->GetNum();
-		uint32 StartPageIndex = 0u;
-		while (StartPageIndex < NumPages)
-		{
-			const uint32 NumPagesForTask = FMath::Min(NumPagesPerTask, NumPages - StartPageIndex);
-			if (NumPagesForTask > 0u)
-			{
-				const uint32 TaskIndex = NumGatherTasks++;
-				FGatherRequestsParameters& Params = GatherRequestsParameters[TaskIndex];
-				Params.System = this;
-				Params.FrameRequested = FrameRequested;
-				Params.bForceContinuousUpdate = Settings.bForceContinuousUpdate;
-				Params.bEnableLoadRequests = Settings.bEnableFeedbackProduce;
-				Params.UniquePageList = UniquePageList;
-				Params.PageUpdateFlushCount = PageUpdateFlushCount;
-				Params.PageUpdateBuffers = Allocator.CreateArray<FPageUpdateBuffer>(PhysicalSpaces.Num());
-				if (TaskIndex == 0u)
-				{
-					Params.RequestList = MergedRequestList;
-				}
-				else
-				{
-					Params.RequestList = Allocator.Create<FUniqueRequestList>(Allocator);
-				}
-				Params.PageStartIndex = StartPageIndex;
-				Params.NumPages = NumPagesForTask;
-				StartPageIndex += NumPagesForTask;
-			}
-		}
-	}
-
-	// Kick all of the tasks
-	FGraphEventArray Tasks;
-	if (NumGatherTasks > 1u)
-	{
-		SCOPE_CYCLE_COUNTER(STAT_ProcessRequests_SubmitTasks);
-		Tasks.Reserve(NumGatherTasks - 1u);
-		for (uint32 TaskIndex = 1u; TaskIndex < NumGatherTasks; ++TaskIndex)
-		{
-			Tasks.Add(TGraphTask<FGatherRequestsTask>::CreateTask().ConstructAndDispatchWhenReady(GatherRequestsParameters[TaskIndex]));
-		}
-	}
-
-	if (NumGatherTasks > 0u)
-	{
-		SCOPE_CYCLE_COUNTER(STAT_ProcessRequests_Gather);
-
-		// first task can run on this thread
-		GatherRequestsTask(GatherRequestsParameters[0]);
-
-		// Wait for them to complete
-		if (Tasks.Num() > 0)
-		{
-			SCOPE_CYCLE_COUNTER(STAT_ProcessRequests_WaitTasks);
-
-			FTaskGraphInterface::Get().WaitUntilTasksComplete(Tasks);
-		}
-	}
+	FGatherRequestsParameters Parameters;
+	Parameters.UniquePageList = UniquePageList;
+	Parameters.RequestList = MergedRequestList;
+	Parameters.PageStartIndex = 0;
+	Parameters.NumPages = UniquePageList->GetNum();
+	Parameters.FrameRequested = FrameRequested;
+	Parameters.bForceContinuousUpdate = Settings.bForceContinuousUpdate;
+	Parameters.bEnableLoadRequests = Settings.bEnableFeedbackProduce;
+	Parameters.TimeBudget = Settings.MaxGatherTime;
 
-	// Merge request lists for all tasks
-	if (NumGatherTasks > 1u)
-	{
-		SCOPE_CYCLE_COUNTER(STAT_ProcessRequests_MergeRequests);
-		for (uint32 TaskIndex = 1u; TaskIndex < NumGatherTasks; ++TaskIndex)
-		{
-			MergedRequestList->MergeRequests(GatherRequestsParameters[TaskIndex].RequestList, Allocator);
-		}
-	}
+	GatherRequestsTask(Parameters);
 }
 
-void FVirtualTextureSystem::AddPageUpdate(FPageUpdateBuffer* Buffers, uint32 FlushCount, uint32 PhysicalSpaceID, uint16 pAddress)
+void FVirtualTextureSystem::AddPageUpdate(uint32 PhysicalSpaceID, uint16 pAddress, uint32 CurrentFrame)
 {
-	FPageUpdateBuffer& RESTRICT Buffer = Buffers[PhysicalSpaceID];
-	if (pAddress == Buffer.PrevPhysicalAddress)
-	{
-		return;
-	}
-	Buffer.PrevPhysicalAddress = pAddress;
-
-	bool bLocked = false;
-	if (Buffer.NumPages >= FlushCount)
-	{
-		// Once we've passed a certain threshold of pending pages to update, try to take the lock then apply the updates
-		FVirtualTexturePhysicalSpace* RESTRICT PhysicalSpace = GetPhysicalSpace(PhysicalSpaceID);
-		FTexturePagePool& RESTRICT PagePool = PhysicalSpace->GetPagePool();
-		FCriticalSection& RESTRICT Lock = PagePool.GetLock();
-
-		if (Buffer.NumPages >= FPageUpdateBuffer::PageCapacity)
-		{
-			// If we've reached capacity, need to take the lock no matter what, may potentially block here
-			Lock.Lock();
-			bLocked = true;
-		}
-		else
-		{
-			// try to take the lock, but avoid stalling
-			bLocked = Lock.TryLock();
-		}
-
-		if(bLocked)
-		{
-			const uint32 CurrentFrame = Frame;
-			PagePool.UpdateUsage(CurrentFrame, pAddress); // Update current request now, if we manage to get the lock
-			for (uint32 i = 0u; i < Buffer.NumPages; ++i)
-			{
-				PagePool.UpdateUsage(CurrentFrame, Buffer.PhysicalAddresses[i]);
-			}
-			Lock.Unlock();
-			Buffer.NumPageUpdates += (Buffer.NumPages + 1u);
-			Buffer.NumPages = 0u;
-		}
-	}
-
-	// Only need to buffer if we didn't lock (otherwise this has already been updated)
-	if (!bLocked)
-	{
-		check(Buffer.NumPages < FPageUpdateBuffer::PageCapacity);
-		Buffer.PhysicalAddresses[Buffer.NumPages++] = pAddress;
-	}
+	FVirtualTexturePhysicalSpace* RESTRICT PhysicalSpace = GetPhysicalSpace(PhysicalSpaceID);
+	FTexturePagePool& RESTRICT PagePool = PhysicalSpace->GetPagePool();
+	PagePool.UpdateUsage(CurrentFrame, pAddress);
 }
 
 void FVirtualTextureSystem::GatherRequestsTask(const FGatherRequestsParameters& Parameters)
 {
+	TRACE_CPUPROFILER_EVENT_SCOPE(FVirtualTextureSystem::GatherRequestsTask);
+	SCOPE_CYCLE_COUNTER(STAT_ProcessRequests_Gather);
+
 	const FUniquePageList* RESTRICT UniquePageList = Parameters.UniquePageList;
-	FPageUpdateBuffer* RESTRICT PageUpdateBuffers = Parameters.PageUpdateBuffers;
 	FUniqueRequestList* RESTRICT RequestList = Parameters.RequestList;
-	const uint32 PageUpdateFlushCount = Parameters.PageUpdateFlushCount;
+	const uint32 PageStartIndex = Parameters.PageStartIndex;
 	const uint32 PageEndIndex = Parameters.PageStartIndex + Parameters.NumPages;
 	const bool bForceContinuousUpdate = Parameters.bForceContinuousUpdate;
+	const double TimeSliceEnd = Parameters.TimeBudget > 0 ? FPlatformTime::Seconds() + Parameters.TimeBudget : 0;
 
 	uint32 NumRequestsPages = 0u;
 	uint32 NumResidentPages = 0u;
 	uint32 NumNonResidentPages = 0u;
 	uint32 NumPrefetchPages = 0u;
 
-	for (uint32 i = Parameters.PageStartIndex; i < PageEndIndex; ++i)
+	for (uint32 i = PageStartIndex; (i < PageEndIndex) && (TimeSliceEnd == 0 || FPlatformTime::Seconds() < TimeSliceEnd); ++i)
 	{
 		const uint32 PageEncoded = UniquePageList->GetPage(i);
 		const uint32 PageCount = UniquePageList->GetCount(i);
@@ -1782,7 +1587,7 @@ void FVirtualTextureSystem::GatherRequestsTask(const FGatherRequestsParameters&
 #endif // DO_GUARD_SLOW
 
 					// Page is already resident, just need to update LRU free list
-					AddPageUpdate(PageUpdateBuffers, PageUpdateFlushCount, PhysicalSpaceIDAndAddress.PhysicalSpaceID, PhysicalSpaceIDAndAddress.pAddress);
+					AddPageUpdate(PhysicalSpaceIDAndAddress.PhysicalSpaceID, PhysicalSpaceIDAndAddress.pAddress, Frame);
 
 					// If continuous update flag is set then add this to pages which can be potentially updated if we have spare upload bandwidth
 					if (bForceContinuousUpdate || Space->GetDescription().bContinuousUpdate)
@@ -1794,7 +1599,6 @@ void FVirtualTextureSystem::GatherRequestsTask(const FGatherRequestsParameters&
 						RequestList->AddContinuousUpdateRequest(FVirtualTextureLocalTileRequest(LocalTile, EVTProducerPriority::Normal, EVTInvalidatePriority::Normal));
 					}
 
-					++PageUpdateBuffers[PhysicalSpaceIDAndAddress.PhysicalSpaceID].WorkingSetSize;
 					++NumResidentPages;
 				}
 				else
@@ -1862,12 +1666,6 @@ void FVirtualTextureSystem::GatherRequestsTask(const FGatherRequestsParameters&
 			
 			const uint32 ProducerPhysicalGroupIndex = AllocatedVT->GetProducerPhysicalGroupIndexForPageTableLayer(PageTableLayerIndex);
 			ProducerGroupMaskToLoad[ProducerIndex] |= 1 << ProducerPhysicalGroupIndex;
-
-			const FVirtualTexturePhysicalSpace* RESTRICT PhysicalSpace = AllocatedVT->GetPhysicalSpaceForPageTableLayer(PageTableLayerIndex);
-			if (PhysicalSpace)
-			{
-				++PageUpdateBuffers[PhysicalSpace->GetID()].WorkingSetSize;
-			}
 		}
 
 		const uint32 vDimensions = Space->GetDimensions();
@@ -1958,7 +1756,7 @@ void FVirtualTextureSystem::GatherRequestsTask(const FGatherRequestsParameters&
 					ensure(Allocated_vLevel <= MaxLevel);
 					const uint32 Allocated_vAddress = vAddress & (0xffffffff << (Allocated_vLevel * vDimensions));
 
-					AddPageUpdate(PageUpdateBuffers, PageUpdateFlushCount, PhysicalSpace->GetID(), Allocated_pAddress);
+					AddPageUpdate(PhysicalSpace->GetID(), Allocated_pAddress, Frame);
 
 					const FPhysicalSpaceIDAndAddress PhysicalSpaceIDAndAddress(PhysicalSpace->GetID(), Allocated_pAddress);
 					uint32 NumMappedPages = 0u;
@@ -2029,7 +1827,7 @@ void FVirtualTextureSystem::GatherRequestsTask(const FGatherRequestsParameters&
 							else
 							{
 								// Need to mark the page as recently used, otherwise it may be evicted later this frame
-								AddPageUpdate(PageUpdateBuffers, PageUpdateFlushCount, PhysicalSpace->GetID(), pAddress);
+								AddPageUpdate(PhysicalSpace->GetID(), pAddress, Frame);
 							}
 						}
 					}
@@ -2089,31 +1887,6 @@ void FVirtualTextureSystem::GatherRequestsTask(const FGatherRequestsParameters&
 		}
 	}
 
-	for (uint32 PhysicalSpaceID = 0u; PhysicalSpaceID < (uint32)PhysicalSpaces.Num(); ++PhysicalSpaceID)
-	{
-		if (PhysicalSpaces[PhysicalSpaceID] == nullptr)
-		{
-			continue;
-		}
-
-		FVirtualTexturePhysicalSpace* RESTRICT PhysicalSpace = GetPhysicalSpace(PhysicalSpaceID);
-		FPageUpdateBuffer& RESTRICT Buffer = PageUpdateBuffers[PhysicalSpaceID];
-
-		if (Buffer.NumPages > 0u)
-		{
-			Buffer.NumPageUpdates += Buffer.NumPages;
-			FTexturePagePool& RESTRICT PagePool = PhysicalSpace->GetPagePool();
-
-			FScopeLock Lock(&PagePool.GetLock());
-			for (uint32 i = 0u; i < Buffer.NumPages; ++i)
-			{
-				PagePool.UpdateUsage(Frame, Buffer.PhysicalAddresses[i]);
-			}
-		}
-		
-		INC_DWORD_STAT_BY(STAT_NumPageUpdate, Buffer.NumPageUpdates);
-	}
-
 	INC_DWORD_STAT_BY(STAT_NumPageVisible, NumRequestsPages);
 	INC_DWORD_STAT_BY(STAT_NumPageVisibleResident, NumResidentPages);
 	INC_DWORD_STAT_BY(STAT_NumPageVisibleNotResident, NumNonResidentPages);
@@ -2401,10 +2174,13 @@ void FVirtualTextureSystem::SubmitThrottledRequests(FRHICommandList& RHICmdList,
 	}
 
 	// Submit the merged, sorted and throttled page load requests.
-	SubmitRequests(RHICmdList, FeatureLevel, Allocator, Settings, MergedRequestList, true);
+	SubmitRequests(RHICmdList, FeatureLevel, Allocator, Settings, MergedRequestList);
 }
 
-void FVirtualTextureSystem::SubmitRequests(FRHICommandList& RHICmdList, ERHIFeatureLevel::Type FeatureLevel, FConcurrentLinearBulkObjectAllocator& Allocator, FVirtualTextureUpdateSettings const& Settings, FUniqueRequestList* RequestList, bool bAsync)
+// Global flag used to request early out of work inside SubmitRequests.
+volatile bool GAsyncTaskEarlyOut = false;
+
+void FVirtualTextureSystem::SubmitRequests(FRHICommandList& RHICmdList, ERHIFeatureLevel::Type FeatureLevel, FConcurrentLinearBulkObjectAllocator& Allocator, FVirtualTextureUpdateSettings const& Settings, FUniqueRequestList* RequestList)
 {
 	TRACE_CPUPROFILER_EVENT_SCOPE(FVirtualTextureSystem::SubmitRequests);
 	LLM_SCOPE(ELLMTag::VirtualTextureSystem);
@@ -2432,9 +2208,12 @@ void FVirtualTextureSystem::SubmitRequests(FRHICommandList& RHICmdList, ERHIFeat
 
 		bool bWaitForProducers = false;
 
-		const bool bSyncProduceLockedTiles = ShouldSyncProduceLockedTiles();
 		const uint32 MaxPagesProduced = Settings.MaxPagesProduced;
 		const uint32 PageFreeThreshold = VirtualTextureScalability::GetPageFreeThreshold();
+		const double StartTime = FPlatformTime::Seconds();
+		const double TimeSliceEarlyOut = Settings.MinSubmitAsyncEarlyOutTime > 0 ? StartTime + Settings.MinSubmitAsyncEarlyOutTime : 0;
+		const double TimeSliceEnd = Settings.MaxSubmitTime > 0 ? StartTime + Settings.MaxSubmitTime : 0;
+		
 		uint32 NumPagesPending = 0u;
 		uint32 NumPagesThrottled = 0u;
 		uint32 NumPagesInvalid = 0u;
@@ -2446,7 +2225,16 @@ void FVirtualTextureSystem::SubmitRequests(FRHICommandList& RHICmdList, ERHIFeat
 		for (uint32 RequestIndex = 0u; RequestIndex < RequestList->GetNumLoadRequests(); ++RequestIndex)
 		{
 			const bool bLockTile = RequestList->IsLocked(RequestIndex);
-			const bool bForceProduceTile = !bAsync || (bLockTile && bSyncProduceLockedTiles) || Settings.bForceSyncPageUpdate;
+			const bool bForceProduceTile = Settings.bSyncProduceAllTiles || (bLockTile && Settings.bSyncProduceLockedTiles);
+
+			// Always process locked tiles but skip work if we have run past our time slice or have been asked to early out.
+			const double CurrentTime = FPlatformTime::Seconds();
+			if (!bLockTile && ((GAsyncTaskEarlyOut && TimeSliceEarlyOut > 0 && CurrentTime > TimeSliceEarlyOut) || (TimeSliceEnd > 0 && CurrentTime > TimeSliceEnd)))
+			{
+				NumPagesThrottled++;
+				continue;
+			}
+			
 			const FVirtualTextureLocalTile TileToLoad = RequestList->GetLoadRequest(RequestIndex).GetTile();
 			const FVirtualTextureProducerHandle ProducerHandle = TileToLoad.GetProducerHandle();
 			const FVirtualTextureProducer& Producer = Producers.GetProducer(ProducerHandle);
@@ -2461,7 +2249,7 @@ void FVirtualTextureSystem::SubmitRequests(FRHICommandList& RHICmdList, ERHIFeat
 				}
 			}
 
-			const EVTRequestPagePriority Priority = bLockTile || Settings.bForceSyncPageUpdate ? EVTRequestPagePriority::High : EVTRequestPagePriority::Normal;
+			const EVTRequestPagePriority Priority = bLockTile || Settings.bSyncProduceAllTiles ? EVTRequestPagePriority::High : EVTRequestPagePriority::Normal;
 			FVTRequestPageResult RequestPageResult = Producer.GetVirtualTexture()->RequestPageData(RHICmdList, ProducerHandle, ProducerTextureLayerMask, TileToLoad.Local_vLevel, TileToLoad.Local_vAddress, Priority);
 			if (RequestPageResult.Status == EVTRequestPageStatus::Pending)
 			{
@@ -2803,74 +2591,13 @@ void FVirtualTextureSystem::GatherFeedbackRequests(FConcurrentLinearBulkObjectAl
 
 	if (Settings.bEnableFeedback)
 	{
-		// Create tasks to read the feedback data
-		// Give each task a section of the feedback buffer to analyze
-		FFeedbackAnalysisParameters FeedbackAnalysisParameters[MaxNumTasks];
-
-		const uint32 MaxNumFeedbackTasks = FMath::Clamp((uint32)Settings.NumFeedbackTasks, 1u, MaxNumTasks);
-		const uint32 FeedbackSizePerTask = FMath::DivideAndRoundUp(FeedbackResult.Size, MaxNumFeedbackTasks);
-
-		uint32 NumFeedbackTasks = 0;
-		uint32 CurrentOffset = 0;
-		while (CurrentOffset < FeedbackResult.Size)
-		{
-			const uint32 TaskIndex = NumFeedbackTasks++;
-			FFeedbackAnalysisParameters& Params = FeedbackAnalysisParameters[TaskIndex];
-			Params.System = this;
-			if (TaskIndex == 0u)
-			{
-				Params.UniquePageList = MergedUniquePageList;
-			}
-			else
-			{
-				Params.UniquePageList = Allocator.Create<FUniquePageList>();
-			}
-			Params.FeedbackBuffer = FeedbackResult.Data + CurrentOffset;
-
-			const uint32 Size = FMath::Min(FeedbackSizePerTask, FeedbackResult.Size - CurrentOffset);
-			Params.FeedbackSize = Size;
-			CurrentOffset += Size;
-		}
-
-		// Kick the tasks
-		const int32 LocalFeedbackTaskCount = Settings.bParallelFeedbackTasks ? 1 : NumFeedbackTasks;
-		const int32 WorkerFeedbackTaskCount = NumFeedbackTasks - LocalFeedbackTaskCount;
-
-		FGraphEventArray Tasks;
-		if (WorkerFeedbackTaskCount > 0)
-		{
-			SCOPE_CYCLE_COUNTER(STAT_ProcessRequests_SubmitTasks);
-			Tasks.Reserve(WorkerFeedbackTaskCount);
-			for (uint32 TaskIndex = LocalFeedbackTaskCount; TaskIndex < NumFeedbackTasks; ++TaskIndex)
-			{
-				Tasks.Add(TGraphTask<FFeedbackAnalysisTask>::CreateTask().ConstructAndDispatchWhenReady(FeedbackAnalysisParameters[TaskIndex]));
-			}
-		}
-
-		if (NumFeedbackTasks > 0u)
-		{
-			SCOPE_CYCLE_COUNTER(STAT_FeedbackAnalysis);
-
-			for (int32 TaskIndex = 0; TaskIndex < LocalFeedbackTaskCount; ++TaskIndex)
-			{
-				FFeedbackAnalysisTask::DoTask(FeedbackAnalysisParameters[TaskIndex]);
-			}
-			if (WorkerFeedbackTaskCount > 0)
-			{
-				SCOPE_CYCLE_COUNTER(STAT_ProcessRequests_WaitTasks);
-
-				FTaskGraphInterface::Get().WaitUntilTasksComplete(Tasks, ENamedThreads::GetRenderThread_Local());
-			}
-		}
+		FFeedbackAnalysisParameters Parameters;
+		Parameters.FeedbackBuffer = FeedbackResult.Data;
+		Parameters.UniquePageList = MergedUniquePageList;
+		Parameters.FeedbackSize = FeedbackResult.Size;
 
-		if (NumFeedbackTasks > 1u)
-		{
-			SCOPE_CYCLE_COUNTER(STAT_ProcessRequests_MergePages);
-			for (uint32 TaskIndex = 1u; TaskIndex < NumFeedbackTasks; ++TaskIndex)
-			{
-				MergedUniquePageList->MergePages(FeedbackAnalysisParameters[TaskIndex].UniquePageList);
-			}
-		}
+		Parameters.UniquePageList->Initialize();
+		FeedbackAnalysisTask(Parameters);
 	}
 
 #if WITH_EDITOR
@@ -2893,15 +2620,13 @@ void FVirtualTextureSystem::GatherFeedbackRequests(FConcurrentLinearBulkObjectAl
 				LevelBias += FMath::Log2(FMath::Max3((float)PageRequestCaptureResolution.X / (float)Settings.PlaybackViewSize.X, (float)PageRequestCaptureResolution.Y / (float)Settings.PlaybackViewSize.Y, 1.f));
 			}
 
-			// todo: We can split this into concurrent tasks. 
 			FAddRequestedTilesParameters Parameters;
-			Parameters.System = this;
 			Parameters.ViewLevelBias = FMath::RoundToInt(FMath::Max(LevelBias, 0.f));
 			Parameters.RequestBuffer = PageRequestPlaybackBuffer.GetData();
 			Parameters.NumRequests = PageRequestPlaybackBuffer.Num();
 			Parameters.UniquePageList = Allocator.Create<FUniquePageList>();
 
-			FAddRequestedTilesTask::DoTask(Parameters);
+			AddRequestedTilesTask(Parameters);
 			MergedUniquePageList->MergePages(Parameters.UniquePageList);
 		}
 
@@ -2909,6 +2634,13 @@ void FVirtualTextureSystem::GatherFeedbackRequests(FConcurrentLinearBulkObjectAl
 		PageRequestCaptureResolution = FIntPoint::ZeroValue;
 	}
 
+	// If we are time slicing GatherRequests then sort the pages to ensure that we only skip lower priority work.
+	if (Settings.MaxGatherTime > 0)
+	{
+		SCOPE_CYCLE_COUNTER(STAT_ProcessRequests_Sort);
+		MergedUniquePageList->SortPages();
+	}
+
 	// Pages from feedback buffer were generated several frames ago, so they may no longer be valid for newly allocated VTs
 	static uint32 PendingFrameDelay = 3u;
 	if (Frame >= PendingFrameDelay)
@@ -3082,7 +2814,6 @@ TUniquePtr<FVirtualTextureUpdater> FVirtualTextureSystem::BeginUpdate(FRDGBuilde
 						// This is not really runtime-relevant (debug cache flush) so it's fine to re-produce the tile with normal priority again : 
 						AddOrMergeTileRequest(FVirtualTextureLocalTileRequest(LockedTile, EVTProducerPriority::Normal, EVTInvalidatePriority::Normal), MappedTilesToProduce);
 					}
-					
 				}
 				// Flush unlocked pages
 				PhysicalSpace->GetPagePool().EvictAllPages(this);
@@ -3120,7 +2851,7 @@ TUniquePtr<FVirtualTextureUpdater> FVirtualTextureSystem::BeginUpdate(FRDGBuilde
 	TUniquePtr<FVirtualTextureUpdater> Updater(new FVirtualTextureUpdater());
 	Updater->Settings = Settings;
 	Updater->FeatureLevel = FeatureLevel;
-	Updater->bAsyncTaskAllowed = Settings.bEnableAsyncTasks && CVarVTAsyncPageRequestTask.GetValueOnRenderThread();
+	Updater->bAsyncTaskAllowed = Settings.bEnableAsyncTasks;
 	Updater->PageUploadBudgetRVT = Settings.MaxRVTPageUploads;
 	Updater->PageUploadBudgetSVT = Settings.MaxSVTPageUploads;
 
@@ -3139,9 +2870,19 @@ void FVirtualTextureSystem::WaitForTasks(FVirtualTextureUpdater* Updater)
 		return;
 	}
 
-	CSV_SCOPED_SET_WAIT_STAT(VirtualTexture);
-	Updater->AsyncTask.Wait();
+	// Set flag that will start dropping work from the updater task.
+	// This can helps in extreme cases where the amount of virtual texture update work is causing us to wait here.
+	// Setting FVirtualTextureUpdateSettings::MaxGatherTime and FVirtualTextureUpdateSettings::MaxSubmitTime are other ways to load balance the work.
+	GAsyncTaskEarlyOut = true;
+
+	{
+		SCOPE_CYCLE_COUNTER(STAT_ProcessRequests_WaitTasks);
+		CSV_SCOPED_SET_WAIT_STAT(VirtualTexture);
+		Updater->AsyncTask.Wait();
+	}
+
 	bUpdating = false;
+	GAsyncTaskEarlyOut = false;
 }
 
 void FVirtualTextureSystem::EndUpdate(FRDGBuilder& GraphBuilder, TUniquePtr<FVirtualTextureUpdater>&& Updater, ERHIFeatureLevel::Type FeatureLevel)
@@ -3161,6 +2902,8 @@ void FVirtualTextureSystem::EndUpdate(FRDGBuilder& GraphBuilder, TUniquePtr<FVir
 
 	if (Updater->bAsyncTaskAllowed)
 	{
+		SCOPE_CYCLE_COUNTER(STAT_ProcessRequests_WaitTasks);
+		CSV_SCOPED_SET_WAIT_STAT(VirtualTexture);
 		Updater->AsyncTask.Wait();
 	}
 	else
diff --git a/Engine/Source/Runtime/Renderer/Private/VT/VirtualTextureSystem.h b/Engine/Source/Runtime/Renderer/Private/VT/VirtualTextureSystem.h
--- a/Engine/Source/Runtime/Renderer/Private/VT/VirtualTextureSystem.h
+++ b/Engine/Source/Runtime/Renderer/Private/VT/VirtualTextureSystem.h
@@ -45,10 +45,13 @@ struct FVirtualTextureUpdateSettings
 	{
 		if (!bEnable)
 		{
+			bSyncProduceAllTiles = true;
 			MaxRVTPageUploads = 99999;
 			MaxSVTPageUploads = 99999;
 			MaxPagesProduced = 99999;
-			bForceSyncPageUpdate = true;
+			MaxGatherTime = 0.0f;
+			MaxSubmitTime = 0.0f;
+			MinSubmitAsyncEarlyOutTime = 0.0f;
 		}
 		return *this;
 	}
@@ -67,21 +70,21 @@ struct FVirtualTextureUpdateSettings
 		return *this;
 	}
 
-	bool bEnableAsyncTasks = true;
 	bool bEnablePageRequests = true;
-	bool bEnableFeedback;
-	bool bEnableFeedbackProduce;
-	bool bEnablePlayback;
-	bool bParallelFeedbackTasks;
-	bool bForceContinuousUpdate;
-	bool bForceSyncPageUpdate;
-	int32 NumFeedbackTasks;
-	int32 NumGatherTasks;
-	int32 MaxGatherPagesBeforeFlush;
-	int32 MaxRVTPageUploads;
-	int32 MaxSVTPageUploads;
-	int32 MaxPagesProduced;
-	int32 MaxContinuousUpdates;
+	bool bEnableFeedback = true;
+	bool bEnableFeedbackProduce = true;
+	bool bEnablePlayback = true;
+	bool bForceContinuousUpdate = false;
+	bool bSyncProduceAllTiles = false;
+	bool bSyncProduceLockedTiles = false;
+	bool bEnableAsyncTasks = true;
+	int32 MaxRVTPageUploads = 0;
+	int32 MaxSVTPageUploads = 0;
+	int32 MaxPagesProduced = 0;
+	int32 MaxContinuousUpdates = 0;
+	float MaxGatherTime = 0.0f;
+	float MaxSubmitTime = 0.0f;
+	float MinSubmitAsyncEarlyOutTime = 0.0f;
 	FIntPoint PlaybackViewSize = FIntPoint(0, 0);
 };
 
@@ -230,11 +233,11 @@ private:
 	};
 
 	void SubmitThrottledRequests(FRHICommandList& RHICmdList, FVirtualTextureUpdater* Updater, EUpdatePhase UpdatePhase);
-	void SubmitRequests(FRHICommandList& RHICmdList, ERHIFeatureLevel::Type FeatureLevel, FConcurrentLinearBulkObjectAllocator& Allocator, FVirtualTextureUpdateSettings const& Settings, FUniqueRequestList* RequestList, bool bAsync);
+	void SubmitRequests(FRHICommandList& RHICmdList, ERHIFeatureLevel::Type FeatureLevel, FConcurrentLinearBulkObjectAllocator& Allocator, FVirtualTextureUpdateSettings const& Settings, FUniqueRequestList* RequestList);
 
 	void GatherRequests(FUniqueRequestList* MergedRequestList, const FUniquePageList* UniquePageList, uint32 FrameRequested, FConcurrentLinearBulkObjectAllocator& Allocator, FVirtualTextureUpdateSettings const& Settings);
 
-	void AddPageUpdate(FPageUpdateBuffer* Buffers, uint32 FlushCount, uint32 PhysicalSpaceID, uint16 pAddress);
+	void AddPageUpdate(uint32 PhysicalSpaceID, uint16 pAddress, uint32 CurrentFrame);
 
 	void FeedbackAnalysisTask(const FFeedbackAnalysisParameters& Parameters);
 	void AddRequestedTilesTask(const FAddRequestedTilesParameters& Parameters);
