Path Tracing: Meshes with decal section have unwanted shadowing

The hidden sections used to indicate decals should not cast shadows

#rnx
#Synced-CL 51098442

diff --git a/Engine/Source/Runtime/Renderer/Private/RayTracing/RayTracingMaterialHitShaders.cpp b/Engine/Source/Runtime/Renderer/Private/RayTracing/RayTracingMaterialHitShaders.cpp
--- a/Engine/Source/Runtime/Renderer/Private/RayTracing/RayTracingMaterialHitShaders.cpp
+++ b/Engine/Source/Runtime/Renderer/Private/RayTracing/RayTracingMaterialHitShaders.cpp
@@ -922,7 +922,7 @@ void FDeferredShadingSceneRenderer::SetupMaterialRayTracingHitGroupBindings(FRDG
 			Binding.RecordIndex = BaseRecordIndex + RAY_TRACING_SHADER_SLOT_SHADOW;
 			Binding.Geometry = DirtyShaderBinding.RayTracingGeometry;
 			Binding.SegmentIndex = MeshCommand.GeometrySegmentIndex;
-			Binding.ShaderIndexInPipeline = OpaqueMeshDecalHitGroupIndex;
+			Binding.ShaderIndexInPipeline = HiddenMeshDecalHitGroupIndex;
 			Binding.BindingType = DirtyShaderBinding.BindingType;
 	
 		}
