Remove IsMobilePlatform and rely solely on the RHI-level bool (GSupportsParallelOcclusionQueries) to run RHIGetRenderQueryResult on async task threads during Visibility.

Saves ~0.6 ms on RenderThread_CriticalPath on some consoles.

#tests Switch2, Switch
#Synced-CL 50723010

diff --git a/Engine/Source/Runtime/Renderer/Private/SceneVisibility.cpp b/Engine/Source/Runtime/Renderer/Private/SceneVisibility.cpp
--- a/Engine/Source/Runtime/Renderer/Private/SceneVisibility.cpp
+++ b/Engine/Source/Runtime/Renderer/Private/SceneVisibility.cpp
@@ -3608,7 +3608,7 @@ FVisibilityTaskConfig::FVisibilityTaskConfig(const FScene& Scene, TConstArrayVie
 
 	if (Schedule == EVisibilityTaskSchedule::Parallel)
 	{
-		if (!FApp::ShouldUseThreadingForPerformance() || !GIsThreadedRendering || !GSupportsParallelOcclusionQueries || GVisualizeOccludedPrimitives > 0 || GVisualizePrimitiveBBoxes > 0 || IsMobilePlatform(Scene.GetShaderPlatform()))
+		if (!FApp::ShouldUseThreadingForPerformance() || !GIsThreadedRendering || !GSupportsParallelOcclusionQueries || GVisualizeOccludedPrimitives > 0 || GVisualizePrimitiveBBoxes > 0)
 		{
 			Schedule = EVisibilityTaskSchedule::RenderThread;
 		}
