MegaLights - fixed lighting channel bug, where Nanite meshes didn't respect lighting channels when casting shadows

#rnx
#Synced-CL 51176602

diff --git a/Engine/Source/Runtime/Renderer/Private/PrimitiveSceneInfo.cpp b/Engine/Source/Runtime/Renderer/Private/PrimitiveSceneInfo.cpp
--- a/Engine/Source/Runtime/Renderer/Private/PrimitiveSceneInfo.cpp
+++ b/Engine/Source/Runtime/Renderer/Private/PrimitiveSceneInfo.cpp
@@ -1403,6 +1403,7 @@ void FPrimitiveSceneInfo::SetupCachedRayTracingInstance(FPrimitiveSceneInfo* Sce
 	OutCachedRayTracingInstance.bIncrementUserDataPerInstance = true;
 
 	OutCachedRayTracingInstance.bApplyLocalBoundsTransform = RayTracingInstance.bApplyLocalBoundsTransform;
+	OutCachedRayTracingInstance.LightingChannelMask = SceneProxy->GetLightingChannelMask();
 
 	OutCachedRayTracingInstance.Flags = ERayTracingInstanceFlags::None;
 
