EyeAdaptation - disable cached lighting pre-exposure for empty scenes and editor viewports

#rnx
#Synced-CL 51176602

diff --git a/Engine/Source/Runtime/Renderer/Private/PostProcess/PostProcessEyeAdaptation.cpp b/Engine/Source/Runtime/Renderer/Private/PostProcess/PostProcessEyeAdaptation.cpp
--- a/Engine/Source/Runtime/Renderer/Private/PostProcess/PostProcessEyeAdaptation.cpp
+++ b/Engine/Source/Runtime/Renderer/Private/PostProcess/PostProcessEyeAdaptation.cpp
@@ -177,7 +177,7 @@ namespace EyeAdaptation
 		return 1.0f / FMath::Exp2(GetCachedLightingPreExposureEV());
 	}
 
-	bool WriteWarnings(const FViewInfo& View, FScreenMessageWriter* Writer)
+	bool WriteWarnings(const FScene* Scene, const FViewInfo& View, FScreenMessageWriter* Writer)
 	{
 		// Float16 has [-14;+16] range in EV units
 		// Additionally need at least 2 stops below and 8 stops above that in order to not clip any lighting response
@@ -187,8 +187,13 @@ namespace EyeAdaptation
 
 		const float ExposureEV = FMath::Log2(1.0f / View.PreExposure);
 
-		bool bShowCachedLightingWarning = false;		
-		if (View.FinalPostProcessSettings.DynamicGlobalIlluminationMethod == EDynamicGlobalIlluminationMethod::Lumen)
+		bool bShowCachedLightingWarning = false;
+
+		// This warning only makes scene if scene is lit, has Lumen and it's the main view
+		if (Scene
+			&& Scene->GetWorld()->WorldType != EWorldType::EditorPreview
+			&& Scene->Lights.Num() > 0
+			&& View.FinalPostProcessSettings.DynamicGlobalIlluminationMethod == EDynamicGlobalIlluminationMethod::Lumen)
 		{
 			bShowCachedLightingWarning = ExposureEV < CachedLightingPreExposureMin || ExposureEV > CachedLightingPreExposureMax;
 		}
diff --git a/Engine/Source/Runtime/Renderer/Private/PostProcess/PostProcessEyeAdaptation.h b/Engine/Source/Runtime/Renderer/Private/PostProcess/PostProcessEyeAdaptation.h
--- a/Engine/Source/Runtime/Renderer/Private/PostProcess/PostProcessEyeAdaptation.h
+++ b/Engine/Source/Runtime/Renderer/Private/PostProcess/PostProcessEyeAdaptation.h
@@ -5,6 +5,7 @@
 #include "ScreenPass.h"
 
 class FLocalExposureParameters;
+class FScene;
 
 namespace AutoExposurePermutation
 {
@@ -23,7 +24,7 @@ namespace EyeAdaptation
 {
 	float GetCachedLightingPreExposure();
 
-	bool WriteWarnings(const FViewInfo& View, FScreenMessageWriter* Writer);
+	bool WriteWarnings(const FScene* Scene, const FViewInfo& View, FScreenMessageWriter* Writer);
 }
 
 // For converting the auto exposure to new values from 4.24 to 4.25
diff --git a/Engine/Source/Runtime/Renderer/Private/SceneRendering.cpp b/Engine/Source/Runtime/Renderer/Private/SceneRendering.cpp
--- a/Engine/Source/Runtime/Renderer/Private/SceneRendering.cpp
+++ b/Engine/Source/Runtime/Renderer/Private/SceneRendering.cpp
@@ -4439,7 +4439,7 @@ void FSceneRenderer::OnRenderFinish(FRDGBuilder& GraphBuilder, FRDGTextureRef Vi
 				const bool bStereoView = IStereoRendering::IsStereoEyeView(View);
 				const bool bGPUSkinCacheVisualzationMode = SkinCache && ViewFamily.EngineShowFlags.VisualizeGPUSkinCache && View.CurrentGPUSkinCacheVisualizationMode != NAME_None;
 				const bool bMegaLightsHasWarning = MegaLights::HasWarning(View);
-				const bool bEyeAdaptationHasWarning = EyeAdaptation::WriteWarnings(View, /*Writer*/ nullptr);
+				const bool bEyeAdaptationHasWarning = EyeAdaptation::WriteWarnings(Scene, View, /*Writer*/ nullptr);
 
 				// display a warning if an ambient cubemap uses non-angular mipmap filtering
 				bool bShowAmbientCubemapMipGenSettingsWarning = false;
@@ -4613,7 +4613,7 @@ void FSceneRenderer::OnRenderFinish(FRDGBuilder& GraphBuilder, FRDGTextureRef Vi
 
 						if (bEyeAdaptationHasWarning)
 						{
-							EyeAdaptation::WriteWarnings(View, &Writer);
+							EyeAdaptation::WriteWarnings(Scene, View, &Writer);
 						}
 
 						if (bNaniteEnabledButNoAtomics)
