[Front Layer Translucency] Added missing SetRenderTargetOutputFormat(...) calls to FFrontLayerTranslucencyClearGBufferPS and FFrontLayerTranslucencyGBufferPS.
- caused validation issues on some platforms.

#Synced-CL 51150711

diff --git a/Engine/Source/Runtime/Renderer/Private/FrontLayerTranslucency.cpp b/Engine/Source/Runtime/Renderer/Private/FrontLayerTranslucency.cpp
--- a/Engine/Source/Runtime/Renderer/Private/FrontLayerTranslucency.cpp
+++ b/Engine/Source/Runtime/Renderer/Private/FrontLayerTranslucency.cpp
@@ -114,6 +114,19 @@ class FFrontLayerTranslucencyClearGBufferPS : public FGlobalShader
 	{
 		return DoesPlatformSupportLumenGI(Parameters.Platform);
 	}
+
+	static void ModifyCompilationEnvironment(const FGlobalShaderPermutationParameters& Parameters, FShaderCompilerEnvironment& OutEnvironment)
+	{
+		FGlobalShader::ModifyCompilationEnvironment(Parameters, OutEnvironment);
+
+		const bool bMegaLightsSupported = ShouldCompileMegaLightsShaders(Parameters.Platform);
+
+		OutEnvironment.SetRenderTargetOutputFormat(0, GetFrontLayerTranslucencyGBufferAFormat());
+		if (bMegaLightsSupported)
+		{
+			OutEnvironment.SetRenderTargetOutputFormat(1, GetFrontLayerTranslucencyGBufferBFormat());
+		}
+	}
 };
 
 IMPLEMENT_GLOBAL_SHADER(FFrontLayerTranslucencyClearGBufferPS, "/Engine/Private/FrontLayerTranslucency.usf", "FrontLayerTranslucencyClearGBufferPS", SF_Pixel);
@@ -157,6 +170,20 @@ public:
 		return DoesPlatformSupportLumenGI(Parameters.Platform) && IsTranslucentBlendMode(Parameters.MaterialParameters) && Parameters.MaterialParameters.bIsTranslucencySurface;
 	}
 
+	static void ModifyCompilationEnvironment(const FMeshMaterialShaderPermutationParameters& Parameters, FShaderCompilerEnvironment& OutEnvironment)
+	{
+		FMeshMaterialShader::ModifyCompilationEnvironment(Parameters, OutEnvironment);
+
+		const bool bMegaLightsSupported = ShouldCompileMegaLightsShaders(Parameters.Platform);
+
+		OutEnvironment.SetRenderTargetOutputFormat(0, GetFrontLayerTranslucencyGBufferAFormat());
+		if (bMegaLightsSupported)
+		{
+			OutEnvironment.SetRenderTargetOutputFormat(1, GetFrontLayerTranslucencyGBufferBFormat());
+			OutEnvironment.SetRenderTargetOutputFormat(2, GetFrontLayerTranslucencyGBufferCFormat());
+		}
+	}
+
 	FFrontLayerTranslucencyGBufferPS(const ShaderMetaType::CompiledShaderInitializerType& Initializer)
 		: FMeshMaterialShader(Initializer)
 	{}
