Eye Adaptation Cached Lighting PreExposure - disable warning when pre-exposure isn't used and print current exposure value

#rnx
#Synced-CL 51176602

diff --git a/Engine/Source/Runtime/Renderer/Private/PostProcess/PostProcessEyeAdaptation.cpp b/Engine/Source/Runtime/Renderer/Private/PostProcess/PostProcessEyeAdaptation.cpp
--- a/Engine/Source/Runtime/Renderer/Private/PostProcess/PostProcessEyeAdaptation.cpp
+++ b/Engine/Source/Runtime/Renderer/Private/PostProcess/PostProcessEyeAdaptation.cpp
@@ -167,6 +167,29 @@ static TAutoConsoleVariable<float> CVarEyeAdaptationCachedLightingPreExposure(
 
 namespace EyeAdaptation
 {
+	// Whether pre exposure should be used for this View
+	bool IsPreExposureRelevant(const FViewInfo& View)
+	{
+		const FSceneViewFamily& ViewFamily = *View.Family;
+
+		const bool bIsPreExposureRelevant = true
+			&& ViewFamily.EngineShowFlags.EyeAdaptation // Controls whether scene luminance is computed at all.
+			&& ViewFamily.EngineShowFlags.Lighting
+			&& IsPostProcessingEnabled(View)
+			&& ViewFamily.bResolveScene
+			&& !ViewFamily.EngineShowFlags.LightMapDensity
+			&& !ViewFamily.EngineShowFlags.StationaryLightOverlap
+			&& !ViewFamily.EngineShowFlags.LightComplexity
+			&& !ViewFamily.EngineShowFlags.LODColoration
+			&& !ViewFamily.EngineShowFlags.HLODColoration
+			&& !ViewFamily.EngineShowFlags.ActorColoration
+			&& ((!ViewFamily.EngineShowFlags.VisualizeBuffer) || View.CurrentBufferVisualizationMode != NAME_None)   // disable pre-exposure for the buffer visualization modes
+			&& (!ViewFamily.EngineShowFlags.PathTracing || PathTracing::NeedsTonemapping())                     // Either we aren't using path tracing at all, or we are using a path tracing mode that needs tone-mapping
+			&& (!ViewFamily.EngineShowFlags.RayTracingDebug || RaytracingDebugViewModeNeedsTonemapping(View)); // Either we aren't using RT debug at all, or we are using a mode that needs tone-mapping
+
+		return bIsPreExposureRelevant;
+	}
+
 	float GetCachedLightingPreExposureEV()
 	{
 		return FMath::Clamp(CVarEyeAdaptationCachedLightingPreExposure.GetValueOnRenderThread(), -16.0f, 16.0f);
@@ -193,6 +216,11 @@ namespace EyeAdaptation
 		if (Scene
 			&& Scene->GetWorld()->WorldType != EWorldType::EditorPreview
 			&& Scene->Lights.Num() > 0
+			&& View.ViewState
+			&& IsPreExposureRelevant(View)
+			&& !View.bIsPlanarReflection
+			&& !View.bIsSceneCapture
+			&& !View.bIsSceneCaptureCube
 			&& View.FinalPostProcessSettings.DynamicGlobalIlluminationMethod == EDynamicGlobalIlluminationMethod::Lumen)
 		{
 			bShowCachedLightingWarning = ExposureEV < CachedLightingPreExposureMin || ExposureEV > CachedLightingPreExposureMax;
@@ -201,7 +229,8 @@ namespace EyeAdaptation
 		if (bShowCachedLightingWarning && Writer)
 		{
 			const FString String = FString::Printf(
-				TEXT("Cached lighting in Lumen and real-time sky capture lighting is going to be clipped. Please adjust r.EyeAdaptation.CachedLightingPreExposure to match expected exposure range in project. Current safe exposure range: [%.1f, %.1f]."),
+				TEXT("Cached lighting in Lumen and real-time sky capture lighting is going to be clipped. Please adjust r.EyeAdaptation.CachedLightingPreExposure to match expected exposure range in project. Exposure: %.1f. Safe exposure range: [%.1f, %.1f]."),
+					ExposureEV,
 					CachedLightingPreExposureMin,
 					CachedLightingPreExposureMax);
 
@@ -1415,7 +1444,6 @@ void FViewInfo::UpdatePreExposure()
 {
 	const FSceneViewFamily& ViewFamily = *Family;
 
-	const bool bIsPostProcessingEnabled = IsPostProcessingEnabled(*this);
 	const bool bMobilePlatform = IsMobilePlatform(GetShaderPlatform());
 	const bool bEnableAutoExposure = !bMobilePlatform || IsMobileEyeAdaptationEnabled(*this);
 	const float PreExposureOverride = CVarEyeAdaptationPreExposureOverride.GetValueOnRenderThread();
@@ -1432,20 +1460,7 @@ void FViewInfo::UpdatePreExposure()
 	// but this is too limiting for certain view. For instance shader preview doesn't have 
 	// volumetric lighting enabled, which makes the view be flagged as rich, and not updating 
 	// the pre-exposition value.
-	const bool bIsPreExposureRelevant = true
-		&& ViewFamily.EngineShowFlags.EyeAdaptation // Controls whether scene luminance is computed at all.
-		&& ViewFamily.EngineShowFlags.Lighting
-		&& bIsPostProcessingEnabled
-		&& ViewFamily.bResolveScene
-		&& !ViewFamily.EngineShowFlags.LightMapDensity
-		&& !ViewFamily.EngineShowFlags.StationaryLightOverlap
-		&& !ViewFamily.EngineShowFlags.LightComplexity
-		&& !ViewFamily.EngineShowFlags.LODColoration
-		&& !ViewFamily.EngineShowFlags.HLODColoration
-		&& !ViewFamily.EngineShowFlags.ActorColoration
-		&& ((!ViewFamily.EngineShowFlags.VisualizeBuffer) || CurrentBufferVisualizationMode != NAME_None)   // disable pre-exposure for the buffer visualization modes
-		&& (!ViewFamily.EngineShowFlags.PathTracing || PathTracing::NeedsTonemapping())                     // Either we aren't using path tracing at all, or we are using a path tracing mode that needs tone-mapping
-		&& (!ViewFamily.EngineShowFlags.RayTracingDebug || RaytracingDebugViewModeNeedsTonemapping(*this)); // Either we aren't using RT debug at all, or we are using a mode that needs tone-mapping
+	const bool bIsPreExposureRelevant = EyeAdaptation::IsPreExposureRelevant(*this);
 
 	// Compute the PreExposure to use.
 	bool bUpdateLastExposure = false;
