[Front Layer Translucency] Skip dedicated front layer translucency GBuffer pass for VSM marking based on View.bTranslucentSurfaceLighting.
- avoids unnecessary work and errors due to resources not being written to.
- also added an ensure to identify the pass being unexpectedly enabled when none of the views are rendering into it.

#Synced-CL 51150711

diff --git a/Engine/Source/Runtime/Renderer/Private/FrontLayerTranslucency.cpp b/Engine/Source/Runtime/Renderer/Private/FrontLayerTranslucency.cpp
--- a/Engine/Source/Runtime/Renderer/Private/FrontLayerTranslucency.cpp
+++ b/Engine/Source/Runtime/Renderer/Private/FrontLayerTranslucency.cpp
@@ -523,13 +523,15 @@ FFrontLayerTranslucencyData FDeferredShadingSceneRenderer::RenderFrontLayerTrans
 			// * Render local front layer data here
 			bool bValidHistory = false;
 			bool bWillLumenOrMegaLightsRenderFrontLayer = false;
+			bool bAnyViewHasTranslucentSurfaceLighting = false;
 			for (FViewInfo& View : InViews)
 			{
 				bValidHistory = bValidHistory || IsFrontLayerHistoryValid(View);
 				bWillLumenOrMegaLightsRenderFrontLayer = bWillLumenOrMegaLightsRenderFrontLayer || IsLumenFrontLayerTranslucencyEnabled(View) || IsMegaLightsFrontLayerTranslucencyEnabled(View);
+				bAnyViewHasTranslucentSurfaceLighting = bAnyViewHasTranslucentSurfaceLighting || View.bTranslucentSurfaceLighting;
 			}
 
-			if (!bValidHistory && !bWillLumenOrMegaLightsRenderFrontLayer)
+			if (bAnyViewHasTranslucentSurfaceLighting && !bValidHistory && !bWillLumenOrMegaLightsRenderFrontLayer)
 			{
 				bNeedFrontLayerData = 2;
 			}
@@ -553,6 +555,20 @@ FFrontLayerTranslucencyData FDeferredShadingSceneRenderer::RenderFrontLayerTrans
 	{
 		RDG_EVENT_SCOPE_STAT(GraphBuilder, FrontLayerTranslucencyGBuffer, "FrontLayerTranslucencyGBuffer");
 
+#if DO_CHECK
+		bool bAnyViewRendersPass = false;
+		for (const FViewInfo& View : InViews)
+		{
+			if (View.ParallelMeshDrawCommandPasses[EMeshPass::FrontLayerTranslucencyGBuffer] != nullptr)
+			{
+				bAnyViewRendersPass = true;
+				break;
+			}
+		}
+
+		ensureMsgf(bAnyViewRendersPass, TEXT("FrontLayerTranslucencyGBuffer resources should only be created if any view will render in this pass."));
+#endif // DO_CHECK
+
 		const bool bMegaLightsSupported = ShouldCompileMegaLightsShaders(ShaderPlatform);
 
 		// Allocate resources for all views
