Fix HairStrandsVF invalid check when the HairStrands plugin is disabled.

#Synced-CL 51261860

diff --git a/Engine/Source/Runtime/Renderer/Private/HairStrands/HairStrandsInterface.cpp b/Engine/Source/Runtime/Renderer/Private/HairStrands/HairStrandsInterface.cpp
--- a/Engine/Source/Runtime/Renderer/Private/HairStrands/HairStrandsInterface.cpp
+++ b/Engine/Source/Runtime/Renderer/Private/HairStrands/HairStrandsInterface.cpp
@@ -619,8 +619,15 @@ bool IsHairVisible(const FMeshBatchAndRelevance& MeshBatch, bool bCheckLengthSca
 bool IsHairStrandsVF(const FVertexFactoryType* InVFType)
 {
 	// NOTE: The VF is defined in a different module, so we must search for it by name
-	static const FVertexFactoryType* VFTypeA = FVertexFactoryType::GetVFByName(TEXT("FHairStrandsVertexFactory"));                   check(VFTypeA != nullptr);
-	static const FVertexFactoryType* VFTypeB = FVertexFactoryType::GetVFByName(TEXT("FLinearSweptSpheresHairStrandsVertexFactory")); check(VFTypeB != nullptr);
+	static const FVertexFactoryType* VFTypeA = FVertexFactoryType::GetVFByName(TEXT("FHairStrandsVertexFactory"));
+	static const FVertexFactoryType* VFTypeB = FVertexFactoryType::GetVFByName(TEXT("FLinearSweptSpheresHairStrandsVertexFactory"));
+	#if !UE_BUILD_SHIPPING
+	if (IsGroomEnabled())
+	{
+		check(VFTypeA != nullptr);
+		check(VFTypeB != nullptr);
+	}
+	#endif
 	return InVFType == VFTypeA || InVFType == VFTypeB;
 }
 
