Fix ensure in FMaterialDebugSubstrateTreePS due to missing parameters in the struct.

- Added PreIntegratedGF and PreIntegratedGFSampler and init them to default textures and sampler.
- This is similar to other shaders like LumenVisualize.

#Synced-CL 51266080

diff --git a/Engine/Source/Runtime/Renderer/Private/Substrate/SubstrateVisualize.cpp b/Engine/Source/Runtime/Renderer/Private/Substrate/SubstrateVisualize.cpp
--- a/Engine/Source/Runtime/Renderer/Private/Substrate/SubstrateVisualize.cpp
+++ b/Engine/Source/Runtime/Renderer/Private/Substrate/SubstrateVisualize.cpp
@@ -248,6 +248,8 @@ class FMaterialDebugSubstrateTreePS : public FGlobalShader
 		SHADER_PARAMETER_STRUCT_REF(FReflectionCaptureShaderData, ReflectionCapture)
 		SHADER_PARAMETER_RDG_UNIFORM_BUFFER(FForwardLightUniformParameters, ForwardLightStruct)
 		SHADER_PARAMETER_STRUCT_INCLUDE(FSkyDiffuseLightingParameters, SkyDiffuseLighting)
+		SHADER_PARAMETER_TEXTURE(Texture2D, PreIntegratedGF)
+		SHADER_PARAMETER_SAMPLER(SamplerState, PreIntegratedGFSampler)
 		RENDER_TARGET_BINDING_SLOTS()
 	END_SHADER_PARAMETER_STRUCT()
 
@@ -1446,6 +1448,9 @@ static void AddVisualizeAdvancedMaterialPasses(FRDGBuilder& GraphBuilder, const
 		FSkyLightSceneProxy* NullSkyLight = nullptr;
 		PassParameters->SkyDiffuseLighting = GetSkyDiffuseLightingParameters(NullSkyLight, DynamicBentNormalAO);
 
+		PassParameters->PreIntegratedGF = GSystemTextures.PreintegratedGF->GetRHI();
+		PassParameters->PreIntegratedGFSampler = TStaticSamplerState<SF_Bilinear, AM_Clamp, AM_Clamp, AM_Clamp>::GetRHI();
+
 		FMaterialDebugSubstrateTreePS::FPermutationDomain PermutationVector;
 		TShaderMapRef<FMaterialDebugSubstrateTreePS> PixelShader(View.ShaderMap, PermutationVector);
 
