Make the mobile renderer respect the Lighting ShowFlag, and make Shipping builds respect the PointLights, SpotLights, and RectLights ShowFlags as these are exposed for SceneCapture

#tests Android client, confirmed lights are no longer drawing with Insights capture
#rnx

diff --git a/Engine/Source/Runtime/Renderer/Private/LightGridInjection.cpp b/Engine/Source/Runtime/Renderer/Private/LightGridInjection.cpp
--- a/Engine/Source/Runtime/Renderer/Private/LightGridInjection.cpp
+++ b/Engine/Source/Runtime/Renderer/Private/LightGridInjection.cpp
@@ -1651,7 +1651,7 @@ FComputeLightGridOutput FSceneRenderer::ComputeLightGrid(FRDGBuilder& GraphBuild
 		ForwardLightUniformParameters->MaxCulledLightsPerCell = GLightLinkedListCulling ? NumVisibleLocalLights : GMaxCulledLightsPerCell;
 		ForwardLightUniformParameters->LightGridPixelSizeShift = FMath::FloorLog2(GLightGridPixelSize);
 		ForwardLightUniformParameters->DirectionalMegaLightsSupportedStartIndex = DirectionalMegaLightsSupportedStart;
-		ForwardLightUniformParameters->DirectLightingShowFlag = ViewFamily.EngineShowFlags.DirectLighting ? 1 : 0;
+		ForwardLightUniformParameters->DirectLightingShowFlag = (ViewFamily.EngineShowFlags.Lighting && ViewFamily.EngineShowFlags.DirectLighting) ? 1 : 0;
 
 		// Clamp far plane to something reasonable
 		const float KilometersToCentimeters = 100000.0f;
diff --git a/Engine/Source/Runtime/Renderer/Private/LightSceneInfo.cpp b/Engine/Source/Runtime/Renderer/Private/LightSceneInfo.cpp
--- a/Engine/Source/Runtime/Renderer/Private/LightSceneInfo.cpp
+++ b/Engine/Source/Runtime/Renderer/Private/LightSceneInfo.cpp
@@ -206,35 +206,36 @@ bool FLightSceneInfo::ShouldRenderLight(const FViewInfo& View, bool bOffscreen)
 	// Only render the light if it is in the view frustum
 	bool bLocalVisible = bVisible && (bOffscreen ? View.VisibleLightInfos[Id].bInDrawRange : View.VisibleLightInfos[Id].bInViewFrustum);
 
-#if !UE_BUILD_SHIPPING
-	switch(Type)
+	if (bLocalVisible)
 	{
-		case LightType_Directional:
-			if(!View.Family->EngineShowFlags.DirectionalLights) 
-			{
-				bLocalVisible = false;
-			}
-			break;
-		case LightType_Point:
-			if(!View.Family->EngineShowFlags.PointLights) 
-			{
-				bLocalVisible = false;
-			}
-			break;
-		case LightType_Spot:
-			if(!View.Family->EngineShowFlags.SpotLights)
-			{
-				bLocalVisible = false;
-			}
-			break;
-		case LightType_Rect:
-			if(!View.Family->EngineShowFlags.RectLights)
-			{
-				bLocalVisible = false;
-			}
-			break;
+		switch(Type)
+		{
+			case LightType_Directional:
+				if (!View.Family->EngineShowFlags.DirectionalLights) 
+				{
+					bLocalVisible = false;
+				}
+				break;
+			case LightType_Point:
+				if (!View.Family->EngineShowFlags.PointLights) 
+				{
+					bLocalVisible = false;
+				}
+				break;
+			case LightType_Spot:
+				if (!View.Family->EngineShowFlags.SpotLights)
+				{
+					bLocalVisible = false;
+				}
+				break;
+			case LightType_Rect:
+				if (!View.Family->EngineShowFlags.RectLights)
+				{
+					bLocalVisible = false;
+				}
+				break;
+		}
 	}
-#endif
 
 	return bLocalVisible
 		// Only render lights with static shadowing for reflection captures, since they are only captured at edit time
diff --git a/Engine/Source/Runtime/Renderer/Private/MobileShadingRenderer.cpp b/Engine/Source/Runtime/Renderer/Private/MobileShadingRenderer.cpp
--- a/Engine/Source/Runtime/Renderer/Private/MobileShadingRenderer.cpp
+++ b/Engine/Source/Runtime/Renderer/Private/MobileShadingRenderer.cpp
@@ -2514,6 +2514,8 @@ void FMobileSceneRenderer::RenderDeferredSinglePass(FRDGBuilder& GraphBuilder, F
 		bDoOcclusionQueries = QueriesPerView.Num() > 0;
 	}
 
+	bool bRenderDeferredLighting = ViewFamily.EngineShowFlags.Lighting && ViewFamily.EngineShowFlags.DeferredLighting;
+
 	for (FRenderViewContext& ViewContext : RenderViews)
 	{
 		FViewInfo& View = *ViewContext.ViewInfo;
@@ -2596,6 +2598,7 @@ void FMobileSceneRenderer::RenderDeferredSinglePass(FRDGBuilder& GraphBuilder, F
 			&SortedLightSet, 
 			bUsingPixelLocalStorage,
 			bDoOcclusionQueries, 
+			bRenderDeferredLighting,
 			DecalViewInfo,
 			ViewQueries = (bDoOcclusionQueries) ? MoveTemp(QueriesPerView[ViewContext.ViewIndex]) : FViewOcclusionQueries()](FRHICommandList& RHICmdList)
 			{
@@ -2614,7 +2617,10 @@ void FMobileSceneRenderer::RenderDeferredSinglePass(FRDGBuilder& GraphBuilder, F
 				RenderDecals(RHICmdList, View, DecalViewInfo, &ParameterCollection->MeshDecalSceneColorAndGBufferInstanceCullingDrawParams);
 				// SceneColor write, SceneDepth is read only
 				RHICmdList.NextSubpass();
-				MobileDeferredShadingPass(RHICmdList, ViewContext.ViewIndex, Views.Num(), View, *Scene, SortedLightSet, VisibleLightInfos, MobileSSRQuality);
+				if (bRenderDeferredLighting)
+				{
+					MobileDeferredShadingPass(RHICmdList, ViewContext.ViewIndex, Views.Num(), View, *Scene, SortedLightSet, VisibleLightInfos, MobileSSRQuality);
+				}
 
 				if (bUsingPixelLocalStorage)
 				{
@@ -2920,6 +2926,7 @@ void FMobileSceneRenderer::RenderDeferredMultiPass(FRDGBuilder& GraphBuilder, FS
 		}
 
 		const EMobileSSRQuality MobileSSRQuality = ActiveMobileSSRQuality(View, bShouldRenderVelocities);
+		const bool bRenderDeferredLighting = ViewFamily.EngineShowFlags.Lighting && ViewFamily.EngineShowFlags.DeferredLighting;
 
 		// Determine if a shadow point light is present.
 		TArray<int32, TInlineAllocator<8>> ShadowPointLightIndexes;
@@ -2929,37 +2936,40 @@ void FMobileSceneRenderer::RenderDeferredMultiPass(FRDGBuilder& GraphBuilder, FS
 		{
 			RDG_EVENT_SCOPE(GraphBuilder, "LightingAndFog", ViewIndex);
 
-			GraphBuilder.AddPass(
-				RDG_EVENT_NAME("Lighting"),
-				PassParameters,
-				ERDGPassFlags::Raster,
-				[this, PassParameters, ViewContext, MobileSSRQuality, &SortedLightSet](FRHICommandList& RHICmdList)
-				{
-					FViewInfo& View = *ViewContext.ViewInfo;
-					MobileDeferredShadingPass(RHICmdList, ViewContext.ViewIndex, Views.Num(), View, *Scene, SortedLightSet, VisibleLightInfos, MobileSSRQuality, PassParameters->BentNormalAOTexture);
-				});
+			if (bRenderDeferredLighting)
+			{
+				GraphBuilder.AddPass(
+					RDG_EVENT_NAME("Lighting"),
+					PassParameters,
+					ERDGPassFlags::Raster,
+					[this, PassParameters, ViewContext, MobileSSRQuality, &SortedLightSet](FRHICommandList& RHICmdList)
+					{
+						FViewInfo& View = *ViewContext.ViewInfo;
+						MobileDeferredShadingPass(RHICmdList, ViewContext.ViewIndex, Views.Num(), View, *Scene, SortedLightSet, VisibleLightInfos, MobileSSRQuality, PassParameters->BentNormalAOTexture);
+					});
 
 
-			if (bSeparateShadowPointLights)
-			{
-				for (int32 Index = 0; Index < ShadowPointLightIndexes.Num(); ++Index)
+				if (bSeparateShadowPointLights)
 				{
-					const FSortedLightSceneInfo& SortedLight = SortedLightSet.SortedLights[ShadowPointLightIndexes[Index]];
-					const FLightSceneInfo& LightSceneInfo = *SortedLight.LightSceneInfo;
-
-					RDG_EVENT_SCOPE(GraphBuilder, "PointLightShadow_%s", LightSceneInfo.Proxy->GetOwnerNameOrLabel());
+					for (int32 Index = 0; Index < ShadowPointLightIndexes.Num(); ++Index)
 					{
-						ClearMobileScreenSpaceShadows(GraphBuilder, View, LightSceneInfo.Proxy, BasePassTextures.ScreenSpaceShadowMaskTexture);
-						RenderShadowProjections(GraphBuilder, SceneTextures, BasePassTextures.ScreenSpaceShadowMaskTexture, nullptr, &LightSceneInfo, false);
-
-						GraphBuilder.AddPass(
-							RDG_EVENT_NAME("PointLightShadowPass"),
-							PassParameters,
-							ERDGPassFlags::Raster,
-							[this, &View, LightSceneInfo](FRHICommandList& RHICmdList)
-							{
-								MobileDeferredRenderLocalLight(RHICmdList, *Scene, View, LightSceneInfo, VisibleLightInfos);
-							});
+						const FSortedLightSceneInfo& SortedLight = SortedLightSet.SortedLights[ShadowPointLightIndexes[Index]];
+						const FLightSceneInfo& LightSceneInfo = *SortedLight.LightSceneInfo;
+
+						RDG_EVENT_SCOPE(GraphBuilder, "PointLightShadow_%s", LightSceneInfo.Proxy->GetOwnerNameOrLabel());
+						{
+							ClearMobileScreenSpaceShadows(GraphBuilder, View, LightSceneInfo.Proxy, BasePassTextures.ScreenSpaceShadowMaskTexture);
+							RenderShadowProjections(GraphBuilder, SceneTextures, BasePassTextures.ScreenSpaceShadowMaskTexture, nullptr, &LightSceneInfo, false);
+
+							GraphBuilder.AddPass(
+								RDG_EVENT_NAME("PointLightShadowPass"),
+								PassParameters,
+								ERDGPassFlags::Raster,
+								[this, &View, LightSceneInfo](FRHICommandList& RHICmdList)
+								{
+									MobileDeferredRenderLocalLight(RHICmdList, *Scene, View, LightSceneInfo, VisibleLightInfos);
+								});
+						}
 					}
 				}
 			}
@@ -3196,12 +3206,16 @@ void FMobileSceneRenderer::UpdateDirectionalLightUniformBuffers(FRDGBuilder& Gra
 
 	AddPass(GraphBuilder, RDG_EVENT_NAME("UpdateDirectionalLightUniformBuffers"), [this, &View](FRHICommandListImmediate& RHICmdList)
 	{
+		const bool bRenderDirectionalLights = ViewFamily.EngineShowFlags.Lighting && ViewFamily.EngineShowFlags.DirectLighting && ViewFamily.EngineShowFlags.DirectionalLights;
 		const bool bDynamicShadows = ViewFamily.EngineShowFlags.DynamicShadows;
 		// Fill in the other entries based on the lights
 		for (int32 ChannelIdx = 0; ChannelIdx < UE_ARRAY_COUNT(Scene->MobileDirectionalLights); ChannelIdx++)
 		{
 			FMobileDirectionalLightShaderParameters Params;
-			SetupMobileDirectionalLightUniformParameters(*Scene, View, VisibleLightInfos, ChannelIdx, bDynamicShadows, Params);
+			if (bRenderDirectionalLights)
+			{
+				SetupMobileDirectionalLightUniformParameters(*Scene, View, VisibleLightInfos, ChannelIdx, bDynamicShadows, Params);
+			}
 			Scene->UniformBuffers.MobileDirectionalLightUniformBuffers[ChannelIdx + 1].UpdateUniformBufferImmediate(RHICmdList, Params);
 		}
 	});
