Fix viewport for debug primitive drawing.
Primitives in SDPG_World are drawn directly to SceneColor, which can have a different ViewRect than the output texture if OutputOverride is set.
Similarly, the UpscaledSceneDepth viewport should be drawn to align with SceneColor, rather than DebugPrimitiveColor, as it is only used for depth testing in the SDPG_World pass.

#Synced-CL 51099840

diff --git a/Engine/Shaders/Private/PostProcessCompositePrimitives.usf b/Engine/Shaders/Private/PostProcessCompositePrimitives.usf
--- a/Engine/Shaders/Private/PostProcessCompositePrimitives.usf
+++ b/Engine/Shaders/Private/PostProcessCompositePrimitives.usf
@@ -259,6 +259,10 @@ void MainTemporalUpsampleDepthPS(
 	OutDeviceZ = -min(-FinalHistoryDeviceZ, 0.0);
 }
 
+#ifdef MainPopulateSceneDepthPS
+
+float4 OutputViewRectMinAndInvSize;
+
 void MainPopulateSceneDepthPS(
 	#if USE_MSAA && !COMPILER_HLSLCC
 		sample noperspective float4 SvPosition : SV_POSITION,
@@ -268,7 +272,7 @@ void MainPopulateSceneDepthPS(
 	out float4 OutColor : SV_Target0,
 	out float OutDepth : SV_DEPTH)
 {
-	float2 ViewportUV = (SvPosition.xy - View.ViewRectMin.xy) * View.ViewSizeAndInvSize.zw;
+	float2 ViewportUV = (SvPosition.xy - OutputViewRectMinAndInvSize.xy) * OutputViewRectMinAndInvSize.zw;
 	
 	float2 DepthUV = clamp(
 		(Depth_ViewportMin + ViewportUV * Depth_ViewportSize - DepthTextureJitter) * Depth_ExtentInverse,
@@ -296,6 +300,8 @@ void MainPopulateSceneDepthPS(
 	OutDepth = Texture2DSampleLevel(DepthTexture, DepthSampler, DepthUV, 0).r;
 }
 
+#endif
+
 float ComputeDepthMask(float SceneDeviceZ, float EditorDeviceZ)
 {
 	// Soft Bias with SceneDeviceZ for best quality
diff --git a/Engine/Source/Runtime/Renderer/Private/PostProcess/PostProcessCompositeDebugPrimitives.cpp b/Engine/Source/Runtime/Renderer/Private/PostProcess/PostProcessCompositeDebugPrimitives.cpp
--- a/Engine/Source/Runtime/Renderer/Private/PostProcess/PostProcessCompositeDebugPrimitives.cpp
+++ b/Engine/Source/Runtime/Renderer/Private/PostProcess/PostProcessCompositeDebugPrimitives.cpp
@@ -83,6 +83,7 @@ static FDebugPrimitivesDrawingContext InitializeDebugPrimitiveTextures(
 						  UpsampledDepth,
 						  nullptr,
 						  UpscaledSceneDepth,
+						  Inputs.SceneColor.ViewRect,
 						  SceneDepthJitter,
 						  UpscaledSceneDepth->Desc.NumSamples,
 						  false,
@@ -124,9 +125,9 @@ static bool DrawDebugPDE(
 			RDG_EVENT_NAME("DrawDebugPrimitives(SDPG_World)"),
 			PassParameters,
 			ERDGPassFlags::Raster,
-			[PassParameters, DebugView=Context.DebugView, OutputViewport](FRDGAsyncTask, FRHICommandList& RHICmdList)
+			[PassParameters, DebugView=Context.DebugView, ViewportRect=Context.SceneColor.ViewRect](FRDGAsyncTask, FRHICommandList& RHICmdList)
 		{
-			RHICmdList.SetViewport(OutputViewport.Rect.Min.X, OutputViewport.Rect.Min.Y, 0.0f, OutputViewport.Rect.Max.X, OutputViewport.Rect.Max.Y, 1.0f);
+			RHICmdList.SetViewport(ViewportRect.Min.X, ViewportRect.Min.Y, 0.0f, ViewportRect.Max.X, ViewportRect.Max.Y, 1.0f);
 
 			FMeshPassProcessorRenderState DrawRenderState;
 			DrawRenderState.SetDepthStencilAccess(FExclusiveDepthStencil::DepthWrite_StencilWrite);
diff --git a/Engine/Source/Runtime/Renderer/Private/PostProcess/PostProcessCompositeEditorPrimitives.cpp b/Engine/Source/Runtime/Renderer/Private/PostProcess/PostProcessCompositeEditorPrimitives.cpp
--- a/Engine/Source/Runtime/Renderer/Private/PostProcess/PostProcessCompositeEditorPrimitives.cpp
+++ b/Engine/Source/Runtime/Renderer/Private/PostProcess/PostProcessCompositeEditorPrimitives.cpp
@@ -281,6 +281,7 @@ static FEditorPrimitivesDrawingContext InitializeEditorPrimitivesTextures(
 							  UpsampledDepth,
 							  EditorPrimitiveColor,
 							  EditorPrimitiveDepth,
+							  EditorView->ViewRect,
 							  SceneDepthJitter,
 							  NumMSAASamples);
 		}	
diff --git a/Engine/Source/Runtime/Renderer/Private/PostProcess/PostProcessCompositePrimitivesCommon.cpp b/Engine/Source/Runtime/Renderer/Private/PostProcess/PostProcessCompositePrimitivesCommon.cpp
--- a/Engine/Source/Runtime/Renderer/Private/PostProcess/PostProcessCompositePrimitivesCommon.cpp
+++ b/Engine/Source/Runtime/Renderer/Private/PostProcess/PostProcessCompositePrimitivesCommon.cpp
@@ -96,6 +96,7 @@ public:
 		SHADER_PARAMETER_STRUCT_REF(FViewUniformShaderParameters, View)
 		SHADER_PARAMETER_STRUCT(FScreenPassTextureViewportParameters, Color)
 		SHADER_PARAMETER_STRUCT(FScreenPassTextureViewportParameters, Depth)
+		SHADER_PARAMETER(FVector4f, OutputViewRectMinAndInvSize)
 		SHADER_PARAMETER_RDG_TEXTURE(Texture2D, ColorTexture)
 		SHADER_PARAMETER_SAMPLER(SamplerState, ColorSampler)
 		SHADER_PARAMETER_RDG_TEXTURE(Texture2D, DepthTexture)
@@ -254,6 +255,7 @@ void PopulateDepthPass(FRDGBuilder& GraphBuilder,
 	const FScreenPassTexture& InSceneDepth,
 	FRDGTextureRef OutPopColor,
 	FRDGTextureRef OutPopDepth,
+	const FIntRect& OutputViewRect,
 	const FVector2f& SceneDepthJitter,
 	uint32 NumMSAASamples,
 	bool bForceDrawColor,
@@ -273,6 +275,9 @@ void PopulateDepthPass(FRDGBuilder& GraphBuilder,
 	PassParameters->Color = GetScreenPassTextureViewportParameters(FScreenPassTextureViewport(InSceneColor));
 	PassParameters->Depth = GetScreenPassTextureViewportParameters(FScreenPassTextureViewport(InSceneDepth));
 
+	FIntPoint ViewRectSize = OutputViewRect.Size();
+	PassParameters->OutputViewRectMinAndInvSize = FVector4f(OutputViewRect.Min.X, OutputViewRect.Min.Y, 1.0f / ViewRectSize.X, 1.0f / ViewRectSize.Y);
+
 	PassParameters->DepthTexture = InSceneDepth.Texture;
 	PassParameters->DepthSampler = PointClampSampler;
 	PassParameters->DepthTextureJitter = SceneDepthJitter;
@@ -293,12 +298,12 @@ void PopulateDepthPass(FRDGBuilder& GraphBuilder,
 		GraphBuilder,
 		View.ShaderMap,
 		RDG_EVENT_NAME("PopulateDepth %dx%d%s",
-			View.ViewRect.Width(),
-			View.ViewRect.Height(),
+			OutputViewRect.Width(),
+			OutputViewRect.Height(),
 			NumMSAASamples > 1 ? TEXT(" MSAA") : TEXT("")),
 		PixelShader,
 		PassParameters,
-		View.ViewRect,
+		OutputViewRect,
 		/* BlendState = */ nullptr,
 		/* RasterizerState = */ nullptr,
 		TStaticDepthStencilState<true, CF_Always>::GetRHI());
diff --git a/Engine/Source/Runtime/Renderer/Private/PostProcess/PostProcessCompositePrimitivesCommon.h b/Engine/Source/Runtime/Renderer/Private/PostProcess/PostProcessCompositePrimitivesCommon.h
--- a/Engine/Source/Runtime/Renderer/Private/PostProcess/PostProcessCompositePrimitivesCommon.h
+++ b/Engine/Source/Runtime/Renderer/Private/PostProcess/PostProcessCompositePrimitivesCommon.h
@@ -105,6 +105,7 @@ void PopulateDepthPass(FRDGBuilder& GraphBuilder,
 	const FScreenPassTexture& InSceneDepth, 
 	FRDGTextureRef OutPopColor, 
 	FRDGTextureRef OutPopDepth, 
+	const FIntRect& OutputViewRect,
 	const FVector2f& SceneDepthJitter,
 	uint32 NumMSAASamples,
 	bool bForceDrawColor = false,
