[Lit Translucency] Fix crash when running with r.TranslucencyLightingVolume.Visualize but lighting is disabled.

#Synced-CL 51253558

diff --git a/Engine/Source/Runtime/Renderer/Private/TranslucentLighting.cpp b/Engine/Source/Runtime/Renderer/Private/TranslucentLighting.cpp
--- a/Engine/Source/Runtime/Renderer/Private/TranslucentLighting.cpp
+++ b/Engine/Source/Runtime/Renderer/Private/TranslucentLighting.cpp
@@ -2956,7 +2956,11 @@ void FDeferredShadingSceneRenderer::RenderTranslucencyVolumeVisualization(
 	const FMinimalSceneTextures& SceneTextures,
 	const FTranslucencyLightingVolumeTextures& TranslucencyLightingVolumeTextures)
 {
-	if (!GUseTranslucencyLightingVolumes || !GTranslucencyLightingVolumeVisualize || !ViewFamily.EngineShowFlags.DirectLighting || !RHISupportsVolumeTextureRendering(ViewFamily.GetShaderPlatform()))
+	if (!GUseTranslucencyLightingVolumes
+		|| !GTranslucencyLightingVolumeVisualize
+		|| !ViewFamily.EngineShowFlags.DirectLighting
+		|| !RHISupportsVolumeTextureRendering(ViewFamily.GetShaderPlatform())
+		|| !TranslucencyLightingVolumeTextures.IsValid())
 	{
 		return;
 	}
