[Transmission] Add CVar to skip transmission denoising for ray traced shadow r.Shadow.denoiser.transmission 0. When transmission is disabled, all noise is expected to be denoised in the temporal upscaler. This is to deal with potential transmission bleeding when we have a good later phase denoiser.

#Synced-CL 51225677

diff --git a/Engine/Shaders/Private/ScreenSpaceDenoise/SSDSpatialAccumulation.usf b/Engine/Shaders/Private/ScreenSpaceDenoise/SSDSpatialAccumulation.usf
--- a/Engine/Shaders/Private/ScreenSpaceDenoise/SSDSpatialAccumulation.usf
+++ b/Engine/Shaders/Private/ScreenSpaceDenoise/SSDSpatialAccumulation.usf
@@ -33,7 +33,7 @@
 #define CONFIG_SIGNAL_PROCESSING DIM_SIGNAL_PROCESSING
 #define CONFIG_UPSCALE           DIM_UPSCALE
 #define CONFIG_SIGNAL_BATCH_SIZE DIM_SIGNAL_BATCH_SIZE
-
+#define CONFIG_DENOISE_TRANSMISSION DIM_DENOISE_TRANSMISSION
 
 // Configures all the pass for each individual signals.
 #if CONFIG_SIGNAL_PROCESSING == SIGNAL_PROCESSING_SHADOW_VISIBILITY_MASK
@@ -607,6 +607,24 @@ FSSDTexture2D SignalInputUint_Textures_3;
 #define SignalInputUint_Textures_3 SignalInputUint_Textures_0
 #endif
 
+#if DIM_STAGE == STAGE_FINAL_OUTPUT && \
+    CONFIG_SIGNAL_PROCESS == SIGNAL_PROCESSING_SHADOW_VISIBILITY_MASK && \
+    !CONFIG_DENOISE_TRANSMISSION
+
+Texture2D OriginalInput_Textures_0;
+
+#if CONFIG_INPUT_TEXTURE_COUNT > 1
+Texture2D OriginalInput_Textures_1;
+#endif
+
+#if CONFIG_INPUT_TEXTURE_COUNT > 2
+Texture2D OriginalInput_Textures_2;
+#endif
+
+#if CONFIG_INPUT_TEXTURE_COUNT > 3
+Texture2D OriginalInput_Textures_3;
+#endif
+#endif
 
 #if !defined(CONFIG_OUTPUT_TEXTURE_COUNT)
 	#error Missing CONFIG_OUTPUT_TEXTURE_COUNT
@@ -1262,8 +1280,32 @@ void MainCS(
 				float Shadow = GetSamplePenumbraSafe(OutputSamples.Array[MultiplexId]);
 
 				const float ShadowFadeFraction = 1;
-				float SSSTransmission = (OutputSamples.Array[MultiplexId].SampleCount > 0 ? OutputSamples.Array[MultiplexId].TransmissionDistance / OutputSamples.Array[MultiplexId].SampleCount : OutputSamples.Array[MultiplexId].TransmissionDistance);
+				float SSSTransmission = 0;
 
+				#if (CONFIG_SIGNAL_PROCESSING == SIGNAL_PROCESSING_SHADOW_VISIBILITY_MASK && !CONFIG_DENOISE_TRANSMISSION)
+				{
+					
+					if (MultiplexId == 0)
+						SSSTransmission = OriginalInput_Textures_0[OutputPixelPostion].a;
+					#if CONFIG_SIGNAL_BATCH_SIZE > 1
+					else if (MultiplexId == 1)
+						SSSTransmission = OriginalInput_Textures_1[OutputPixelPostion].a;
+					#endif
+					#if CONFIG_SIGNAL_BATCH_SIZE > 2
+					else if (MultiplexId == 2)
+						SSSTransmission = OriginalInput_Textures_2[OutputPixelPostion].a;
+					#endif
+					#if CONFIG_SIGNAL_BATCH_SIZE > 3
+					else if (MultiplexId == 3)
+						SSSTransmission = OriginalInput_Textures_3[OutputPixelPostion].a;
+					#endif
+						
+				}
+				#else
+				{
+					SSSTransmission = (OutputSamples.Array[MultiplexId].SampleCount > 0 ? OutputSamples.Array[MultiplexId].TransmissionDistance / OutputSamples.Array[MultiplexId].SampleCount : OutputSamples.Array[MultiplexId].TransmissionDistance);
+				}
+				#endif
 				// 0 is shadowed, 1 is unshadowed
 				// RETURN_COLOR not needed unless writing to SceneColor;
 				float FadedShadow = lerp(1.0f, Shadow, ShadowFadeFraction);
diff --git a/Engine/Source/Runtime/Renderer/Private/ScreenSpaceDenoise.cpp b/Engine/Source/Runtime/Renderer/Private/ScreenSpaceDenoise.cpp
--- a/Engine/Source/Runtime/Renderer/Private/ScreenSpaceDenoise.cpp
+++ b/Engine/Source/Runtime/Renderer/Private/ScreenSpaceDenoise.cpp
@@ -43,6 +43,13 @@ static TAutoConsoleVariable<int32> CVarShadowHistoryConvolutionSampleCount(
 	TEXT("Number of samples to use to convolve the history over time."),
 	ECVF_RenderThreadSafe);
 
+static TAutoConsoleVariable<int32> CVarShadowDenoiserTransmission(
+	TEXT("r.Shadow.Denoiser.Transmission"), 1,
+	TEXT("Whether to denoise the transmission component for ray traced shadow.\n")
+	TEXT(" 0: Use the raw transmission from the ray traced shadow pass.\n")
+	TEXT(" 1: Denoise the transmission, might have transmission bleeding (default)."),
+	ECVF_RenderThreadSafe);
+
 static TAutoConsoleVariable<int32> CVarUseReflectionDenoiser(
 	TEXT("r.Reflections.Denoiser"),
 	2,
@@ -1023,8 +1030,8 @@ class FSSDSpatialAccumulationCS : public FGlobalShader
 
 	class FStageDim : SHADER_PERMUTATION_ENUM_CLASS("DIM_STAGE", EStage);
 	class FUpscaleDim : SHADER_PERMUTATION_BOOL("DIM_UPSCALE");
-
-	using FPermutationDomain = TShaderPermutationDomain<FSignalProcessingDim, FStageDim, FUpscaleDim, FSignalBatchSizeDim, FMultiSPPDim>;
+	class FDenoiseTransmissionDim : SHADER_PERMUTATION_BOOL("DIM_DENOISE_TRANSMISSION");
+	using FPermutationDomain = TShaderPermutationDomain<FSignalProcessingDim, FStageDim, FUpscaleDim, FSignalBatchSizeDim, FMultiSPPDim, FDenoiseTransmissionDim>;
 	
 	static bool ShouldCompilePermutation(const FGlobalShaderPermutationParameters& Parameters)
 	{
@@ -1093,6 +1100,14 @@ class FSSDSpatialAccumulationCS : public FGlobalShader
 			return false;
 		}
 
+		// transmission denoising permutation is only required to change for ShadowVisibilitymask 
+		// FinalOutput.
+		if (!(SignalProcessing == ESignalProcessing::ShadowVisibilityMask &&
+			PermutationVector.Get<FStageDim>() == EStage::FinalOutput))
+		{
+			PermutationVector.Set<FDenoiseTransmissionDim>(false);
+		}
+
 		// Compile out the shader if this permutation gets remapped.
 		if (RemapPermutationVector(PermutationVector) != PermutationVector)
 		{
@@ -1141,6 +1156,7 @@ class FSSDSpatialAccumulationCS : public FGlobalShader
 		SHADER_PARAMETER_STRUCT_INCLUDE(FSSDConvolutionMetaData, ConvolutionMetaData)
 
 		SHADER_PARAMETER_STRUCT(FSSDSignalTextures, SignalInput)
+		SHADER_PARAMETER_STRUCT(FSSDSignalTextures, OriginalInput)
 		SHADER_PARAMETER_STRUCT(FSSDSignalSRVs, SignalInputUint)
 		SHADER_PARAMETER_STRUCT(FSSDSignalUAVs, SignalOutput)
 
@@ -2298,6 +2314,7 @@ static void DenoiseSignalAtConstantPixelDensity(
 			PassParameters->CommonParameters = CommonParameters;
 			PassParameters->ConvolutionMetaData = ConvolutionMetaData;
 			PassParameters->SignalInput = SignalHistory;
+			PassParameters->OriginalInput = InputSignal;
 			PassParameters->SignalOutput = CreateMultiplexedUAVs(GraphBuilder, *OutputSignal);
 
 			FSSDSpatialAccumulationCS::FPermutationDomain PermutationVector;
@@ -2305,6 +2322,7 @@ static void DenoiseSignalAtConstantPixelDensity(
 			PermutationVector.Set<FSignalBatchSizeDim>(Settings.SignalBatchSize);
 			PermutationVector.Set<FSSDSpatialAccumulationCS::FStageDim>(FSSDSpatialAccumulationCS::EStage::FinalOutput);
 			PermutationVector.Set<FMultiSPPDim>(true);
+			PermutationVector.Set<FSSDSpatialAccumulationCS::FDenoiseTransmissionDim>(CVarShadowDenoiserTransmission.GetValueOnRenderThread()!= 0);
 
 			TShaderMapRef<FSSDSpatialAccumulationCS> ComputeShader(View.ShaderMap, PermutationVector);
 			FComputeShaderUtils::AddPass(
