Moved BeginDeferredCulling out of ShadowDepths profiling scope to reduce confusion.
#rnx

#Synced-CL 51369525

diff --git a/Engine/Source/Runtime/Renderer/Private/ShadowDepthRendering.cpp b/Engine/Source/Runtime/Renderer/Private/ShadowDepthRendering.cpp
--- a/Engine/Source/Runtime/Renderer/Private/ShadowDepthRendering.cpp
+++ b/Engine/Source/Runtime/Renderer/Private/ShadowDepthRendering.cpp
@@ -123,7 +123,8 @@ IMPLEMENT_SHADER_TYPE(template<>, TScreenVSForGS<true>, TEXT("/Engine/Private/Sc
 static TAutoConsoleVariable<int32> CVarShouldBeginDeferredCullingAfterShadowRendering(
 	TEXT("r.Shadow.ShouldBeginDeferredCullingAfterShadowRendering"),
 	1,
-	TEXT("Temporary: If turned on, a new deferred scope of instance culling will be started directly after shadow depth rendering (iff instance occlusion culling is enabled) to work around a bug."),
+	TEXT("If turned on (default), a new deferred scope of instance culling will be started directly after shadow depth rendering (iff instance occlusion culling is enabled).\n")\
+	TEXT("  This has a mild performance benefit on some platforms."),
 	ECVF_RenderThreadSafe);
 
 TAutoConsoleVariable<int32> CVarNaniteShadows(
@@ -1789,6 +1790,7 @@ void FSceneRenderer::RenderShadowDepthMaps(FRDGBuilder& GraphBuilder, FDynamicSh
 		return;
 	}
 
+	{
 	RDG_CSV_STAT_EXCLUSIVE_SCOPE(GraphBuilder, RenderShadows);
 
 	TRACE_CPUPROFILER_EVENT_SCOPE(FSceneRenderer::RenderShadowDepthMaps);
@@ -2050,12 +2052,6 @@ void FSceneRenderer::RenderShadowDepthMaps(FRDGBuilder& GraphBuilder, FDynamicSh
 
 	}, &GPersistentShadowsPipe);
 	
-	// Begin another deferred batch here to avoid any subsequent passes enabling HZB for all (an unfortunate side effect).
-	if (FInstanceCullingContext::IsOcclusionCullingEnabled() && CVarShouldBeginDeferredCullingAfterShadowRendering.GetValueOnRenderThread() != 0)
-	{
-		InstanceCullingManager.BeginDeferredCulling(GraphBuilder);
-	}
-
 	if (ShouldRenderHeterogeneousVolumes(Scene) && ShouldHeterogeneousVolumesCastShadows() &&
 		(HeterogeneousVolumes::GetShadowPipeline() == HeterogeneousVolumes::EShadowPipeline::LiveShading))
 	{
@@ -2071,7 +2067,15 @@ void FSceneRenderer::RenderShadowDepthMaps(FRDGBuilder& GraphBuilder, FDynamicSh
 			);
 		}
 	}
+	}
 
+	// Begin another deferred batch for any instance culling after ShadowDepths. Orininally introduced to work around a bug (now fixed) 
+	// that would cause subsequent passes enabling HZB for all passes. Left in place because it can under some circumstances be a small 
+	// performance win, possibly by changing the batching allowing a different permutation to be picked.
+	if (FInstanceCullingContext::IsOcclusionCullingEnabled() && CVarShouldBeginDeferredCullingAfterShadowRendering.GetValueOnRenderThread() != 0)
+	{
+		InstanceCullingManager.BeginDeferredCulling(GraphBuilder);
+	}
 	bShadowDepthRenderCompleted = true;
 }
 
