Path tracing: Heterogeneous Volume preshading support in Callable shaders.

#tests EngineTest, SVT-backed volumetrics
#Synced-CL 51373270

diff --git a/Engine/Shaders/Private/HeterogeneousVolumes/HeterogeneousVolumesMaterialBakingPipeline.usf b/Engine/Shaders/Private/HeterogeneousVolumes/HeterogeneousVolumesMaterialBakingPipeline.usf
--- a/Engine/Shaders/Private/HeterogeneousVolumes/HeterogeneousVolumesMaterialBakingPipeline.usf
+++ b/Engine/Shaders/Private/HeterogeneousVolumes/HeterogeneousVolumesMaterialBakingPipeline.usf
@@ -4,6 +4,7 @@
 
 #define SUPPORT_CONTACT_SHADOWS 0
 #define EYE_ADAPTATION_LOOSE_PARAMETERS 1
+#define SQRT_THREE 1.7320508075688772
 
 #include "/Engine/Generated/Material.ush"
 #include "../BlueNoise.ush"
@@ -25,6 +26,9 @@ float4x4 WorldToLocal;
 float3 LocalBoundsOrigin;
 float3 LocalBoundsExtent;
 int PrimitiveId;
+int bJitter;
+int bEvaluateEmission;
+int bEvaluateAlbedo;
 
 // Volume data
 int3 VolumeResolution;
@@ -44,6 +48,11 @@ FPrimitiveSceneData GetPrimitiveData(FMaterialPixelParameters Parameters)
 	return GetPrimitiveData(Parameters, LocalToWorld, WorldToLocal, LocalBoundsOrigin, LocalBoundsExtent);
 }
 
+uint2 VoxelIndex2D(uint3 VoxelIndex)
+{
+	return uint2(VoxelIndex.y * VolumeResolution.x + VoxelIndex.x, VoxelIndex.z);
+}
+
 [numthreads(THREADGROUP_SIZE_3D, THREADGROUP_SIZE_3D, THREADGROUP_SIZE_3D)]
 void HeterogeneousVolumesBakeMaterialCS(
 	uint3 GroupThreadId : SV_GroupThreadID,
@@ -55,7 +64,18 @@ void HeterogeneousVolumesBakeMaterialCS(
 	uint3 VoxelIndex = DispatchThreadId;
 	if (all(VoxelIndex < VolumeResolution))
 	{
-		float3 UVW = (VoxelIndex + 0.5) / float3(VolumeResolution);
+		float3 WorldOrigin = mul(float4(0.0, 0.0, 0.0, 1.0), LocalToWorld).xyz;
+		float3 WorldVoxelPos = mul(float4(1.0, 1.0, 1.0, 1.0), LocalToWorld).xyz;
+		float WorldVoxelVolume = length(WorldVoxelPos - WorldOrigin);
+
+		float3 Jitter = 0.5;
+		if (bJitter)
+		{
+			Jitter.xy = BlueNoiseVec2(VoxelIndex2D(VoxelIndex), View.StateFrameIndex);
+			Jitter.z = BlueNoiseScalar(VoxelIndex2D(VoxelIndex), View.StateFrameIndex);
+		}
+
+		float3 UVW = clamp((VoxelIndex + Jitter) / float3(VolumeResolution), 0, VolumeResolution);
 		float3 LocalBoundsMin = LocalBoundsOrigin - LocalBoundsExtent;
 		float3 LocalPosition = UVW * 2.0 * LocalBoundsExtent + LocalBoundsMin;
 		float3 WorldPosition = mul(float4(LocalPosition, 1.0), LocalToWorld).xyz;
@@ -75,9 +95,17 @@ void HeterogeneousVolumesBakeMaterialCS(
 		CalcPixelMaterialInputs(MaterialParameters, PixelMaterialInputs);
 
 		// Extract volume rendering coefficients
+		if (bEvaluateAlbedo)
+		{
+			RWAlbedoTexture[VoxelIndex] = SampleAlbedo(PixelMaterialInputs);
+		}
+
+		if (bEvaluateEmission)
+		{
+			RWEmissionTexture[VoxelIndex] = SampleEmissive(PixelMaterialInputs);
+		}
+
 		RWExtinctionTexture[VoxelIndex] = SampleExtinctionCoefficients(PixelMaterialInputs);
-		RWEmissionTexture[VoxelIndex] = SampleEmissive(PixelMaterialInputs);
-		RWAlbedoTexture[VoxelIndex] = SampleAlbedo(PixelMaterialInputs);
 	}
 }
 
@@ -153,7 +181,6 @@ int FixedBottomLevelResolution;
 //RWBuffer<uint> RWHashToVoxelBuffer;
 
 // Sampling mode
-int bJitter;
 int bSampleAtVertices;
 
 float HomogeneousThreshold;
@@ -273,6 +300,10 @@ void RasterizeBottomLevelOrthoGridCS(
 			float3 LocalPosition = mul(float4(WorldPosition, 1), WorldToLocal).xyz;
 			if (LocalPositionIntersectsPrimitive(LocalPosition))
 			{
+				float3 WorldOrigin = mul(float4(0.0, 0.0, 0.0, 1.0), LocalToWorld).xyz;
+				float3 WorldVoxelPos = mul(float4(1.0, 1.0, 1.0, 1.0), LocalToWorld).xyz;
+				float WorldVoxelVolume = length(WorldVoxelPos - WorldOrigin);
+
 				// Setup evaluation context
 				FMaterialPixelParameters MaterialParameters = MakeInitializedMaterialPixelParameters();
 				MaterialParameters.LWCData = MakeMaterialLWCData(MaterialParameters);
@@ -288,8 +319,8 @@ void RasterizeBottomLevelOrthoGridCS(
 				CalcPixelMaterialInputs(MaterialParameters, PixelMaterialInputs);
 
 				// Extract volume rendering coefficients
-				Extinction = SampleExtinctionCoefficients(PixelMaterialInputs);
-				Emission = SampleEmissive(PixelMaterialInputs);
+				Extinction = SampleExtinctionCoefficients(PixelMaterialInputs) / WorldVoxelVolume * SQRT_THREE;
+				Emission = SampleEmissive(PixelMaterialInputs) / WorldVoxelVolume * SQRT_THREE;
 				Albedo = SampleAlbedo(PixelMaterialInputs) * Extinction;
 			}
 		}
@@ -490,6 +521,10 @@ void RasterizeBottomLevelFrustumGridCS(
 
 	if (WorldPositionIntersectsPrimitive(WorldPosition))
 	{
+		float3 WorldOrigin = mul(float4(0.0, 0.0, 0.0, 1.0), LocalToWorld).xyz;
+		float3 WorldVoxelPos = mul(float4(1.0, 1.0, 1.0, 1.0), LocalToWorld).xyz;
+		float WorldVoxelVolume = length(WorldVoxelPos - WorldOrigin);
+
 		FMaterialPixelParameters MaterialParameters = MakeInitializedMaterialPixelParameters();
 		MaterialParameters.LWCData = MakeMaterialLWCData(MaterialParameters);
 		MaterialParameters.PrimitiveId = PrimitiveId;
@@ -504,7 +539,7 @@ void RasterizeBottomLevelFrustumGridCS(
 		CalcPixelMaterialInputs(MaterialParameters, PixelMaterialInputs);
 
 		// Rasterize coefficients
-		Extinction = SampleExtinctionCoefficients(PixelMaterialInputs);
+		Extinction = SampleExtinctionCoefficients(PixelMaterialInputs) / WorldVoxelVolume * SQRT_THREE;
 		Emission = SampleEmissive(PixelMaterialInputs);
 		Albedo = SampleAlbedo(PixelMaterialInputs) * Extinction;
 	}
@@ -588,4 +623,4 @@ void RasterizeBottomLevelFrustumGridCS(
 			SetAlbedo(RWScatteringGridBuffer[BottomLevelVoxelLinearIndex], Albedo);
 		}
 	}
-}
\ No newline at end of file
+}
diff --git a/Engine/Shaders/Private/HeterogeneousVolumes/HeterogeneousVolumesPreshadingPipeline.usf b/Engine/Shaders/Private/HeterogeneousVolumes/HeterogeneousVolumesPreshadingPipeline.usf
--- a/Engine/Shaders/Private/HeterogeneousVolumes/HeterogeneousVolumesPreshadingPipeline.usf
+++ b/Engine/Shaders/Private/HeterogeneousVolumes/HeterogeneousVolumesPreshadingPipeline.usf
@@ -49,6 +49,29 @@ void CopyTexture3D(
 	}
 }
 
+Texture3D InputTextureLHS;
+Texture3D InputTextureRHS;
+
+[numthreads(THREADGROUP_SIZE_3D, THREADGROUP_SIZE_3D, THREADGROUP_SIZE_3D)]
+void MaxTexture3D(
+	uint3 DispatchThreadId : SV_DispatchThreadID
+)
+{
+	uint3 Index = DispatchThreadId;
+	if (all(Index < TextureResolution))
+	{
+		float3 ValueLHS = InputTextureLHS[Index];
+		float3 ValueRHS = InputTextureRHS[Index];
+#if 0
+		if (any(ValueRHS > 0.0))
+		{
+			ValueRHS = 1.0;
+		}
+#endif
+		RWOutputTexture[Index] = max(ValueLHS, ValueRHS);
+	}
+}
+
 SamplerState TextureSampler;
 
 [numthreads(THREADGROUP_SIZE_3D, THREADGROUP_SIZE_3D, THREADGROUP_SIZE_3D)]
diff --git a/Engine/Shaders/Private/HeterogeneousVolumes/HeterogeneousVolumesVoxelGridPipeline.usf b/Engine/Shaders/Private/HeterogeneousVolumes/HeterogeneousVolumesVoxelGridPipeline.usf
--- a/Engine/Shaders/Private/HeterogeneousVolumes/HeterogeneousVolumesVoxelGridPipeline.usf
+++ b/Engine/Shaders/Private/HeterogeneousVolumes/HeterogeneousVolumesVoxelGridPipeline.usf
@@ -539,7 +539,6 @@ void DownsampleMajorantVoxelGridCS(
 	if (LinearThreadIndex == 0)
 	{
 		FMajorantData MajorantData;
-		MajorantData.Minorant = 0;
 		MajorantData.Majorant = GSMajorant[0];
 		SetMajorantData(RWOutputBuffer[LinearOuputIndex], MajorantData);
 	}
diff --git a/Engine/Shaders/Private/HeterogeneousVolumes/HeterogeneousVolumesVoxelGridUtils.ush b/Engine/Shaders/Private/HeterogeneousVolumes/HeterogeneousVolumesVoxelGridUtils.ush
--- a/Engine/Shaders/Private/HeterogeneousVolumes/HeterogeneousVolumesVoxelGridUtils.ush
+++ b/Engine/Shaders/Private/HeterogeneousVolumes/HeterogeneousVolumesVoxelGridUtils.ush
@@ -8,59 +8,50 @@
 
 struct FMajorantData
 {
-	float Minorant;
 	float Majorant;
 };
 
 FMajorantData CreateMajorantData()
 {
 	FMajorantData MajorantData;
-	MajorantData.Minorant = 1.0e4;
 	MajorantData.Majorant = 0.0;
 	return MajorantData;
 }
 
 void UpdateMajorantData(inout FMajorantData MajorantData, float3 Value)
 {
-	float MinValue = min(Value.x, min(Value.y, Value.z));
 	float MaxValue = max(Value.x, max(Value.y, Value.z));
-	MajorantData.Minorant = min(MajorantData.Minorant, MinValue);
 	MajorantData.Majorant = max(MajorantData.Majorant, MaxValue);
 }
 
-FMajorantData CreateMajorantData(float3 MinValue, float3 MaxValue)
+FMajorantData CreateMajorantData(float3 Majorant)
 {
 	FMajorantData MajorantData = CreateMajorantData();
-	UpdateMajorantData(MajorantData, MinValue);
-	UpdateMajorantData(MajorantData, MaxValue);
+	UpdateMajorantData(MajorantData, Majorant);
 	return MajorantData;
 }
 
 // Updating majorant data should occur for a single media
 void UpdateMajorantData(inout FMajorantData MajorantData, FMajorantData Value)
 {
-	MajorantData.Minorant = min(MajorantData.Minorant, Value.Minorant);
 	MajorantData.Majorant = max(MajorantData.Majorant, Value.Majorant);
 }
 
 // Merging majorant data should occur when mixing media
 void MergeMajorantData(inout FMajorantData MajorantData, FMajorantData OverlappingMajorant)
 {
-	MajorantData.Minorant += OverlappingMajorant.Minorant;
 	MajorantData.Majorant += OverlappingMajorant.Majorant;
 }
 
 FMajorantData GetMajorantData(FScalarGridData MajorantGridData)
 {
 	FMajorantData Data;
-	Data.Minorant = asfloat(MajorantGridData.PackedData[0]);
 	Data.Majorant = asfloat(MajorantGridData.PackedData[1]);
 	return Data;
 }
 
 void SetMajorantData(inout FScalarGridData MajorantGridData, FMajorantData MajorantData)
 {
-	MajorantGridData.PackedData[0] = asuint(MajorantData.Minorant);
 	MajorantGridData.PackedData[1] = asuint(MajorantData.Majorant);
 }
 
diff --git a/Engine/Shaders/Private/PathTracing/PathTracingCore.ush b/Engine/Shaders/Private/PathTracing/PathTracingCore.ush
--- a/Engine/Shaders/Private/PathTracing/PathTracingCore.ush
+++ b/Engine/Shaders/Private/PathTracing/PathTracingCore.ush
@@ -863,7 +863,7 @@ FPackedPathTracingPayload TraceTransparentRay(inout FPathState PathState, int Bo
 							float3 Ro = PathState.Ray.Origin;
 							float3 Rd = PathState.Ray.Direction;
 							float3 TranslatedWorldPos = Ro + Interval.VolumeTMin * Rd;
-							FVolumeShadedResult Result = VolumeGetDensity(TranslatedWorldPos, Interval.VolumeMask, PATH_TRACER_KERNEL_TYPE == PATH_TRACER_KERNEL_CAMERA);
+							FVolumeShadedResult Result = VolumeGetDensity(TranslatedWorldPos, TrackingResult.HeterogeneousVolumesIndex, Interval.VolumeMask, PATH_TRACER_KERNEL_TYPE == PATH_TRACER_KERNEL_CAMERA);
 
 							// clamp to make sure we never exceed the majorant (should not be the case, but need to avoid any possible numerical issues)
 							float3 SigmaT = min(Result.SigmaT, TrackingResult.SigmaBar);
diff --git a/Engine/Shaders/Private/PathTracing/PathTracingHeterogeneousVolumeMaterialShader.usf b/Engine/Shaders/Private/PathTracing/PathTracingHeterogeneousVolumeMaterialShader.usf
new file mode 100644
--- /dev/null
+++ b/Engine/Shaders/Private/PathTracing/PathTracingHeterogeneousVolumeMaterialShader.usf
@@ -0,0 +1,273 @@
+// Copyright Epic Games, Inc. All Rights Reserved.
+
+#define PATH_TRACING 1
+#define PATH_TRACER_USE_CLOUD_SHADER 1
+
+#define USE_PRESHADED_UNIFORM_BUFFER 1
+#define USE_CALLABLE_SHADERS 1
+#define USE_MULTIPLE_VOLUMES 1
+
+#include "/Engine/Private/Common.ush"
+#include "/Engine/Private/RayTracing/RayTracingCommon.ush"
+#include "/Engine/Private/PathTracing/PathTracingShaderUtils.ush"
+#include "/Engine/Private/PathTracing/Volume/PathTracingVolumeSampling.ush"
+#include "/Engine/Private/HeterogeneousVolumes/HeterogeneousVolumesTracingUtils.ush"
+#include "/Engine/Private/HeterogeneousVolumes/HeterogeneousVolumesVoxelGridTypes.ush"
+#include "/Engine/Private/HeterogeneousVolumes/HeterogeneousVolumesVoxelGridUtils.ush"
+#include "/Engine/Private/PathTracing/Volume/PathTracingVolumeCommon.ush"
+#include "/Engine/Private/PathTracing/Volume/PathTracingHeterogeneousVolumesForwardDecl.ush"
+#include "/Engine/Private/PathTracing/Volume/PathTracingHeterogeneousVolumesMaterialCommon.ush"
+#include "/Engine/Private/PathTracing/Volume/PathTracingHeterogeneousVolumesTrackingCommon.ush"
+#include "/Engine/Private/PathTracing/Volume/PathTracingHeterogeneousVolumesPreshadingCommon.ush"
+
+RAY_TRACING_ENTRY_CALLABLE(PathTracingHeterogeneousVolumeMaterialShader, FPackedPathTracingPayload, PackedPayload)
+{
+	// Not implemented
+}
+
+RAY_TRACING_ENTRY_CALLABLE(PathTracingHeterogeneousVolumesIntersectCS, FPackedPathTracingPayload, PackedPayload)
+{
+	ResolvedView = ResolveView();
+	FRayDesc Ray = GetHeterogeneousVolumeCallableShaderRay(PackedPayload);
+
+	// Convert ray to local-space
+	float3 WorldOrigin = GetHeterogeneousVolumeCallableShaderWorldPos(Ray.Origin);
+	float3 WorldDirection = Ray.Direction;
+	float3 LocalOrigin = mul(float4(WorldOrigin, 1), PreshadingUniformBuffer.WorldToLocal).xyz;
+	float3 LocalDirection = mul(float4(WorldDirection, 0), PreshadingUniformBuffer.WorldToLocal).xyz;
+
+	float2 LocalRayHitT = IntersectAABB(
+		LocalOrigin,
+		LocalDirection,
+		Ray.TMin,
+		Ray.TMax,
+		PreshadingUniformBuffer.LocalBoundsOrigin - PreshadingUniformBuffer.LocalBoundsExtent,
+		PreshadingUniformBuffer.LocalBoundsOrigin + PreshadingUniformBuffer.LocalBoundsExtent
+	);
+
+#if 0
+	// Refine intersection interval to include non-zero majorant voxels
+	if (LocalRayHitT.y - LocalRayHitT.x > 0)
+	{
+		// Convert to Majorant voxel-space units
+		float3 LocalSpanT0 = LocalOrigin + LocalDirection * LocalRayHitT.x;
+		float3 LocalSpanT1 = LocalOrigin + LocalDirection * LocalRayHitT.y;
+		float3 LocalBoundsMin = PreshadingUniformBuffer.LocalBoundsOrigin - PreshadingUniformBuffer.LocalBoundsExtent;
+		float3 LocalBoundsDiagonal = 2.0 * PreshadingUniformBuffer.LocalBoundsExtent;
+		float3 MajorantVoxelSpanT0 = (LocalSpanT0 - LocalBoundsMin) / LocalBoundsDiagonal * PreshadingUniformBuffer.MajorantResolution;
+		float3 MajorantVoxelSpanT1 = (LocalSpanT1 - LocalBoundsMin) / LocalBoundsDiagonal * PreshadingUniformBuffer.MajorantResolution;
+
+		float3 MajorantVoxelOrigin = MajorantVoxelSpanT0;
+		float3 MajorantVoxelDirection = MajorantVoxelSpanT1 - MajorantVoxelSpanT0;
+		float MajorantVoxelTMin = 0.0;
+		float MajorantVoxelTMax = length(MajorantVoxelDirection);
+		MajorantVoxelDirection /= MajorantVoxelTMax;
+		float3 MajorantVoxelDirectionInv = 1.0 / MajorantVoxelDirection;
+		float VoxelToLocalScale = (LocalRayHitT.y - LocalRayHitT.x) / MajorantVoxelTMax;
+
+		float3 MajorantVoxelPos = MajorantVoxelOrigin;
+		float3 MajorantVoxelBounds;
+		MajorantVoxelBounds.x = sign(MajorantVoxelDirection.x) > 0 ? floor(MajorantVoxelPos.x) + 1 : ceil(MajorantVoxelPos.x) - 1;
+		MajorantVoxelBounds.y = sign(MajorantVoxelDirection.y) > 0 ? floor(MajorantVoxelPos.y) + 1 : ceil(MajorantVoxelPos.y) - 1;
+		MajorantVoxelBounds.z = sign(MajorantVoxelDirection.z) > 0 ? floor(MajorantVoxelPos.z) + 1 : ceil(MajorantVoxelPos.z) - 1;
+
+		float3 MajorantVoxelBoundsHitT = (MajorantVoxelBounds - MajorantVoxelPos) * MajorantVoxelDirectionInv;
+		// Remove floating-point rounding error
+		float Epsilon = 1.0e-4;
+		if (MajorantVoxelBoundsHitT.x <= Epsilon) MajorantVoxelBoundsHitT.x += abs(MajorantVoxelDirectionInv.x);
+		if (MajorantVoxelBoundsHitT.y <= Epsilon) MajorantVoxelBoundsHitT.y += abs(MajorantVoxelDirectionInv.y);
+		if (MajorantVoxelBoundsHitT.z <= Epsilon) MajorantVoxelBoundsHitT.z += abs(MajorantVoxelDirectionInv.z);
+
+		// March in voxel space
+		float2 MajorantVoxelRayHitT = float2(MajorantVoxelTMax, MajorantVoxelTMin);
+		float MajorantVoxelT = 0;
+		while (MajorantVoxelT < MajorantVoxelTMax)
+		{
+			float MajorantVoxelDeltaT = 0.0;
+			if (MajorantVoxelBoundsHitT.x < MajorantVoxelBoundsHitT.y)
+			{
+				if (MajorantVoxelBoundsHitT.x < MajorantVoxelBoundsHitT.z)
+				{
+					MajorantVoxelDeltaT = MajorantVoxelBoundsHitT.x - MajorantVoxelT;
+					MajorantVoxelBoundsHitT.x += abs(MajorantVoxelDirectionInv.x);
+				}
+				else
+				{
+					MajorantVoxelDeltaT = MajorantVoxelBoundsHitT.z - MajorantVoxelT;
+					MajorantVoxelBoundsHitT.z += abs(MajorantVoxelDirectionInv.z);
+				}
+			}
+			else
+			{
+				if (MajorantVoxelBoundsHitT.y < MajorantVoxelBoundsHitT.z)
+				{
+					MajorantVoxelDeltaT = MajorantVoxelBoundsHitT.y - MajorantVoxelT;
+					MajorantVoxelBoundsHitT.y += abs(MajorantVoxelDirectionInv.y);
+				}
+				else
+				{
+					MajorantVoxelDeltaT = MajorantVoxelBoundsHitT.z - MajorantVoxelT;
+					MajorantVoxelBoundsHitT.z += abs(MajorantVoxelDirectionInv.z);
+				}
+			}
+
+			// Clip voxel delta-t by overall voxel ray-length.
+			if (MajorantVoxelT + MajorantVoxelDeltaT > MajorantVoxelTMax)
+			{
+				MajorantVoxelDeltaT = MajorantVoxelTMax - MajorantVoxelT;
+			}
+
+			// Convert midpoint to UVW for sampling
+			float3 MajorantVoxelEndPos = MajorantVoxelPos + MajorantVoxelDirection * MajorantVoxelDeltaT;
+			float3 MajorantVoxelSamplePos = (MajorantVoxelPos + MajorantVoxelEndPos) * 0.5;
+			float3 UVW = MajorantVoxelSamplePos / PreshadingUniformBuffer.MajorantResolution;
+
+			int MipLevel = 0;
+			float3 Majorant = Texture3DSampleLevel(
+				PreshadingUniformBuffer.MajorantTexture,
+				PreshadingUniformBuffer.PointTextureSampler,
+				UVW,
+				MipLevel
+			).rgb;
+
+			// Expand hit interval if non-zero majorant voxel
+			if (any(Majorant > 0))
+			{
+				MajorantVoxelRayHitT.x = min(MajorantVoxelRayHitT.x, MajorantVoxelT);
+				MajorantVoxelRayHitT.y = max(MajorantVoxelRayHitT.y, MajorantVoxelT + MajorantVoxelDeltaT);
+			}
+
+			MajorantVoxelT += MajorantVoxelDeltaT;
+			MajorantVoxelPos = MajorantVoxelEndPos;
+		}
+
+		float LocalRayHitT0 = LocalRayHitT.x;
+		LocalRayHitT.x = LocalRayHitT0 + MajorantVoxelRayHitT.x * VoxelToLocalScale;
+		LocalRayHitT.y = LocalRayHitT0 + MajorantVoxelRayHitT.y * VoxelToLocalScale;
+	}
+#endif
+
+	SetHeterogeneousVolumeCallableShaderHitT(LocalRayHitT, PackedPayload);
+}
+
+RAY_TRACING_ENTRY_CALLABLE(PathTracingHeterogeneousVolumesGetDensityBoundsCS, FPackedPathTracingPayload, PackedPayload)
+{
+	// Not implemented
+}
+
+// Specialize GetDensity() to locally fetch from the preshaded volume
+FVolumeShadedResult HeterogeneousVolumesGetDensity(float3 TranslatedWorldPos, int HeterogeneousVolumesIndex)
+{
+	float3 SigmaT = 0;
+	float3 Albedo = 0;
+	float3 Emission = 0;
+
+	SampleVolumeRenderingCoefficients(TranslatedWorldPos, SigmaT, Albedo, Emission);
+
+	// Return struct
+	FVolumeShadedResult Result = (FVolumeShadedResult)0;
+	Result.SigmaT = SigmaT;
+	Result.SigmaSHG = Albedo * SigmaT;
+	Result.PhaseG = GetHeterogeneousVolumeCallableShaderPhaseG();
+	Result.Emission = Emission;
+	return Result;
+}
+
+RAY_TRACING_ENTRY_CALLABLE(PathTracingHeterogeneousVolumesGetDensityCS, FPackedPathTracingPayload, PackedPayload)
+{
+	ResolvedView = ResolveView();
+
+	float3 SigmaT = 0;
+	float3 Albedo = 0;
+	float3 Emission = 0;
+
+	float3 TranslatedWorldPos = GetHeterogeneousVolumeCallableShaderTranslatedWorldPos(PackedPayload);
+	SampleVolumeRenderingCoefficients(TranslatedWorldPos, SigmaT, Albedo, Emission);
+
+	SetHeterogeneousVolumeCallableShaderSigmaT(SigmaT, PackedPayload);
+	SetHeterogeneousVolumeCallableShaderAlbedo(Albedo, PackedPayload);
+	SetHeterogeneousVolumeCallableShaderEmission(Emission, PackedPayload);
+}
+
+RAY_TRACING_ENTRY_CALLABLE(PathTracingHeterogeneousVolumesSampleDistanceCS, FPackedPathTracingPayload, PackedPayload)
+{
+	ResolvedView = ResolveView();
+
+	FRandomSequence RandSequence = GetHeterogeneousVolumeCallableShaderRandomSequence(PackedPayload);
+	FRayDesc Ray = GetHeterogeneousVolumeCallableShaderRay(PackedPayload);
+	float3 PathThroughput = GetHeterogeneousVolumeCallableShaderThroughput(PackedPayload);
+	float Majorant = GetHeterogeneousVolumeCallableShaderMajorant(PackedPayload);
+
+	// The tracking interval may represent overlapping volumes, so refine it against the active volume bounds
+	{
+		float3 WorldOrigin = GetHeterogeneousVolumeCallableShaderWorldPos(Ray.Origin);
+		float3 WorldDirection = Ray.Direction;
+		float3 LocalOrigin = mul(float4(WorldOrigin, 1), PreshadingUniformBuffer.WorldToLocal).xyz;
+		float3 LocalDirection = mul(float4(WorldDirection, 0), PreshadingUniformBuffer.WorldToLocal).xyz;
+
+		 float2 LocalRayHitT = IntersectAABB(
+			LocalOrigin,
+			LocalDirection,
+			Ray.TMin,
+			Ray.TMax,
+			PreshadingUniformBuffer.LocalBoundsOrigin - PreshadingUniformBuffer.LocalBoundsExtent,
+			PreshadingUniformBuffer.LocalBoundsOrigin + PreshadingUniformBuffer.LocalBoundsExtent
+		);
+
+		 // Invalid intersection
+		 if (LocalRayHitT.x >= LocalRayHitT.y)
+		 {
+			 SetHeterogeneousVolumeCallableShaderThroughput(PathThroughput, PackedPayload);
+			 SetHeterogeneousVolumeCallableShaderPdf(1, PackedPayload);
+			 SetHeterogeneousVolumeCallableShaderDistance(Ray.TMax, PackedPayload);
+			 SetHeterogeneousVolumeCallableShaderMajorant(0, PackedPayload);
+			 SetHeterogeneousVolumeCallableShaderIsCollision(false, PackedPayload);
+			 return;
+		 }
+	}
+
+	bool bTransmittanceOnly = false;
+	FVolumeTrackingResult Result = HeterogeneousVolumesMajorantBasedTrackingPreshading(
+		PathThroughput,
+		Ray.Origin,
+		Ray.Direction,
+		Ray.TMin,
+		Ray.TMax,
+		bTransmittanceOnly,
+		CreateMajorantData(Majorant),
+		RandSequence
+	);
+
+	SetHeterogeneousVolumeCallableShaderThroughput(Result.Throughput, PackedPayload);
+	SetHeterogeneousVolumeCallableShaderPdf(Result.Pdf, PackedPayload);
+	SetHeterogeneousVolumeCallableShaderDistance(Result.Distance, PackedPayload);
+	SetHeterogeneousVolumeCallableShaderSigmaBar(Result.SigmaBar, PackedPayload);
+	SetHeterogeneousVolumeCallableShaderIsCollision(Result.bIsCollision, PackedPayload);
+
+	SetHeterogeneousVolumeCallableShaderRandomSequence(RandSequence, PackedPayload);
+}
+
+RAY_TRACING_ENTRY_CALLABLE(PathTracingHeterogeneousVolumesGetTransmittanceCS, FPackedPathTracingPayload, PackedPayload)
+{
+	ResolvedView = ResolveView();
+
+	FRandomSequence RandSequence = GetHeterogeneousVolumeCallableShaderRandomSequence(PackedPayload);
+	FRayDesc Ray = GetHeterogeneousVolumeCallableShaderRay(PackedPayload);
+	float3 PathThroughput = GetHeterogeneousVolumeCallableShaderThroughput(PackedPayload);
+
+	bool bTransmittanceOnly = true;
+	FMajorantData EmptyOverlappingMajorant = CreateMajorantData(0.0f);
+	FVolumeTrackingResult Result = HeterogeneousVolumesMajorantBasedTrackingPreshading(
+		PathThroughput,
+		Ray.Origin,
+		Ray.Direction,
+		Ray.TMin,
+		Ray.TMax,
+		bTransmittanceOnly,
+		EmptyOverlappingMajorant,
+		RandSequence
+	);
+
+	SetHeterogeneousVolumeCallableShaderThroughput(Result.Throughput, PackedPayload);
+	SetHeterogeneousVolumeCallableShaderRandomSequence(RandSequence, PackedPayload);
+}
diff --git a/Engine/Shaders/Private/PathTracing/Volume/PathTracingHeterogeneousVolumes.ush b/Engine/Shaders/Private/PathTracing/Volume/PathTracingHeterogeneousVolumes.ush
--- a/Engine/Shaders/Private/PathTracing/Volume/PathTracingHeterogeneousVolumes.ush
+++ b/Engine/Shaders/Private/PathTracing/Volume/PathTracingHeterogeneousVolumes.ush
@@ -2,23 +2,26 @@
 
 #pragma once
 
+#define USE_MEGAVOXEL 0
+#define USE_CALLABLE_SHADERS 1
+#define USE_MULTIPLE_VOLUMES 1
+
+#define USE_PRESHADED_UNIFORM_BUFFER (PATH_TRACER_HETEROGENEOUS_VOLUMES_PIPELINE != USE_MEGAVOXEL)
+
 #include "../../HeterogeneousVolumes/HeterogeneousVolumesTracingUtils.ush"
 #include "../../HeterogeneousVolumes/HeterogeneousVolumesVoxelGridTraversal.ush"
 #include "../Utilities/PathTracingRandomSequence.ush"
 #include "PathTracingVolumeSampling.ush"
-
-float3 GetTranslatedWorldPos(float3 WorldPos)
-{
-	float3 TranslatedWorldPos = DFFastToTranslatedWorld(WorldPos, PrimaryView.PreViewTranslation);
-	return TranslatedWorldPos;
-}
-
-float GetPhaseG()
-{
-	return 0.0;
-}
-
-FVolumeIntersection HeterogeneousVolumesIntersect(float3 Origin, float3 Direction, float TMin, float TMax)
+#include "PathTracingVolumeCommon.ush"
+#include "PathTracingHeterogeneousVolumesForwardDecl.ush"
+#include "PathTracingHeterogeneousVolumesMaterialCommon.ush"
+#include "PathTracingHeterogeneousVolumesTrackingCommon.ush"
+#include "PathTracingHeterogeneousVolumesPreshadingCommon.ush"
+
+//
+// Intersection API..
+//
+FVolumeIntersection HeterogeneousVolumesIntersectInternal(float3 Origin, float3 Direction, float TMin, float TMax)
 {
 	float3 WorldBoundsMin = FrustumGridUniformBuffer.TopLevelGridWorldBoundsMin;
 	float3 WorldBoundsMax = FrustumGridUniformBuffer.TopLevelGridWorldBoundsMax;
@@ -30,8 +33,8 @@ FVolumeIntersection HeterogeneousVolumesIntersect(float3 Origin, float3 Directio
 
 	if (OrthoGridUniformBuffer.bUseOrthoGrid || FrustumGridUniformBuffer.bUseFrustumGrid)
 	{
-		float3 TranslatedWorldBoundsMin = GetTranslatedWorldPos(WorldBoundsMin);
-		float3 TranslatedWorldBoundsMax = GetTranslatedWorldPos(WorldBoundsMax);
+		float3 TranslatedWorldBoundsMin = GetHeterogeneousVolumeCallableShaderTranslatedWorldPos(WorldBoundsMin);
+		float3 TranslatedWorldBoundsMax = GetHeterogeneousVolumeCallableShaderTranslatedWorldPos(WorldBoundsMax);
 
 		float2 RayHitT = IntersectAABB(
 			Origin,
@@ -47,20 +50,47 @@ FVolumeIntersection HeterogeneousVolumesIntersect(float3 Origin, float3 Directio
 	return CreateEmptyVolumeIntersection();
 }
 
-FVolumeDensityBounds HeterogeneousVolumesGetDensityBounds(float3 Origin, float3 Direction, float TMin, float TMax)
+FVolumeIntersection HeterogeneousVolumesIntersect(float3 Origin, float3 Direction, float TMin, float TMax)
+{
+#if PATH_TRACER_HETEROGENEOUS_VOLUMES_PIPELINE == USE_MEGAVOXEL
+	return HeterogeneousVolumesIntersectInternal(Origin, Direction, TMin, TMax);
+#else
+	return HeterogeneousVolumesIntersectPreshading(Origin, Direction, TMin, TMax);
+#endif // PATH_TRACER_HETEROGENEOUS_VOLUMES_PIPELINE
+}
+
+//
+// Density APIs..
+//
+FVolumeDensityBounds HeterogeneousVolumesGetDensityBoundsInternal(float3 Origin, float3 Direction, float TMin, float TMax)
 {
 	// Density bounds are iteratively discovered when executing the DDA on the majorant grid
 	return CreateVolumeDensityBound(0, 0);
 }
 
+FVolumeDensityBounds HeterogeneousVolumesGetDensityBoundsPreshading(float3 Origin, float3 Direction, float TMin, float TMax)
+{
+	// Density bounds are iteratively discovered when executing the DDA on the majorant grid
+	return CreateVolumeDensityBound(0, 1.0);
+}
+
+FVolumeDensityBounds HeterogeneousVolumesGetDensityBounds(float3 Origin, float3 Direction, float TMin, float TMax)
+{
+#if PATH_TRACER_HETEROGENEOUS_VOLUMES_PIPELINE == USE_MEGAVOXEL
+	return HeterogeneousVolumesGetDensityBoundsInternal(Origin, Direction, TMin, TMax);
+#else
+	return HeterogeneousVolumesGetDensityBoundsPreshading(Origin, Direction, TMin, TMax);
+#endif // PATH_TRACER_HETEROGENEOUS_VOLUMES_PIPELINE
+}
+
 FVolumeShadedResult GetOrthoVoxelGridDensity(float3 TranslatedWorldPos)
 {
 	float3 SigmaT = 0.0;
 	float3 Albedo = 0.0;
 	float3 Emission = 0.0;
 
-	float3 TranslatedWorldBoundsMin = GetTranslatedWorldPos(OrthoGridUniformBuffer.TopLevelGridWorldBoundsMin);
-	float3 TranslatedWorldBoundsMax = GetTranslatedWorldPos(OrthoGridUniformBuffer.TopLevelGridWorldBoundsMax);
+	float3 TranslatedWorldBoundsMin = GetHeterogeneousVolumeCallableShaderTranslatedWorldPos(OrthoGridUniformBuffer.TopLevelGridWorldBoundsMin);
+	float3 TranslatedWorldBoundsMax = GetHeterogeneousVolumeCallableShaderTranslatedWorldPos(OrthoGridUniformBuffer.TopLevelGridWorldBoundsMax);
 	float3 TranslatedWorldBoundsExtent = TranslatedWorldBoundsMax - TranslatedWorldBoundsMin;
 	float3 GridUV = (TranslatedWorldPos - TranslatedWorldBoundsMin) / TranslatedWorldBoundsExtent;
 	if (all(GridUV >= 0.0) && all(GridUV <= 1.0))
@@ -118,7 +148,7 @@ FVolumeShadedResult GetOrthoVoxelGridDensity(float3 TranslatedWorldPos)
 	FVolumeShadedResult Result = (FVolumeShadedResult)0;
 	Result.SigmaT = SigmaT;
 	Result.SigmaSHG = Albedo;
-	Result.PhaseG = GetPhaseG();
+	Result.PhaseG = GetHeterogeneousVolumeCallableShaderPhaseG();
 	Result.Emission = Emission;
 
 	return Result;
@@ -159,17 +189,17 @@ FVolumeShadedResult GetFrustumVoxelGridDensity(float3 TranslatedWorldPos, inout
 	}
 
 	// Return struct
-	FVolumeShadedResult Result = (FVolumeShadedResult) 0;
+	FVolumeShadedResult Result = (FVolumeShadedResult)0;
 	Result.SigmaT = SigmaT;
 	Result.SigmaSHG = Albedo;
-	Result.PhaseG = GetPhaseG();
+	Result.PhaseG = GetHeterogeneousVolumeCallableShaderPhaseG();
 	Result.Emission = Emission;
 	return Result;
 }
 
-FVolumeShadedResult HeterogeneousVolumesGetDensity(float3 TranslatedWorldPos)
+FVolumeShadedResult HeterogeneousVolumesGetDensityInternal(float3 TranslatedWorldPos)
 {
-	FVolumeShadedResult Result = (FVolumeShadedResult) 0;
+	FVolumeShadedResult Result = (FVolumeShadedResult)0;
 
 	bool bInFrustum = false;
 	if (FrustumGridUniformBuffer.bUseFrustumGrid)
@@ -183,155 +213,39 @@ FVolumeShadedResult HeterogeneousVolumesGetDensity(float3 TranslatedWorldPos)
 	return Result;
 }
 
-struct FVolumeTrackingResult
-{
-	float3 Throughput;
-	float Pdf;
-	float Distance;
-	float3 SigmaBar;
-	bool bIsCollision;
-};
-
-FVolumeTrackingResult CreateVolumeTrackingResult(float Distance)
-{
-	FVolumeTrackingResult Result = (FVolumeTrackingResult)0;
-	Result.Throughput = 1;
-	Result.Pdf = 0;
-	Result.Distance = Distance;
-	Result.SigmaBar = 0;
-	Result.bIsCollision = false;
-
-	return Result;
-}
-
-// Concatenates the result of rhs onto lhs
-FVolumeTrackingResult ConcatenateVolumeTrackingResult(FVolumeTrackingResult lhs, FVolumeTrackingResult rhs)
-{
-	FVolumeTrackingResult Result;
-	Result.Throughput = lhs.Throughput * rhs.Throughput;
-	Result.Pdf = lhs.Pdf * rhs.Pdf;
-	Result.Distance = lhs.Distance + rhs.Distance;
-	Result.SigmaBar = rhs.SigmaBar;
-	Result.bIsCollision = rhs.bIsCollision;
-	return Result;
-}
-
-FVolumeTrackingResult HeterogeneousVolumesDeltaTracking(float3 PathThroughput, float3 Origin, float3 Direction, float TMax, FMajorantData MajorantData, inout FRandomSequence RandSequence)
-{
-	FVolumeTrackingResult Sample = CreateVolumeTrackingResult(0);
-	Sample.SigmaBar = max(MajorantData.Majorant, 0);
-
-	float RandValue = RandSequence.Get1D();
-#define SAMPLE_ACHROMATIC 0
-#if SAMPLE_ACHROMATIC
-	float DeltaT = -log(1.0 - RandValue) / Sample.SigmaBar.x;
-#else
-	float DeltaT = SampleSpectralTransmittance(RandValue, Sample.SigmaBar, PathThroughput);
-	if (DeltaT < 0)
-	{
-		Sample.Throughput = 0;
-		Sample.Distance = -1;
-		return Sample;
-	}
-#endif
-
-	Sample.Distance = min(Sample.Distance + DeltaT, TMax);
-	if (Sample.Distance < TMax)
-	{
-		Sample.bIsCollision = true;
-		float CollisionWeight = 1.0;
-#if SAMPLE_ACHROMATIC
-		Sample.Throughput *= CollisionWeight / Sample.SigmaBar.x;
-		float Transmittance = exp(-Sample.SigmaBar.x * Sample.Distance);
-		Sample.Pdf = Transmittance * Sample.SigmaBar.x;
-#else
-		float4 EvalResult = EvaluateSpectralTransmittanceHit(Sample.Distance, Sample.SigmaBar, PathThroughput);
-		Sample.Throughput = CollisionWeight * EvalResult.xyz;
-		Sample.Pdf = EvalResult.w;
-	}
-	else
-	{
-		float4 EvalResult = EvaluateSpectralTransmittanceMiss(Sample.Distance, Sample.SigmaBar, PathThroughput);
-		Sample.Throughput = EvalResult.xyz;
-		Sample.Pdf = EvalResult.w;
-#endif
-	}
-
-	return Sample;
-}
-
-FVolumeTrackingResult HeterogeneousVolumesRatioTracking(float3 PathThroughput, float3 Origin, float3 Direction, float TMax, FMajorantData MajorantData, inout FRandomSequence RandSequence)
+FVolumeShadedResult HeterogeneousVolumesGetDensity(float3 TranslatedWorldPos, int HeterogeneousVolumesIndex)
 {
-	FVolumeTrackingResult Sample = CreateVolumeTrackingResult(0);
-	Sample.SigmaBar = max(MajorantData.Majorant, 0);
-
-	while ((Sample.Distance < TMax))
-	{
-		float RandValue = RandSequence.Get1D();
-#if SAMPLE_ACHROMATIC
-		float DeltaT = -log(1.0 - RandValue) / Sample.SigmaBar.x;
-#else
-		float3 CombinedThroughput = PathThroughput * Sample.Throughput;
-		float DeltaT = SampleSpectralTransmittance(RandValue, Sample.SigmaBar, CombinedThroughput);
-		if (DeltaT < 0)
-		{
-			Sample.Throughput = 0;
-			Sample.Distance = -1;
-			return Sample;
-		}
-#endif
-
-		float CollisionDistance = Sample.Distance + DeltaT;
-		if (CollisionDistance < TMax)
-		{
-			float3 SamplePos = Origin + Direction * CollisionDistance;
-			float3 SigmaT = HeterogeneousVolumesGetDensity(SamplePos).SigmaT;
-			float3 SpectralNullCollisionWeight = max(Sample.SigmaBar - SigmaT, 0);
-#if SAMPLE_ACHROMATIC
-			Sample.Throughput *= SpectralNullCollisionWeight / Sample.SigmaBar.x;
-			float Transmittance = exp(-Sample.SigmaBar.x * DeltaT);
-			Sample.Pdf *= Transmittance * Sample.SigmaBar.x;
+#if PATH_TRACER_HETEROGENEOUS_VOLUMES_PIPELINE == USE_MEGAVOXEL
+	return HeterogeneousVolumesGetDensityInternal(TranslatedWorldPos);
 #else
-			float4 EvalResult = EvaluateSpectralTransmittanceHit(DeltaT, Sample.SigmaBar, CombinedThroughput);
-			Sample.Throughput *= SpectralNullCollisionWeight * EvalResult.xyz;
-			Sample.Pdf *= EvalResult.w;
-		}
-		else
-		{
-			float4 EvalResult = EvaluateSpectralTransmittanceMiss(TMax - Sample.Distance, Sample.SigmaBar, CombinedThroughput);
-			Sample.Throughput *= EvalResult.xyz;
-			Sample.Pdf *= EvalResult.w;
-#endif
-		}
-
-		Sample.Distance = min(CollisionDistance, TMax);
-	}
-
-	return Sample;
+	return HeterogeneousVolumesGetDensityPreshading(TranslatedWorldPos, HeterogeneousVolumesIndex);
+#endif // PATH_TRACER_HETEROGENEOUS_VOLUMES_PIPELINE
 }
 
-FVolumeTrackingResult HeterogeneousVolumesMajorantBasedTracking(float3 PathThroughput, float3 WorldRayOrigin, float3 Direction, float TMin, float TMax, int bTransmittanceOnly, FMajorantData OverlappingMajorant, inout FRandomSequence RandSequence)
+FVolumeTrackingResult HeterogeneousVolumesMajorantBasedTrackingInternal(float3 PathThroughput, float3 WorldRayOrigin, float3 WorldRayDirection, float WorldTMin, float WorldTMax, int bTransmittanceOnly, FMajorantData OverlappingMajorant, inout FRandomSequence RandSequence)
 {
-	FVolumeTrackingResult Result = CreateVolumeTrackingResult(TMin);
+	FVolumeTrackingResult Result = CreateVolumeTrackingResult(WorldTMin);
 	Result.Throughput = PathThroughput;
 
-	float3 WorldRayBegin = WorldRayOrigin + Direction * TMin;
-	float3 WorldRayEnd = WorldRayOrigin + Direction * TMax;
+	float3 WorldRayBegin = WorldRayOrigin + WorldRayDirection * WorldTMin;
+	float3 WorldRayEnd = WorldRayOrigin + WorldRayDirection * WorldTMax;
 	float WorldRayTMax = length(WorldRayEnd - WorldRayBegin);
 
 	// Transform to voxel-space
-	float3 WorldBoundsMin = GetTranslatedWorldPos(OrthoGridUniformBuffer.TopLevelGridWorldBoundsMin);
-	float3 WorldBoundsMax = GetTranslatedWorldPos(OrthoGridUniformBuffer.TopLevelGridWorldBoundsMax);
+	float3 WorldBoundsMin = GetHeterogeneousVolumeCallableShaderTranslatedWorldPos(OrthoGridUniformBuffer.TopLevelGridWorldBoundsMin);
+	float3 WorldBoundsMax = GetHeterogeneousVolumeCallableShaderTranslatedWorldPos(OrthoGridUniformBuffer.TopLevelGridWorldBoundsMax);
 	float3 TopLevelGridWorldBoundsExtent = WorldBoundsMax - WorldBoundsMin;
 	float3 VoxelRayBegin = (WorldRayBegin - WorldBoundsMin) / TopLevelGridWorldBoundsExtent * OrthoGridUniformBuffer.TopLevelGridResolution;
 	float3 VoxelRayEnd = (WorldRayEnd - WorldBoundsMin) / TopLevelGridWorldBoundsExtent * OrthoGridUniformBuffer.TopLevelGridResolution;
 	float3 VoxelRayDirection = VoxelRayEnd - VoxelRayBegin;
 	float VoxelRayTMin = 0.0;
 	float VoxelRayTMax = length(VoxelRayDirection);
+	float VoxelToWorldScale = WorldRayTMax / VoxelRayTMax;
 
 	if (VoxelRayTMin >= VoxelRayTMax && !bTransmittanceOnly)
 	{
-		FVolumeTrackingResult Sample = HeterogeneousVolumesDeltaTracking(Result.Throughput, WorldRayBegin, Direction, TMax - TMin, OverlappingMajorant, RandSequence);
+		float LocalToWorldScale = 1.0;
+		FVolumeTrackingResult Sample = HeterogeneousVolumesDeltaTracking(1, 0, WorldTMax - WorldTMin, LocalToWorldScale, OverlappingMajorant, RandSequence);
 		if (Sample.Distance > 0.0)
 		{
 			Result = ConcatenateVolumeTrackingResult(Result, Sample);
@@ -340,8 +254,6 @@ FVolumeTrackingResult HeterogeneousVolumesMajorantBasedTracking(float3 PathThrou
 	}
 	VoxelRayDirection /= VoxelRayTMax;
 
-	float VoxelToWorldScale = WorldRayTMax / VoxelRayTMax;
-
 	// March majorant grid via DDA
 	float3 VoxelRayDirectionInv = 1.0 / VoxelRayDirection;
 
@@ -404,17 +316,20 @@ FVolumeTrackingResult HeterogeneousVolumesMajorantBasedTracking(float3 PathThrou
 		FMajorantData MajorantData = GetMajorantData(OrthoGridUniformBuffer.MajorantGridBuffer[VoxelSampleLinearPos]);
 		MergeMajorantData(MajorantData, OverlappingMajorant);
 
-		float3 TrackingOrigin = WorldRayBegin + Direction * VoxelRayMarchT * VoxelToWorldScale;
+		float WorldRayMarchT = VoxelRayMarchT * VoxelToWorldScale;
+		float3 WorldTrackingOrigin = WorldRayBegin + WorldRayDirection * WorldRayMarchT;
 		float WorldRayMarchDeltaT = VoxelRayMarchDeltaT * VoxelToWorldScale;
 
 		FVolumeTrackingResult Sample;
 		if (bTransmittanceOnly)
 		{
-			Sample = HeterogeneousVolumesRatioTracking(Result.Throughput, TrackingOrigin, Direction, WorldRayMarchDeltaT, MajorantData, RandSequence);
+			float LocalToWorldScale = 1.0;
+			Sample = HeterogeneousVolumesRatioTracking(Result.Throughput, WorldTrackingOrigin, WorldRayDirection, WorldRayMarchDeltaT, LocalToWorldScale, MajorantData, RandSequence);
 		}
 		else
 		{
-			Sample = HeterogeneousVolumesDeltaTracking(Result.Throughput, TrackingOrigin, Direction, WorldRayMarchDeltaT, MajorantData, RandSequence);
+			float LocalToWorldScale = 1.0;
+			Sample = HeterogeneousVolumesDeltaTracking(1.0, 0.0, WorldRayMarchDeltaT, LocalToWorldScale, MajorantData, RandSequence);
 		}
 
 		if (Sample.Distance < 0)
@@ -428,18 +343,46 @@ FVolumeTrackingResult HeterogeneousVolumesMajorantBasedTracking(float3 PathThrou
 		VoxelRayMarchPos = VoxelEndPos;
 	}
 
+	// Convert distances to world-units
+
 	return Result;
 }
 
-float3 HeterogeneousVolumesGetTransmittance(float3 PathThroughput, float3 Origin, float3 Direction, float TMin, float TMax, inout FRandomSequence RandSequence)
+//
+// Distance Sampling API
+//
+
+FVolumeTrackingResult HeterogeneousVolumesSampleDistanceInternal(float3 PathThroughput, float3 Origin, float3 Direction, float TMin, float TMax, FMajorantData OverlappingMajorant, inout FRandomSequence RandSequence)
 {
-	bool bTransmittanceOnly = true;
-	FMajorantData EmptyOverlappingMajorant = CreateMajorantData(0.0f, 0.0f);
-	return HeterogeneousVolumesMajorantBasedTracking(PathThroughput, Origin, Direction, TMin, TMax, bTransmittanceOnly, EmptyOverlappingMajorant, RandSequence).Throughput;
+	bool bTransmittanceOnly = false;
+	return HeterogeneousVolumesMajorantBasedTrackingInternal(PathThroughput, Origin, Direction, TMin, TMax, bTransmittanceOnly, OverlappingMajorant, RandSequence);
 }
 
 FVolumeTrackingResult HeterogeneousVolumesSampleDistance(float3 PathThroughput, float3 Origin, float3 Direction, float TMin, float TMax, FMajorantData OverlappingMajorant, inout FRandomSequence RandSequence)
 {
-	bool bTransmittanceOnly = false;
-	return HeterogeneousVolumesMajorantBasedTracking(PathThroughput, Origin, Direction, TMin, TMax, bTransmittanceOnly, OverlappingMajorant, RandSequence);
+#if PATH_TRACER_HETEROGENEOUS_VOLUMES_PIPELINE == USE_MEGAVOXEL
+	return HeterogeneousVolumesSampleDistanceInternal(PathThroughput, Origin, Direction, TMin, TMax, OverlappingMajorant, RandSequence);
+#else
+	return HeterogeneousVolumesSampleDistancePreshading(PathThroughput, Origin, Direction, TMin, TMax, OverlappingMajorant, RandSequence);
+#endif // PATH_TRACER_HETEROGENEOUS_VOLUMES_PIPELINE
+}
+
+//
+// Transmittance API
+//
+
+float3 HeterogeneousVolumesGetTransmittanceInternal(float3 PathThroughput, float3 Origin, float3 Direction, float TMin, float TMax, inout FRandomSequence RandSequence)
+{
+	bool bTransmittanceOnly = true;
+	FMajorantData EmptyOverlappingMajorant = CreateMajorantData(0.0f);
+	return HeterogeneousVolumesMajorantBasedTrackingInternal(PathThroughput, Origin, Direction, TMin, TMax, bTransmittanceOnly, EmptyOverlappingMajorant, RandSequence).Throughput;
+}
+
+float3 HeterogeneousVolumesGetTransmittance(float3 PathThroughput, float3 Origin, float3 Direction, float TMin, float TMax, inout FRandomSequence RandSequence)
+{
+#if PATH_TRACER_HETEROGENEOUS_VOLUMES_PIPELINE == USE_MEGAVOXEL
+	return HeterogeneousVolumesGetTransmittanceInternal(PathThroughput, Origin, Direction, TMin, TMax, RandSequence);
+#else
+	return HeterogeneousVolumesGetTransmittancePreshading(PathThroughput, Origin, Direction, TMin, TMax, RandSequence);
+#endif // PATH_TRACER_HETEROGENEOUS_VOLUMES_PIPELINE
 }
diff --git a/Engine/Shaders/Private/PathTracing/Volume/PathTracingHeterogeneousVolumesForwardDecl.ush b/Engine/Shaders/Private/PathTracing/Volume/PathTracingHeterogeneousVolumesForwardDecl.ush
new file mode 100644
--- /dev/null
+++ b/Engine/Shaders/Private/PathTracing/Volume/PathTracingHeterogeneousVolumesForwardDecl.ush
@@ -0,0 +1,26 @@
+// Copyright Epic Games, Inc. All Rights Reserved.
+
+#pragma once
+
+FVolumeIntersection HeterogeneousVolumesIntersect(float3 Origin, float3 Direction, float TMin, float TMax);
+FVolumeDensityBounds HeterogeneousVolumesGetDensityBounds(float3 Origin, float3 Direction, float TMin, float TMax);
+FVolumeShadedResult HeterogeneousVolumesGetDensity(float3 TranslatedWorldPos, int CustomData);
+
+struct FVolumeTrackingResult
+{
+	float3 Throughput;
+	float Pdf;
+	float Distance;
+	float3 SigmaBar;
+	bool bIsCollision;
+	int HeterogeneousVolumesIndex;
+};
+
+FVolumeTrackingResult CreateVolumeTrackingResult(float Distance);
+FVolumeTrackingResult ConcatenateVolumeTrackingResult(FVolumeTrackingResult lhs, FVolumeTrackingResult rhs);
+
+FVolumeTrackingResult HeterogeneousVolumesSampleDistance(float3 PathThroughput, float3 Origin, float3 Direction, float TMin, float TMax, FMajorantData OverlappingMajorant, inout FRandomSequence RandSequence);
+float3 HeterogeneousVolumesGetTransmittance(float3 PathThroughput, float3 Origin, float3 Direction, float TMin, float TMax, inout FRandomSequence RandSequence);
+
+FVolumeTrackingResult HeterogeneousVolumesDeltaTracking(float3 PathThroughput, float LocalTMin, float LocalTMax, float LocalToWorldScale, FMajorantData MajorantData, inout FRandomSequence RandSequence);
+FVolumeTrackingResult HeterogeneousVolumesRatioTracking(float3 PathThroughput, float3 TranslatedWorldOrigin, float3 WorldDirection, float LocalTMax, float LocalToWorldScale, FMajorantData MajorantData, inout FRandomSequence RandSequence);
diff --git a/Engine/Shaders/Private/PathTracing/Volume/PathTracingHeterogeneousVolumesMaterialCommon.ush b/Engine/Shaders/Private/PathTracing/Volume/PathTracingHeterogeneousVolumesMaterialCommon.ush
new file mode 100644
--- /dev/null
+++ b/Engine/Shaders/Private/PathTracing/Volume/PathTracingHeterogeneousVolumesMaterialCommon.ush
@@ -0,0 +1,295 @@
+// Copyright Epic Games, Inc. All Rights Reserved.
+
+#pragma once
+
+#include "../../Common.ush"
+
+// Transform
+
+float3 GetHeterogeneousVolumeCallableShaderWorldPos(float3 TranslatedWorldPos)
+{
+	float3 WorldPos = DFFastLocalSubtractDemote(TranslatedWorldPos, PrimaryView.PreViewTranslation);
+	return WorldPos;
+}
+
+float3 GetHeterogeneousVolumeCallableShaderTranslatedWorldPos(float3 WorldPos)
+{
+	float3 TranslatedWorldPos = DFFastToTranslatedWorld(WorldPos, PrimaryView.PreViewTranslation);
+	return TranslatedWorldPos;
+}
+
+// Shading
+
+float GetHeterogeneousVolumeCallableShaderPhaseG()
+{
+	return 0.0;
+}
+
+// Intersection
+
+void SetHeterogeneousVolumeCallableShaderHitT(
+	float2 HitT,
+	inout FPackedPathTracingPayload Payload
+)
+{
+	Payload.HitT = HitT.x;
+	Payload.PackedData[0] = asuint(HitT.y);
+}
+
+float2 GetHeterogeneousVolumeCallableShaderHitT(
+	FPackedPathTracingPayload Payload
+)
+{
+	float2 HitT;
+	HitT.x = Payload.HitT;
+	HitT.y = asfloat(Payload.PackedData[0]);
+	return HitT;
+}
+
+// Delta Tracking
+
+void SetHeterogeneousVolumeCallableShaderRay(
+	FRayDesc Ray,
+	inout FPackedPathTracingPayload Payload
+)
+{
+	Payload.HitT = Ray.TMin;
+	Payload.PackedData[0] = asuint(Ray.TMax);
+	
+	Payload.PackedData[1] = asuint(Ray.Origin.x);
+	Payload.PackedData[2] = asuint(Ray.Origin.y);
+	Payload.PackedData[3] = asuint(Ray.Origin.z);
+
+	Payload.PackedData[4] = asuint(Ray.Direction.x);
+	Payload.PackedData[5] = asuint(Ray.Direction.y);
+	Payload.PackedData[6] = asuint(Ray.Direction.z);
+}
+
+FRayDesc GetHeterogeneousVolumeCallableShaderRay(
+	FPackedPathTracingPayload Payload
+)
+{
+	FRayDesc Ray;
+	Ray.TMin = Payload.HitT;
+	Ray.TMax = asfloat(Payload.PackedData[0]);
+
+	Ray.Origin.x = asfloat(Payload.PackedData[1]);
+	Ray.Origin.y = asfloat(Payload.PackedData[2]);
+	Ray.Origin.z = asfloat(Payload.PackedData[3]);
+
+	Ray.Direction.x = asfloat(Payload.PackedData[4]);
+	Ray.Direction.y = asfloat(Payload.PackedData[5]);
+	Ray.Direction.z = asfloat(Payload.PackedData[6]);
+	return Ray;
+}
+
+void SetHeterogeneousVolumeCallableShaderRandomSequence(
+	FRandomSequence RandSeq,
+	inout FPackedPathTracingPayload Payload
+)
+{
+	Payload.PackedData[7] = RandSeq.SampleIndex;
+	Payload.PackedData[8] = RandSeq.SampleSeed;
+}
+
+FRandomSequence GetHeterogeneousVolumeCallableShaderRandomSequence(
+	FPackedPathTracingPayload Payload
+)
+{
+	FRandomSequence RandSeq;
+	RandSeq.SampleIndex = Payload.PackedData[7];
+	RandSeq.SampleSeed = Payload.PackedData[8];
+	return RandSeq;
+}
+
+void SetHeterogeneousVolumeCallableShaderThroughput(
+	float3 Throughput,
+	inout FPackedPathTracingPayload Payload
+)
+{
+	Payload.PackedData[9] = asuint(Throughput.x);
+	Payload.PackedData[10] = asuint(Throughput.y);
+	Payload.PackedData[11] = asuint(Throughput.z);
+}
+
+float3 GetHeterogeneousVolumeCallableShaderThroughput(
+	FPackedPathTracingPayload Payload
+)
+{
+	float3 Throughput = 0;
+	Throughput.x = asfloat(Payload.PackedData[9]);
+	Throughput.y = asfloat(Payload.PackedData[10]);
+	Throughput.z = asfloat(Payload.PackedData[11]);
+	return Throughput;
+}
+
+void SetHeterogeneousVolumeCallableShaderMajorant(
+	float Majorant,
+	FPackedPathTracingPayload Payload
+)
+{
+	Payload.PackedData[12] = asuint(Majorant);
+}
+
+float GetHeterogeneousVolumeCallableShaderMajorant(
+	FPackedPathTracingPayload Payload
+)
+{
+	float Majorant = asfloat(Payload.PackedData[12]);
+	return Majorant;
+}
+
+// GetDensity
+
+void SetHeterogeneousVolumeCallableShaderTranslatedWorldPos(
+	float3 TranslatedWorldPos,
+	inout FPackedPathTracingPayload Payload
+)
+{
+	Payload.PackedData[0] = asuint(TranslatedWorldPos.x);
+	Payload.PackedData[1] = asuint(TranslatedWorldPos.y);
+	Payload.PackedData[2] = asuint(TranslatedWorldPos.z);
+}
+
+float3 GetHeterogeneousVolumeCallableShaderTranslatedWorldPos(
+	FPackedPathTracingPayload Payload
+)
+{
+	float3 TranslatedWorldPos = 0;
+	TranslatedWorldPos.x = asfloat(Payload.PackedData[0]);
+	TranslatedWorldPos.y = asfloat(Payload.PackedData[1]);
+	TranslatedWorldPos.z = asfloat(Payload.PackedData[2]);
+	return TranslatedWorldPos;
+}
+
+void SetHeterogeneousVolumeCallableShaderSigmaT(
+	float3 SigmaT,
+	inout FPackedPathTracingPayload Payload
+)
+{
+	Payload.PackedData[0] = asuint(SigmaT.x);
+	Payload.PackedData[1] = asuint(SigmaT.y);
+	Payload.PackedData[2] = asuint(SigmaT.z);
+}
+
+float3 GetHeterogeneousVolumeCallableShaderSigmaT(
+	FPackedPathTracingPayload Payload
+)
+{
+	float3 SigmaT = 0;
+	SigmaT.x = asfloat(Payload.PackedData[0]);
+	SigmaT.y = asfloat(Payload.PackedData[1]);
+	SigmaT.z = asfloat(Payload.PackedData[2]);
+	return SigmaT;
+}
+
+void SetHeterogeneousVolumeCallableShaderAlbedo(
+	float3 Albedo,
+	inout FPackedPathTracingPayload Payload
+)
+{
+	Payload.PackedData[3] = asuint(Albedo.x);
+	Payload.PackedData[4] = asuint(Albedo.y);
+	Payload.PackedData[5] = asuint(Albedo.z);
+}
+
+float3 GetHeterogeneousVolumeCallableShaderAlbedo(
+	FPackedPathTracingPayload Payload
+)
+{
+	float3 Albedo = 0;
+	Albedo.x = asfloat(Payload.PackedData[3]);
+	Albedo.y = asfloat(Payload.PackedData[4]);
+	Albedo.z = asfloat(Payload.PackedData[5]);
+	return Albedo;
+}
+
+void SetHeterogeneousVolumeCallableShaderEmission(
+	float3 Emission,
+	inout FPackedPathTracingPayload Payload
+)
+{
+	Payload.PackedData[6] = asuint(Emission.x);
+	Payload.PackedData[7] = asuint(Emission.y);
+	Payload.PackedData[8] = asuint(Emission.z);
+}
+
+float3 GetHeterogeneousVolumeCallableShaderEmission(
+	FPackedPathTracingPayload Payload
+)
+{
+	float3 Emission = 0;
+	Emission.x = asfloat(Payload.PackedData[6]);
+	Emission.y = asfloat(Payload.PackedData[7]);
+	Emission.z = asfloat(Payload.PackedData[8]);
+	return Emission;
+}
+
+void SetHeterogeneousVolumeCallableShaderPdf(
+	float Pdf,
+	inout FPackedPathTracingPayload Payload
+)
+{
+	Payload.PackedData[0] = asuint(Pdf);
+}
+
+float GetHeterogeneousVolumeCallableShaderPdf(
+	FPackedPathTracingPayload Payload
+)
+{
+	float Pdf = asfloat(Payload.PackedData[0]);
+	return Pdf;
+}
+
+void SetHeterogeneousVolumeCallableShaderDistance(
+	float Distance,
+	inout FPackedPathTracingPayload Payload
+)
+{
+	Payload.PackedData[1] = asuint(Distance);
+}
+
+float GetHeterogeneousVolumeCallableShaderDistance(
+	FPackedPathTracingPayload Payload
+)
+{
+	float Distance = asfloat(Payload.PackedData[1]);
+	return Distance;
+}
+
+void SetHeterogeneousVolumeCallableShaderSigmaBar(
+	float3 SigmaBar,
+	inout FPackedPathTracingPayload Payload
+)
+{
+	Payload.PackedData[2] = asuint(SigmaBar.x);
+	Payload.PackedData[3] = asuint(SigmaBar.y);
+	Payload.PackedData[4] = asuint(SigmaBar.z);
+}
+
+float3 GetHeterogeneousVolumeCallableShaderSigmaBar(
+	FPackedPathTracingPayload Payload
+)
+{
+	float3 SigmaBar = 0;
+	SigmaBar.x = asfloat(Payload.PackedData[2]);
+	SigmaBar.y = asfloat(Payload.PackedData[3]);
+	SigmaBar.z = asfloat(Payload.PackedData[4]);
+	return SigmaBar;
+}
+
+void SetHeterogeneousVolumeCallableShaderIsCollision(
+	bool bIsCollision,
+	inout FPackedPathTracingPayload Payload
+)
+{
+	Payload.PackedData[5] = bIsCollision ? 1 : 0;
+}
+
+bool GetHeterogeneousVolumeCallableShaderIsCollision(
+	FPackedPathTracingPayload Payload
+)
+{
+	bool bIsCollision = Payload.PackedData[5] == 1 ? true : false;
+	return bIsCollision;
+}
diff --git a/Engine/Shaders/Private/PathTracing/Volume/PathTracingHeterogeneousVolumesPreshadingCommon.ush b/Engine/Shaders/Private/PathTracing/Volume/PathTracingHeterogeneousVolumesPreshadingCommon.ush
new file mode 100644
--- /dev/null
+++ b/Engine/Shaders/Private/PathTracing/Volume/PathTracingHeterogeneousVolumesPreshadingCommon.ush
@@ -0,0 +1,422 @@
+// Copyright Epic Games, Inc. All Rights Reserved.
+
+#pragma once
+
+#ifndef USE_PRESHADED_UNIFORM_BUFFER
+#define USE_PRESHADED_UNIFORM_BUFFER 0
+#endif
+
+#ifndef USE_CALLABLE_SHADERS
+#define USE_CALLABLE_SHADERS 0
+#endif
+
+#ifndef USE_MULTIPLE_VOLUMES
+#define USE_MULTIPLE_VOLUMES 0
+#endif 
+
+#if USE_PRESHADED_UNIFORM_BUFFER
+
+int HeterogeneousVolumesNum;
+Buffer<int> HeterogeneousVolumesCallableShaderIdIntersect;
+Buffer<int> HeterogeneousVolumesCallableShaderIdGetDensityBounds;
+Buffer<int> HeterogeneousVolumesCallableShaderIdGetDensity;
+Buffer<int> HeterogeneousVolumesCallableShaderIdSampleDistance;
+Buffer<int> HeterogeneousVolumesCallableShaderIdGetTransmittance;
+
+void ExpandHitT(float2 RayHitT, inout float2 ResultHitT)
+{
+	ResultHitT[0] = min(RayHitT[0], ResultHitT[0]);
+	ResultHitT[1] = max(RayHitT[1], ResultHitT[1]);
+}
+
+FVolumeIntersection HeterogeneousVolumesIntersectPreshading(float3 TranslatedWorldOrigin, float3 WorldDirection, float TMin, float TMax)
+{
+	float2 RayHitT = float2(TMax, TMin);
+#if USE_CALLABLE_SHADERS
+	int HeterogeneousVolumesIndex = 0;
+#if USE_MULTIPLE_VOLUMES
+	for (HeterogeneousVolumesIndex = 0; HeterogeneousVolumesIndex < HeterogeneousVolumesNum; ++HeterogeneousVolumesIndex)
+#endif
+		if (HeterogeneousVolumesCallableShaderIdIntersect[HeterogeneousVolumesIndex] >= 0)
+		{
+			FPackedPathTracingPayload Payload = (FPackedPathTracingPayload)0;
+			{
+				FRayDesc Ray;
+				Ray.Origin = TranslatedWorldOrigin;
+				Ray.Direction = WorldDirection;
+				Ray.TMin = TMin;
+				Ray.TMax = TMax;
+				SetHeterogeneousVolumeCallableShaderRay(Ray, Payload);
+
+			}
+
+			CallShader(HeterogeneousVolumesCallableShaderIdIntersect[HeterogeneousVolumesIndex], Payload);
+			ExpandHitT(GetHeterogeneousVolumeCallableShaderHitT(Payload), RayHitT);
+		}
+#else
+	{
+		// Convert ray to local-space
+		float3 WorldOrigin = GetHeterogeneousVolumeCallableShaderWorldPos(TranslatedWorldOrigin);
+		float3 LocalOrigin = mul(float4(WorldOrigin, 1), PreshadingUniformBuffer.WorldToLocal).xyz;
+		float3 LocalDirection = mul(float4(WorldOrigin + WorldDirection, 0), PreshadingUniformBuffer.WorldToLocal).xyz;
+
+		RayHitT = IntersectAABB(
+			LocalOrigin,
+			LocalDirection,
+			TMin,
+			TMax,
+			PreshadingUniformBuffer.LocalBoundsOrigin - PreshadingUniformBuffer.LocalBoundsExtent,
+			PreshadingUniformBuffer.LocalBoundsOrigin + PreshadingUniformBuffer.LocalBoundsExtent
+		);
+	}
+#endif // USE_CALLABLE_SHADERS
+
+	if (RayHitT.y > RayHitT.x)
+	{
+		return CreateVolumeIntersection(RayHitT.x, RayHitT.y);
+	}
+	return CreateEmptyVolumeIntersection();
+}
+
+void SampleVolumeRenderingCoefficients(
+	float3 TranslatedWorldPos,
+	inout float3 SigmaT,
+	inout float3 Albedo,
+	inout float3 Emission
+)
+{
+	float3 LocalPos = mul(float4(GetHeterogeneousVolumeCallableShaderWorldPos(TranslatedWorldPos), 1), PreshadingUniformBuffer.WorldToLocal).xyz;
+	float3 LocalBoundsMin = PreshadingUniformBuffer.LocalBoundsOrigin - PreshadingUniformBuffer.LocalBoundsExtent;
+	float3 LocalBoundsMax = PreshadingUniformBuffer.LocalBoundsOrigin + PreshadingUniformBuffer.LocalBoundsExtent;
+	float3 UVW = (LocalPos - LocalBoundsMin) / (2.0 * PreshadingUniformBuffer.LocalBoundsExtent);
+
+	int MipLevel = 0;
+	SigmaT = Texture3DSampleLevel(
+		PreshadingUniformBuffer.ExtinctionTexture,
+		PreshadingUniformBuffer.TextureSampler,
+		UVW,
+		MipLevel
+	).rgb;
+
+	float3 Majorant = Texture3DSampleLevel(
+		PreshadingUniformBuffer.MajorantTexture,
+		PreshadingUniformBuffer.PointTextureSampler,
+		UVW,
+		MipLevel
+	).rgb;
+	SigmaT = min(SigmaT, Majorant);
+
+	Albedo = Texture3DSampleLevel(
+		PreshadingUniformBuffer.AlbedoTexture,
+		PreshadingUniformBuffer.TextureSampler,
+		UVW,
+		MipLevel
+	).rgb;
+
+	Emission = Texture3DSampleLevel(
+		PreshadingUniformBuffer.EmissionTexture,
+		PreshadingUniformBuffer.TextureSampler,
+		UVW,
+		MipLevel
+	).rgb;
+}
+
+FVolumeShadedResult HeterogeneousVolumesGetDensityPreshading(float3 TranslatedWorldPos, int HeterogeneousVolumesIndex)
+{
+	float3 SigmaT = 0;
+	float3 Albedo = 0;
+	float3 Emission = 0;
+
+#if USE_CALLABLE_SHADERS
+	if (HeterogeneousVolumesIndex < 0)
+	{
+		// Cycle through all volumes, compositing aggregate volumetric properties
+		int Index = 0;
+#if USE_MULTIPLE_VOLUMES
+		for (; Index < HeterogeneousVolumesNum; ++Index)
+#endif
+		{
+			FPackedPathTracingPayload Payload = (FPackedPathTracingPayload)0;
+			{
+				SetHeterogeneousVolumeCallableShaderTranslatedWorldPos(TranslatedWorldPos, Payload);
+			}
+
+			CallShader(HeterogeneousVolumesCallableShaderIdGetDensity[Index], Payload);
+
+			SigmaT += GetHeterogeneousVolumeCallableShaderSigmaT(Payload);
+			Albedo += GetHeterogeneousVolumeCallableShaderAlbedo(Payload);
+			Emission += GetHeterogeneousVolumeCallableShaderEmission(Payload);
+		}
+	}
+	else if (HeterogeneousVolumesCallableShaderIdGetDensity[HeterogeneousVolumesIndex] >= 0)
+	{
+		FPackedPathTracingPayload Payload = (FPackedPathTracingPayload)0;
+		{
+			SetHeterogeneousVolumeCallableShaderTranslatedWorldPos(TranslatedWorldPos, Payload);
+		}
+
+		CallShader(HeterogeneousVolumesCallableShaderIdGetDensity[HeterogeneousVolumesIndex], Payload);
+
+		SigmaT = GetHeterogeneousVolumeCallableShaderSigmaT(Payload);
+		Albedo = GetHeterogeneousVolumeCallableShaderAlbedo(Payload);
+		Emission = GetHeterogeneousVolumeCallableShaderEmission(Payload);
+	}
+#else
+	{
+		SampleVolumeRenderingCoefficients(TranslatedWorldPos, SigmaT, Albedo, Emission);
+	}
+#endif // USE_CALLABLE_SHADERS
+
+	// Return struct
+	FVolumeShadedResult Result = (FVolumeShadedResult)0;
+	Result.SigmaT = SigmaT;
+	Result.SigmaSHG = Albedo * SigmaT;
+	//Result.SigmaSHG = Albedo;
+	Result.PhaseG = GetHeterogeneousVolumeCallableShaderPhaseG();
+	Result.Emission = Emission;
+	return Result;
+}
+
+FVolumeTrackingResult HeterogeneousVolumesMajorantBasedTrackingPreshading(float3 PathThroughput, float3 TranslatedWorldOrigin, float3 WorldDirection, float TMin, float TMax, int bTransmittanceOnly, FMajorantData OverlappingMajorant, inout FRandomSequence RandSequence)
+{
+	FVolumeTrackingResult Result = CreateVolumeTrackingResult(TMin);
+	Result.Throughput = PathThroughput;
+
+	// Project ray-span to local space
+	float3 TranslatedWorldRayBegin = TranslatedWorldOrigin + WorldDirection * TMin;
+	float3 LocalSpanT0 = mul(float4(GetHeterogeneousVolumeCallableShaderWorldPos(TranslatedWorldOrigin + WorldDirection * TMin), 1), PreshadingUniformBuffer.WorldToLocal).xyz;
+	float3 LocalSpanT1 = mul(float4(GetHeterogeneousVolumeCallableShaderWorldPos(TranslatedWorldOrigin + WorldDirection * TMax), 1), PreshadingUniformBuffer.WorldToLocal).xyz;
+
+	float3 LocalOrigin = mul(float4(GetHeterogeneousVolumeCallableShaderWorldPos(TranslatedWorldOrigin), 1), PreshadingUniformBuffer.WorldToLocal).xyz;
+	float LocalTMin = length(LocalSpanT0 - LocalOrigin);
+	float LocalTMax = length(LocalSpanT1 - LocalOrigin);
+	float LocalToWorldScale = (TMax - TMin) / (LocalTMax - LocalTMin);
+
+	// Convert to Majorant voxel-space units
+	float3 LocalBoundsMin = PreshadingUniformBuffer.LocalBoundsOrigin - PreshadingUniformBuffer.LocalBoundsExtent;
+	float3 LocalBoundsDiagonal = 2.0 * PreshadingUniformBuffer.LocalBoundsExtent;
+	float3 MajorantVoxelSpanT0 = (LocalSpanT0 - LocalBoundsMin) / LocalBoundsDiagonal * PreshadingUniformBuffer.MajorantResolution;
+	float3 MajorantVoxelSpanT1 = (LocalSpanT1 - LocalBoundsMin) / LocalBoundsDiagonal * PreshadingUniformBuffer.MajorantResolution;
+
+	// Determine Majorant voxel-space ray parameters
+	float3 MajorantVoxelOrigin = MajorantVoxelSpanT0;
+	float3 MajorantVoxelDirection = MajorantVoxelSpanT1 - MajorantVoxelSpanT0;
+	float MajorantVoxelTMin = 0.0;
+	float MajorantVoxelTMax = length(MajorantVoxelDirection);
+	MajorantVoxelDirection /= MajorantVoxelTMax;
+	float3 MajorantVoxelDirectionInv = 1.0 / MajorantVoxelDirection;
+	float VoxelToWorldScale = (TMax - TMin) / MajorantVoxelTMax;
+	float VoxelToLocalScale = (LocalTMax - LocalTMin) / MajorantVoxelTMax;
+
+	if (MajorantVoxelTMin >= MajorantVoxelTMax && !bTransmittanceOnly)
+	{
+		float LocalToWorldScale = 1.0;
+		FVolumeTrackingResult Sample = HeterogeneousVolumesDeltaTracking(1, 0, TMax - TMin, LocalToWorldScale, OverlappingMajorant, RandSequence);
+		if (Sample.Distance > 0.0)
+		{
+			Result = ConcatenateVolumeTrackingResult(Result, Sample);
+		}
+		return Result;
+	}
+
+	float3 MajorantVoxelPos = MajorantVoxelOrigin;
+	float3 MajorantVoxelBounds;
+	MajorantVoxelBounds.x = sign(MajorantVoxelDirection.x) > 0 ? floor(MajorantVoxelPos.x) + 1 : ceil(MajorantVoxelPos.x) - 1;
+	MajorantVoxelBounds.y = sign(MajorantVoxelDirection.y) > 0 ? floor(MajorantVoxelPos.y) + 1 : ceil(MajorantVoxelPos.y) - 1;
+	MajorantVoxelBounds.z = sign(MajorantVoxelDirection.z) > 0 ? floor(MajorantVoxelPos.z) + 1 : ceil(MajorantVoxelPos.z) - 1;
+
+	float3 MajorantVoxelBoundsHitT = (MajorantVoxelBounds - MajorantVoxelPos) * MajorantVoxelDirectionInv;
+	// Remove floating-point rounding error
+	float Epsilon = 1.0e-4;
+	if (MajorantVoxelBoundsHitT.x <= Epsilon) MajorantVoxelBoundsHitT.x += abs(MajorantVoxelDirectionInv.x);
+	if (MajorantVoxelBoundsHitT.y <= Epsilon) MajorantVoxelBoundsHitT.y += abs(MajorantVoxelDirectionInv.y);
+	if (MajorantVoxelBoundsHitT.z <= Epsilon) MajorantVoxelBoundsHitT.z += abs(MajorantVoxelDirectionInv.z);
+
+	// Ray-march in majorant voxel space
+	float MajorantVoxelT = 0;
+	while ((MajorantVoxelT < MajorantVoxelTMax) && any(Result.Throughput > Epsilon) && !Result.bIsCollision)
+	{
+		float MajorantVoxelDeltaT = 0.0;
+		if (MajorantVoxelBoundsHitT.x < MajorantVoxelBoundsHitT.y)
+		{
+			if (MajorantVoxelBoundsHitT.x < MajorantVoxelBoundsHitT.z)
+			{
+				MajorantVoxelDeltaT = MajorantVoxelBoundsHitT.x - MajorantVoxelT;
+				MajorantVoxelBoundsHitT.x += abs(MajorantVoxelDirectionInv.x);
+			}
+			else
+			{
+				MajorantVoxelDeltaT = MajorantVoxelBoundsHitT.z - MajorantVoxelT;
+				MajorantVoxelBoundsHitT.z += abs(MajorantVoxelDirectionInv.z);
+			}
+		}
+		else
+		{
+			if (MajorantVoxelBoundsHitT.y < MajorantVoxelBoundsHitT.z)
+			{
+				MajorantVoxelDeltaT = MajorantVoxelBoundsHitT.y - MajorantVoxelT;
+				MajorantVoxelBoundsHitT.y += abs(MajorantVoxelDirectionInv.y);
+			}
+			else
+			{
+				MajorantVoxelDeltaT = MajorantVoxelBoundsHitT.z - MajorantVoxelT;
+				MajorantVoxelBoundsHitT.z += abs(MajorantVoxelDirectionInv.z);
+			}
+		}
+
+		// Clip voxel delta-t by overall voxel ray-length.
+		if (MajorantVoxelT + MajorantVoxelDeltaT > MajorantVoxelTMax)
+		{
+			MajorantVoxelDeltaT = MajorantVoxelTMax - MajorantVoxelT;
+		}
+
+		// Convert midpoint to UVW for sampling
+		float3 MajorantVoxelEndPos = MajorantVoxelPos + MajorantVoxelDirection * MajorantVoxelDeltaT;
+		float3 MajorantVoxelSamplePos = (MajorantVoxelPos + MajorantVoxelEndPos) * 0.5;
+		float3 UVW = MajorantVoxelSamplePos / PreshadingUniformBuffer.MajorantResolution;
+
+		int MipLevel = 0;
+		float3 Majorant = Texture3DSampleLevel(
+			PreshadingUniformBuffer.MajorantTexture,
+			// ORIGINAL BUG..
+			//PreshadingUniformBuffer.TextureSampler,
+			// WITH FIX..
+			PreshadingUniformBuffer.PointTextureSampler,
+			UVW,
+			MipLevel
+		).rgb;
+
+		FMajorantData MajorantData;
+		MajorantData.Majorant = max(Majorant.x, max(Majorant.y, Majorant.z));
+		MergeMajorantData(MajorantData, OverlappingMajorant);
+
+		float3 TranslatedWorldTrackingOrigin = TranslatedWorldRayBegin + WorldDirection * MajorantVoxelT * VoxelToWorldScale;
+		float WorldDeltaT = MajorantVoxelDeltaT * VoxelToWorldScale;
+		float LocalDeltaT = MajorantVoxelDeltaT * VoxelToLocalScale;
+
+		FVolumeTrackingResult Sample = CreateVolumeTrackingResult(0.0);
+		//if (any(MajorantData.Majorant > 0))
+		{
+			if (bTransmittanceOnly)
+			{
+				Sample = HeterogeneousVolumesRatioTracking(Result.Throughput, TranslatedWorldTrackingOrigin, WorldDirection, LocalDeltaT, LocalToWorldScale, MajorantData, RandSequence);
+			}
+			else
+			{
+				Sample = HeterogeneousVolumesDeltaTracking(1.0, 0.0, LocalDeltaT, LocalToWorldScale, MajorantData, RandSequence);
+			}
+
+			if (Sample.Distance < 0)
+			{
+				return Sample;
+			}
+		}
+		Result = ConcatenateVolumeTrackingResult(Result, Sample);
+
+		MajorantVoxelT += MajorantVoxelDeltaT;
+		MajorantVoxelPos = MajorantVoxelEndPos;
+	}
+
+	return Result;
+}
+
+FVolumeTrackingResult HeterogeneousVolumesSampleDistancePreshading(float3 PathThroughput, float3 Origin, float3 Direction, float TMin, float TMax, FMajorantData OverlappingMajorant, inout FRandomSequence RandSequence)
+{
+	bool bTransmittanceOnly = false;
+	float MinDistance = TMax;
+
+	FVolumeTrackingResult Result = (FVolumeTrackingResult)0;
+	FRandomSequence RandSequenceResult = RandSequence;
+#if USE_CALLABLE_SHADERS
+
+	Result.Throughput = PathThroughput;
+	Result.Pdf = 1.0;
+	Result.Distance = TMax;
+	int HeterogeneousVolumesIndex = 0;
+#if USE_MULTIPLE_VOLUMES
+	for (HeterogeneousVolumesIndex = 0; HeterogeneousVolumesIndex < HeterogeneousVolumesNum; ++HeterogeneousVolumesIndex)
+#endif
+		if (HeterogeneousVolumesCallableShaderIdSampleDistance[HeterogeneousVolumesIndex])
+		{
+			FPackedPathTracingPayload Payload = (FPackedPathTracingPayload)0;
+			{
+				FRayDesc Ray;
+				Ray.Origin = Origin;
+				Ray.Direction = Direction;
+				Ray.TMin = TMin;
+				Ray.TMax = TMax;
+
+				SetHeterogeneousVolumeCallableShaderRay(Ray, Payload);
+				SetHeterogeneousVolumeCallableShaderThroughput(PathThroughput, Payload);
+				SetHeterogeneousVolumeCallableShaderRandomSequence(RandSequence, Payload);
+				SetHeterogeneousVolumeCallableShaderMajorant(OverlappingMajorant.Majorant, Payload);
+			}
+
+			CallShader(HeterogeneousVolumesCallableShaderIdSampleDistance[HeterogeneousVolumesIndex], Payload);
+			RandSequenceResult = GetHeterogeneousVolumeCallableShaderRandomSequence(Payload);
+
+			float Distance = GetHeterogeneousVolumeCallableShaderDistance(Payload);
+			bool bIsCollision = GetHeterogeneousVolumeCallableShaderIsCollision(Payload);
+			if (bIsCollision && Distance <= MinDistance)
+			{
+				MinDistance = Distance;
+
+				Result.Throughput = GetHeterogeneousVolumeCallableShaderThroughput(Payload);
+				Result.Pdf = GetHeterogeneousVolumeCallableShaderPdf(Payload);
+				Result.Distance = GetHeterogeneousVolumeCallableShaderDistance(Payload);
+				Result.SigmaBar = GetHeterogeneousVolumeCallableShaderSigmaBar(Payload);
+				Result.bIsCollision = GetHeterogeneousVolumeCallableShaderIsCollision(Payload);
+				Result.HeterogeneousVolumesIndex = HeterogeneousVolumesIndex;
+			}
+		}
+#else
+	{
+		Result = HeterogeneousVolumesMajorantBasedTrackingPreshading(PathThroughput, Origin, Direction, TMin, TMax, bTransmittanceOnly, OverlappingMajorant, RandSequence);
+	}
+#endif // USE_CALLABLE_SHADERS
+
+	RandSequence = RandSequenceResult;
+	return Result;
+}
+
+float3 HeterogeneousVolumesGetTransmittancePreshading(float3 PathThroughput, float3 Origin, float3 Direction, float TMin, float TMax, inout FRandomSequence RandSequence)
+{
+	bool bTransmittanceOnly = true;
+	FMajorantData EmptyOverlappingMajorant = CreateMajorantData(0.0f);
+	float3 Throughput = 1.0;
+
+#if USE_CALLABLE_SHADERS
+	int HeterogeneousVolumesIndex = 0;
+#if USE_MULTIPLE_VOLUMES
+	for (; HeterogeneousVolumesIndex < HeterogeneousVolumesNum; ++HeterogeneousVolumesIndex)
+#endif
+		if (HeterogeneousVolumesCallableShaderIdGetTransmittance[HeterogeneousVolumesIndex])
+		{
+			FPackedPathTracingPayload Payload = (FPackedPathTracingPayload)0;
+			{
+				FRayDesc Ray;
+				Ray.Origin = Origin;
+				Ray.Direction = Direction;
+				Ray.TMin = TMin;
+				Ray.TMax = TMax;
+
+				SetHeterogeneousVolumeCallableShaderRay(Ray, Payload);
+				SetHeterogeneousVolumeCallableShaderThroughput(PathThroughput, Payload);
+				SetHeterogeneousVolumeCallableShaderRandomSequence(RandSequence, Payload);
+			}
+
+			CallShader(HeterogeneousVolumesCallableShaderIdGetTransmittance[HeterogeneousVolumesIndex], Payload);
+
+			Throughput *= GetHeterogeneousVolumeCallableShaderThroughput(Payload);
+			RandSequence = GetHeterogeneousVolumeCallableShaderRandomSequence(Payload);
+		}
+#else
+	{
+		Throughput = HeterogeneousVolumesMajorantBasedTrackingPreshading(PathThroughput, Origin, Direction, TMin, TMax, bTransmittanceOnly, EmptyOverlappingMajorant, RandSequence).Throughput;
+	}
+#endif // USE_CALLABLE_SHADERS
+
+	return Throughput;
+}
+
+#endif // USE_PRESHADED_UNIFORM_BUFFER
\ No newline at end of file
diff --git a/Engine/Shaders/Private/PathTracing/Volume/PathTracingHeterogeneousVolumesTrackingCommon.ush b/Engine/Shaders/Private/PathTracing/Volume/PathTracingHeterogeneousVolumesTrackingCommon.ush
new file mode 100644
--- /dev/null
+++ b/Engine/Shaders/Private/PathTracing/Volume/PathTracingHeterogeneousVolumesTrackingCommon.ush
@@ -0,0 +1,128 @@
+// Copyright Epic Games, Inc. All Rights Reserved.
+
+#pragma once
+
+FVolumeTrackingResult CreateVolumeTrackingResult(float Distance)
+{
+	FVolumeTrackingResult Result = (FVolumeTrackingResult)0;
+	Result.Throughput = 1;
+	Result.Pdf = 0;
+	Result.Distance = Distance;
+	Result.SigmaBar = 0;
+	Result.bIsCollision = false;
+	Result.HeterogeneousVolumesIndex = 0;
+
+	return Result;
+}
+
+FVolumeTrackingResult ConcatenateVolumeTrackingResult(FVolumeTrackingResult lhs, FVolumeTrackingResult rhs)
+{
+	FVolumeTrackingResult Result;
+	// Result for combined segment
+	Result.Throughput = lhs.Throughput * rhs.Throughput;
+	Result.Pdf = lhs.Pdf * rhs.Pdf;
+	Result.Distance = lhs.Distance + rhs.Distance;
+	// Reporting for newest segment
+	Result.SigmaBar = rhs.SigmaBar;
+	Result.bIsCollision = rhs.bIsCollision;
+	Result.HeterogeneousVolumesIndex = rhs.HeterogeneousVolumesIndex;
+
+	return Result;
+}
+
+FVolumeTrackingResult HeterogeneousVolumesDeltaTracking(float3 PathThroughput, float LocalTMin, float LocalTMax, float LocalToWorldScale, FMajorantData MajorantData, inout FRandomSequence RandSequence)
+{
+	FVolumeTrackingResult Sample = CreateVolumeTrackingResult(LocalTMin);
+	Sample.SigmaBar = MajorantData.Majorant;
+
+	float RandValue = RandSequence.Get1D();
+#define SAMPLE_ACHROMATIC 0
+#if SAMPLE_ACHROMATIC
+	float DeltaT = -log(1.0 - RandValue) / Sample.SigmaBar.x;
+#else
+	float DeltaT = SampleSpectralTransmittance(RandValue, Sample.SigmaBar, PathThroughput);
+	if (DeltaT < 0)
+	{
+		return Sample;
+	}
+#endif
+
+	Sample.Distance = min(LocalTMin + DeltaT, LocalTMax);
+	Sample.Throughput = PathThroughput;
+	Sample.bIsCollision = Sample.Distance < LocalTMax;
+	if (Sample.bIsCollision)
+	{
+#if SAMPLE_ACHROMATIC
+		Sample.Throughput /= Sample.SigmaBar.x;
+		float Transmittance = exp(-Sample.SigmaBar.x * DeltaT);
+		Sample.Pdf = Transmittance * Sample.SigmaBar.x;
+#else
+		float4 Evaluation = EvaluateSpectralTransmittanceHit(DeltaT, Sample.SigmaBar, Sample.Throughput);
+		Sample.Throughput *= Evaluation.xyz;
+		Sample.Pdf = Evaluation.w;
+	}
+	else
+	{
+		float4 Evaluation = EvaluateSpectralTransmittanceMiss(LocalTMax - LocalTMin, Sample.SigmaBar, Sample.Throughput);
+		Sample.Throughput *= Evaluation.xyz;
+		Sample.Pdf = Evaluation.w;
+#endif
+	}
+
+	Sample.Distance *= LocalToWorldScale;
+	return Sample;
+}
+
+FVolumeTrackingResult HeterogeneousVolumesRatioTracking(float3 PathThroughput, float3 TranslatedWorldOrigin, float3 WorldDirection, float LocalTMax, float LocalToWorldScale, FMajorantData MajorantData, inout FRandomSequence RandSequence)
+{
+	FVolumeTrackingResult Sample = CreateVolumeTrackingResult(0);
+	Sample.SigmaBar = MajorantData.Majorant;
+
+	while ((Sample.Distance < LocalTMax))
+	{
+		float RandValue = RandSequence.Get1D();
+#if SAMPLE_ACHROMATIC
+		float DeltaT = -log(1.0 - RandValue) / Sample.SigmaBar.x;
+#else
+		float3 CombinedThroughput = PathThroughput * Sample.Throughput;
+		float DeltaT = SampleSpectralTransmittance(RandValue, Sample.SigmaBar, CombinedThroughput);
+		if (DeltaT < 0)
+		{
+			Sample.Throughput = 0;
+			Sample.Distance = -1;
+			return Sample;
+		}
+#endif
+
+		float LocalCollisionDistance = Sample.Distance + DeltaT;
+		if (LocalCollisionDistance < LocalTMax)
+		{
+			float3 TranslatedWorldSamplePos = TranslatedWorldOrigin + WorldDirection * LocalCollisionDistance * LocalToWorldScale;
+			int HeterogeneousVolumesIndex = 0;
+			float3 SigmaT = HeterogeneousVolumesGetDensity(TranslatedWorldSamplePos, HeterogeneousVolumesIndex).SigmaT;
+			SigmaT = min(SigmaT, Sample.SigmaBar);
+
+			float3 SpectralNullCollisionWeight = max(Sample.SigmaBar - SigmaT, 0);
+#if SAMPLE_ACHROMATIC
+			Sample.Throughput *= SpectralNullCollisionWeight / Sample.SigmaBar.x;
+			float Transmittance = exp(-Sample.SigmaBar.x * DeltaT);
+			Sample.Pdf *= Transmittance * Sample.SigmaBar.x;
+#else
+			float4 EvalResult = EvaluateSpectralTransmittanceHit(DeltaT, Sample.SigmaBar, CombinedThroughput);
+			Sample.Throughput *= SpectralNullCollisionWeight * EvalResult.xyz;
+			Sample.Pdf *= EvalResult.w;
+		}
+		else
+		{
+			float4 EvalResult = EvaluateSpectralTransmittanceMiss(LocalTMax - Sample.Distance, Sample.SigmaBar, CombinedThroughput);
+			Sample.Throughput *= EvalResult.xyz;
+			Sample.Pdf *= EvalResult.w;
+#endif
+		}
+
+		Sample.Distance = min(LocalCollisionDistance, LocalTMax);
+	}
+
+	Sample.Distance *= LocalToWorldScale;
+	return Sample;
+}
\ No newline at end of file
diff --git a/Engine/Shaders/Private/PathTracing/Volume/PathTracingVolume.ush b/Engine/Shaders/Private/PathTracing/Volume/PathTracingVolume.ush
--- a/Engine/Shaders/Private/PathTracing/Volume/PathTracingVolume.ush
+++ b/Engine/Shaders/Private/PathTracing/Volume/PathTracingVolume.ush
@@ -68,7 +68,7 @@ FVolumeDensityBounds VolumeGetDensityBounds(float3 Origin, float3 Direction, FVo
 }
 
 // Given a point in world space, return the amount of volume and its scattering properties
-FVolumeShadedResult VolumeGetDensity(float3 TranslatedWorldPos, uint VolumeMask, bool bIsCameraRay)
+FVolumeShadedResult VolumeGetDensity(float3 TranslatedWorldPos, int HeterogeneousVolumesIndex, uint VolumeMask, bool bIsCameraRay)
 {
 	FVolumeShadedResult Result = (FVolumeShadedResult)0;
 	if (VolumeMask & PATH_TRACER_VOLUME_ENABLE_ATMOSPHERE)
@@ -81,7 +81,7 @@ FVolumeShadedResult VolumeGetDensity(float3 TranslatedWorldPos, uint VolumeMask,
 	}
 	if (VolumeMask & PATH_TRACER_VOLUME_ENABLE_HETEROGENEOUS_VOLUMES)
 	{
-		MergeVolumeShadedResult(Result, HeterogeneousVolumesGetDensity(TranslatedWorldPos), bIsCameraRay && (VolumeFlags & PATH_TRACER_VOLUME_HOLDOUT_HETEROGENEOUS_VOLUMES) != 0);
+		MergeVolumeShadedResult(Result, HeterogeneousVolumesGetDensity(TranslatedWorldPos, HeterogeneousVolumesIndex), bIsCameraRay && (VolumeFlags & PATH_TRACER_VOLUME_HOLDOUT_HETEROGENEOUS_VOLUMES) != 0);
 	}
 	return Result;
 }
@@ -275,7 +275,8 @@ FVolumeShadedResult VolumeGetDensityAndTransmittance(float3 Origin, float3 Direc
 #endif
 	{
 		// non-cloud case, just get density of all volumes
-		Result = VolumeGetDensity(Origin + Direction * Interval.VolumeTMax, Interval.VolumeMask, Bounce == 0);
+		int HeterogeneousVolumesIndex = -1;
+		Result = VolumeGetDensity(Origin + Direction * Interval.VolumeTMax, HeterogeneousVolumesIndex, Interval.VolumeMask, Bounce == 0);
 		CloudFactor = 1.0;
 		// Factor in SigmaS to throughput
 		Throughput *= Result.SigmaSRayleigh + Result.SigmaSHG + Result.SigmaSDualHG;
@@ -304,7 +305,7 @@ FVolumeTrackingResult VolumeSampleDistance(float3 PathThroughput, float3 Origin,
 	if (Interval.VolumeMask & PATH_TRACER_VOLUME_ENABLE_HETEROGENEOUS_VOLUMES)
 	{
 		// Interval is crossing a heterogeneous volume
-		FMajorantData OverlappingMajorant = CreateMajorantData(VolumeDensityBounds.SigmaMin, SigmaBar);
+		FMajorantData OverlappingMajorant = CreateMajorantData(SigmaBar);
 		TrackingResult = HeterogeneousVolumesSampleDistance(PathThroughput, Origin, Direction, Interval.VolumeTMin, Interval.VolumeTMax, OverlappingMajorant, RandSequence);
 	}
 	else
diff --git a/Engine/Source/Runtime/Renderer/Private/DeferredShadingRenderer.cpp b/Engine/Source/Runtime/Renderer/Private/DeferredShadingRenderer.cpp
--- a/Engine/Source/Runtime/Renderer/Private/DeferredShadingRenderer.cpp
+++ b/Engine/Source/Runtime/Renderer/Private/DeferredShadingRenderer.cpp
@@ -2882,6 +2882,12 @@ void FDeferredShadingSceneRenderer::Render(FRDGBuilder& GraphBuilder, const FSce
 		// Async AS builds can potentially overlap with BasePass.
 		bool bNeedToSetupRayTracingRenderingData = DispatchRayTracingWorldUpdates(GraphBuilder, DynamicGeometryScratchBuffer, DynamicRTResourceAccessPipelines);
 
+		// Must be prepared before the SBT is built
+		if (ViewFamily.EngineShowFlags.PathTracing)
+		{
+			PreparePathTracingHeterogeneousVolumes(GraphBuilder, Scene, Views);
+		}
+
 		/** Should be called somewhere before "SetupRayTracingRenderingData" */
 		SetupRayTracingLightDataForViews(GraphBuilder);
 #endif
diff --git a/Engine/Source/Runtime/Renderer/Private/HeterogeneousVolumes/HeterogeneousVolumes.h b/Engine/Source/Runtime/Renderer/Private/HeterogeneousVolumes/HeterogeneousVolumes.h
--- a/Engine/Source/Runtime/Renderer/Private/HeterogeneousVolumes/HeterogeneousVolumes.h
+++ b/Engine/Source/Runtime/Renderer/Private/HeterogeneousVolumes/HeterogeneousVolumes.h
@@ -280,6 +280,10 @@ BEGIN_UNIFORM_BUFFER_STRUCT(FPreshadingUniformBufferParameters, )
 	SHADER_PARAMETER_RDG_TEXTURE(Texture3D, EmissionTexture)
 	SHADER_PARAMETER_RDG_TEXTURE(Texture3D, AlbedoTexture)
 	SHADER_PARAMETER_SAMPLER(SamplerState, TextureSampler)
+
+	SHADER_PARAMETER(FIntVector, MajorantResolution)
+	SHADER_PARAMETER_RDG_TEXTURE(Texture3D, MajorantTexture)
+	SHADER_PARAMETER_SAMPLER(SamplerState, PointTextureSampler)
 END_UNIFORM_BUFFER_STRUCT()
 
 BEGIN_SHADER_PARAMETER_STRUCT(FLightingCacheParameters, )
@@ -288,6 +292,43 @@ BEGIN_SHADER_PARAMETER_STRUCT(FLightingCacheParameters, )
 	SHADER_PARAMETER_RDG_TEXTURE(Texture3D, LightingCacheTexture)
 END_SHADER_PARAMETER_STRUCT()
 
+// Path tracing structures
+
+struct FHeterogeneousVolumesPrimitiveInstanceKey
+{
+	uint32 PersistentPrimitiveIndex;
+	uint32 InstanceIndex;
+
+	friend FORCEINLINE bool operator == (
+		const FHeterogeneousVolumesPrimitiveInstanceKey& A,
+		const FHeterogeneousVolumesPrimitiveInstanceKey& B
+	)
+	{
+		return A.PersistentPrimitiveIndex == B.PersistentPrimitiveIndex && A.InstanceIndex == B.InstanceIndex;
+	}
+
+	friend FORCEINLINE uint32 GetTypeHash(const FHeterogeneousVolumesPrimitiveInstanceKey& Key)
+	{
+		return HashCombine(GetTypeHash(Key.PersistentPrimitiveIndex), GetTypeHash(Key.InstanceIndex));
+	}
+};
+
+struct FPathTracingHeterogeneousVolumesCallableInfo
+{
+	// Preshading data
+	TRefCountPtr<IPooledRenderTarget> MajorantRT;
+	TRefCountPtr<IPooledRenderTarget> ExtinctionRT;
+	TRefCountPtr<IPooledRenderTarget> EmissionRT;
+	TRefCountPtr<IPooledRenderTarget> AlbedoRT;
+
+	// Callable interface
+	int32 IntersectShaderId;
+	int32 GetDensityBoundsShaderId;
+	int32 GetDensityShaderId;
+	int32 SampleDistanceShaderId;
+	int32 GetTransmittanceShaderId;
+};
+
 // Adaptive Voxel Grid structures
 
 struct FTopLevelGridBitmaskData
@@ -401,6 +442,40 @@ void BuildFrustumVoxelGrid(
 	TRDGUniformBufferRef<FFrustumVoxelGridUniformBufferParameters>& FrustumVoxelGridUniformBuffer
 );
 
+void CollectRenderableHeterogeneousVolumes(
+	FRDGBuilder& GraphBuilder,
+	const FScene* Scene,
+	const FViewInfo& View,
+	TArray<const IHeterogeneousVolumeInterface*>& HeterogeneousVolumeInterfaces,
+	TArray<const FMaterialRenderProxy*>& MaterialRenderProxies
+);
+
+void BuildPreshadingGrid(
+	FRDGBuilder& GraphBuilder,
+	const FScene* Scene,
+	const FViewInfo& View,
+	const IHeterogeneousVolumeInterface* HeterogeneousVolumeInterface,
+	const FMaterialRenderProxy* MaterialRenderProxy,
+	TRefCountPtr<IPooledRenderTarget>& PersistentMajorantRT,
+	TRDGUniformBufferRef<FPreshadingUniformBufferParameters>& PreshadingUniformBuffer
+);
+
+void BuildPreshadingGrid(
+	FRDGBuilder& GraphBuilder,
+	const FScene* Scene,
+	const FViewInfo& View,
+	const IHeterogeneousVolumeInterface* HeterogeneousVolumeInterface,
+	const FMaterialRenderProxy* MaterialRenderProxy,
+	TRefCountPtr<IPooledRenderTarget>& PersistentMajorantRT,
+	FPathTracingHeterogeneousVolumesCallableInfo& CallableInfo,
+	FPreshadingUniformBufferParameters& PreshadingUniformBufferParameters
+);
+
+void BuildEmptyPreshadingUniformBuffer(
+	FRDGBuilder& GraphBuilder,
+	TRDGUniformBufferRef<FPreshadingUniformBufferParameters>& PreshadingUniformBuffer
+);
+
 struct FAdaptiveFrustumGridParameterCache
 {
 	FMatrix44f WorldToClip;
@@ -891,12 +966,40 @@ void ComputeHeterogeneousVolumeBakeMaterial(
 	const FMaterialRenderProxy* MaterialRenderProxy,
 	const FPersistentPrimitiveIndex &PersistentPrimitiveIndex,
 	const FBoxSphereBounds LocalBoxSphereBounds,
+	// Sample data
+	bool bJitter,
+	bool bLooseBounds,
+	bool bEvaluateAlbedo,
+	bool bEvaluateEmission,
+	// Volume data
+	FIntVector VolumeResolution,
+	// Output
+	FRDGTextureRef HeterogeneousVolumeExtinctionTexture,
+	FRDGTextureRef HeterogeneousVolumeEmissionTexture,
+	FRDGTextureRef HeterogeneousVolumeAlbedoTexture
+);
+
+void ComputeHeterogeneousVolumeBakeMaterial(
+	FRDGBuilder& GraphBuilder,
+	// Scene data
+	const FScene* Scene,
+	const FViewInfo& View,
+	// Object data
+	const IHeterogeneousVolumeInterface* HeterogeneousVolumeInterface,
+	const FMaterialRenderProxy* MaterialRenderProxy,
+	const FPersistentPrimitiveIndex& PersistentPrimitiveIndex,
+	const FBoxSphereBounds LocalBoxSphereBounds,
+	// Sample data
+	bool bJitter,
+	bool bLooseBounds,
+	bool bEvaluateAlbedo,
+	bool bEvaluateEmission,
 	// Volume data
 	FIntVector VolumeResolution,
 	// Output
-	FRDGTextureRef& HeterogeneousVolumeExtinctionTexture,
-	FRDGTextureRef& HeterogeneousVolumeEmissionTexture,
-	FRDGTextureRef& HeterogeneousVolumeAlbedoTexture
+	TRefCountPtr<IPooledRenderTarget>& HeterogeneousVolumeExtinctionTexture,
+	TRefCountPtr<IPooledRenderTarget>& HeterogeneousVolumeEmissionTexture,
+	TRefCountPtr<IPooledRenderTarget>& HeterogeneousVolumeAlbedoTexture
 );
 
 // Sparse Voxel Pipeline
diff --git a/Engine/Source/Runtime/Renderer/Private/HeterogeneousVolumes/HeterogeneousVolumesMaterialBakingPipeline.cpp b/Engine/Source/Runtime/Renderer/Private/HeterogeneousVolumes/HeterogeneousVolumesMaterialBakingPipeline.cpp
--- a/Engine/Source/Runtime/Renderer/Private/HeterogeneousVolumes/HeterogeneousVolumesMaterialBakingPipeline.cpp
+++ b/Engine/Source/Runtime/Renderer/Private/HeterogeneousVolumes/HeterogeneousVolumesMaterialBakingPipeline.cpp
@@ -35,14 +35,17 @@ class FHeterogeneousVolumesBakeMaterialCS : public FMeshMaterialShader
 
 		// Sample data
 		SHADER_PARAMETER(int, bJitter)
+		SHADER_PARAMETER(int, bEvaluateAlbedo)
+		SHADER_PARAMETER(int, bEvaluateEmission)
+		SHADER_PARAMETER_STRUCT_REF(FBlueNoise, BlueNoise)
 
 		// Dispatch data
 		SHADER_PARAMETER(FIntVector, GroupCount)
 
 		// Output
-		SHADER_PARAMETER_RDG_TEXTURE_UAV(RWTexture2D<float3>, RWExtinctionTexture)
-		SHADER_PARAMETER_RDG_TEXTURE_UAV(RWTexture2D<float3>, RWEmissionTexture)
-		SHADER_PARAMETER_RDG_TEXTURE_UAV(RWTexture2D<float3>, RWAlbedoTexture)
+		SHADER_PARAMETER_RDG_TEXTURE_UAV(RWTexture3D<float3>, RWExtinctionTexture)
+		SHADER_PARAMETER_RDG_TEXTURE_UAV(RWTexture3D<float3>, RWEmissionTexture)
+		SHADER_PARAMETER_RDG_TEXTURE_UAV(RWTexture3D<float3>, RWAlbedoTexture)
 		//SHADER_PARAMETER_RDG_BUFFER_UAV(RWStructuredBuffer<Volumes::FDebugOutput>, RWDebugOutputBuffer)
 	END_SHADER_PARAMETER_STRUCT()
 
@@ -108,14 +111,19 @@ void ComputeHeterogeneousVolumeBakeMaterial(
 	// Object data
 	const IHeterogeneousVolumeInterface* HeterogeneousVolumeInterface,
 	const FMaterialRenderProxy* MaterialRenderProxy,
-	const FPersistentPrimitiveIndex &PersistentPrimitiveIndex,
+	const FPersistentPrimitiveIndex& PersistentPrimitiveIndex,
 	const FBoxSphereBounds LocalBoxSphereBounds,
+	// Sample data
+	bool bJitter,
+	bool bLooseBounds,
+	bool bEvaluateAlbedo,
+	bool bEvaluateEmission,
 	// Volume data
 	FIntVector VolumeResolution,
 	// Output
-	FRDGTextureRef& HeterogeneousVolumeExtinctionTexture,
-	FRDGTextureRef& HeterogeneousVolumeEmissionTexture,
-	FRDGTextureRef& HeterogeneousVolumeAlbedoTexture
+	FRDGTextureRef HeterogeneousVolumeExtinctionTexture,
+	FRDGTextureRef HeterogeneousVolumeEmissionTexture,
+	FRDGTextureRef HeterogeneousVolumeAlbedoTexture
 )
 {
 	const FMaterial& Material = MaterialRenderProxy->GetMaterialWithFallback(View.GetFeatureLevel(), MaterialRenderProxy);
@@ -158,7 +166,11 @@ void ComputeHeterogeneousVolumeBakeMaterial(
 		PassParameters->VolumeResolution = VolumeResolution;
 
 		// Sample data
-		PassParameters->bJitter = HeterogeneousVolumes::ShouldJitter();
+		PassParameters->bJitter = bJitter;
+		FBlueNoise BlueNoise = GetBlueNoiseGlobalParameters();
+		PassParameters->BlueNoise = CreateUniformBufferImmediate(BlueNoise, EUniformBufferUsage::UniformBuffer_SingleDraw);
+		PassParameters->bEvaluateAlbedo = bEvaluateAlbedo;
+		PassParameters->bEvaluateEmission = bEvaluateEmission;
 
 		// Dispatch data
 		PassParameters->GroupCount = GroupCount;
@@ -188,3 +200,100 @@ void ComputeHeterogeneousVolumeBakeMaterial(
 		}
 	);
 }
+
+void ComputeHeterogeneousVolumeBakeMaterial(
+	FRDGBuilder& GraphBuilder,
+	// Scene data
+	const FScene* Scene,
+	const FViewInfo& View,
+	// Object data
+	const IHeterogeneousVolumeInterface* HeterogeneousVolumeInterface,
+	const FMaterialRenderProxy* MaterialRenderProxy,
+	const FPersistentPrimitiveIndex& PersistentPrimitiveIndex,
+	const FBoxSphereBounds LocalBoxSphereBounds,
+	// Sample data
+	bool bJitter,
+	bool bLooseBounds,
+	bool bEvaluateAlbedo,
+	bool bEvaluateEmission,
+	// Volume data
+	FIntVector VolumeResolution,
+	// Output
+	TRefCountPtr<IPooledRenderTarget>& HeterogeneousVolumeExtinctionTexture,
+	TRefCountPtr<IPooledRenderTarget>& HeterogeneousVolumeEmissionTexture,
+	TRefCountPtr<IPooledRenderTarget>& HeterogeneousVolumeAlbedoTexture
+)
+{
+	const FMaterial& Material = MaterialRenderProxy->GetMaterialWithFallback(View.GetFeatureLevel(), MaterialRenderProxy);
+	check(Material.GetMaterialDomain() == MD_Volume);
+
+	// FMath::DivideAndRoundUp() is not compatible with FIntVector?
+	FIntVector GroupSize = FIntVector(FHeterogeneousVolumesBakeMaterialCS::GetThreadGroupSize3D());
+	FIntVector GroupCount = (VolumeResolution + GroupSize - FIntVector(1)) / FHeterogeneousVolumesBakeMaterialCS::GetThreadGroupSize3D();
+
+#if 0
+	FRDGBufferRef DebugOutputBuffer = GraphBuilder.CreateBuffer(
+		FRDGBufferDesc::CreateStructuredDesc(sizeof(Volumes::FDebugOutput), GetVoxelCount(VolumeResolution)),
+		TEXT("HeterogeneousVolumes.DebugOutputBuffer"));
+	AddClearUAVPass(GraphBuilder, GraphBuilder.CreateUAV(DebugOutputBuffer), 0);
+#endif
+
+	FHeterogeneousVolumesBakeMaterialCS::FParameters* PassParameters = GraphBuilder.AllocParameters<FHeterogeneousVolumesBakeMaterialCS::FParameters>();
+	{
+		// Scene data
+		PassParameters->View = View.ViewUniformBuffer;
+		PassParameters->Scene = View.GetSceneUniforms().GetBuffer(GraphBuilder);
+		PassParameters->EyeAdaptationBuffer = GraphBuilder.CreateSRV(GetEyeAdaptationBuffer(GraphBuilder, View));
+
+		// Object data
+		// TODO: Convert to relative-local space
+		//FVector3f ViewOriginHigh = FDFVector3(View.ViewMatrices.GetViewOrigin()).High;
+		//FMatrix44f RelativeLocalToWorld = FDFMatrix::MakeToRelativeWorldMatrix(ViewOriginHigh, HeterogeneousVolumeInterface->GetLocalToWorld()).M;
+		FMatrix InstanceToLocal = HeterogeneousVolumeInterface->GetInstanceToLocal();
+		FMatrix LocalToWorld = HeterogeneousVolumeInterface->GetLocalToWorld();
+		PassParameters->LocalToWorld = FMatrix44f(InstanceToLocal * LocalToWorld);
+		PassParameters->WorldToLocal = PassParameters->LocalToWorld.Inverse();
+
+		FMatrix LocalToInstance = InstanceToLocal.Inverse();
+		FBoxSphereBounds InstanceBoxSphereBounds = LocalBoxSphereBounds.TransformBy(LocalToInstance);
+		PassParameters->LocalBoundsOrigin = FVector3f(InstanceBoxSphereBounds.Origin);
+		PassParameters->LocalBoundsExtent = FVector3f(InstanceBoxSphereBounds.BoxExtent);
+		PassParameters->PrimitiveId = PersistentPrimitiveIndex.Index;
+
+		// Volume data
+		PassParameters->VolumeResolution = VolumeResolution;
+
+		// Sample data
+		PassParameters->bJitter = bJitter;
+		FBlueNoise BlueNoise = GetBlueNoiseGlobalParameters();
+		PassParameters->BlueNoise = CreateUniformBufferImmediate(BlueNoise, EUniformBufferUsage::UniformBuffer_SingleDraw);
+		PassParameters->bEvaluateAlbedo = bEvaluateAlbedo;
+		PassParameters->bEvaluateEmission = bEvaluateEmission;
+
+		// Dispatch data
+		PassParameters->GroupCount = GroupCount;
+
+		PassParameters->RWExtinctionTexture = GraphBuilder.CreateUAV(GraphBuilder.RegisterExternalTexture(HeterogeneousVolumeExtinctionTexture));
+		PassParameters->RWEmissionTexture = GraphBuilder.CreateUAV(GraphBuilder.RegisterExternalTexture(HeterogeneousVolumeEmissionTexture));
+		PassParameters->RWAlbedoTexture = GraphBuilder.CreateUAV(GraphBuilder.RegisterExternalTexture(HeterogeneousVolumeAlbedoTexture));
+	}
+
+	GraphBuilder.AddPass(
+		RDG_EVENT_NAME("HeterogeneousVolumesBakedMaterialCS"),
+		PassParameters,
+		ERDGPassFlags::Compute,
+		[PassParameters, Scene, &View, MaterialRenderProxy, &Material, GroupCount](FRDGAsyncTask, FRHIComputeCommandList& RHICmdList)
+		{
+			FHeterogeneousVolumesBakeMaterialCS::FPermutationDomain PermutationVector;
+			TShaderRef<FHeterogeneousVolumesBakeMaterialCS> ComputeShader = Material.GetShader<FHeterogeneousVolumesBakeMaterialCS>(&FLocalVertexFactory::StaticType, PermutationVector, false);
+
+			if (!ComputeShader.IsNull())
+			{
+				FMeshDrawShaderBindings ShaderBindings;
+				UE::MeshPassUtils::SetupComputeBindings(ComputeShader, Scene, Scene->GetFeatureLevel(), nullptr, *MaterialRenderProxy, Material, ShaderBindings);
+
+				UE::MeshPassUtils::Dispatch(RHICmdList, ComputeShader, ShaderBindings, *PassParameters, GroupCount);
+			}
+		}
+	);
+}
diff --git a/Engine/Source/Runtime/Renderer/Private/HeterogeneousVolumes/HeterogeneousVolumesPreshadingPipeline.cpp b/Engine/Source/Runtime/Renderer/Private/HeterogeneousVolumes/HeterogeneousVolumesPreshadingPipeline.cpp
--- a/Engine/Source/Runtime/Renderer/Private/HeterogeneousVolumes/HeterogeneousVolumesPreshadingPipeline.cpp
+++ b/Engine/Source/Runtime/Renderer/Private/HeterogeneousVolumes/HeterogeneousVolumesPreshadingPipeline.cpp
@@ -22,7 +22,8 @@
 	TEXT("Determines the pixel format for preshaded results (Default = 2)\n")
 	TEXT("0: PF_R16F\n")
 	TEXT("1: PF_FloatR11G11B10\n")
-	TEXT("2: PF_FloatRGB\n"),
+	TEXT("2: PF_FloatRGB\n")
+	TEXT("3: PF_FloatRGBA\n"),
 	ECVF_RenderThreadSafe
 );
 
@@ -39,6 +40,8 @@
 		case 2:
 		default:
 			return PF_FloatRGB;
+		case 3:
+			return PF_FloatRGBA;
 		}
 	}
 
@@ -303,6 +306,71 @@
 
 IMPLEMENT_GLOBAL_SHADER(FRenderSingleScatteringWithPreshadingCS, "/Engine/Private/HeterogeneousVolumes/HeterogeneousVolumesPreshadingPipeline.usf", "RenderSingleScatteringWithPreshadingCS", SF_Compute);
 
+class FMaxTexture3D : public FGlobalShader
+{
+	DECLARE_GLOBAL_SHADER(FMaxTexture3D);
+	SHADER_USE_PARAMETER_STRUCT(FMaxTexture3D, FGlobalShader);
 
+	static bool ShouldCompilePermutation(const FGlobalShaderPermutationParameters& Parameters)
+	{
+		return DoesPlatformSupportHeterogeneousVolumes(Parameters.Platform);
+	}
+
+	static void ModifyCompilationEnvironment(
+		const FGlobalShaderPermutationParameters& Parameters,
+		FShaderCompilerEnvironment& OutEnvironment
+	)
+	{
+		OutEnvironment.SetDefine(TEXT("THREADGROUP_SIZE_3D"), GetThreadGroupSize3D());
+
+		OutEnvironment.CompilerFlags.Add(CFLAG_AllowTypedUAVLoads);
+	}
+
+	BEGIN_SHADER_PARAMETER_STRUCT(FParameters, )
+		// Input
+		SHADER_PARAMETER(FIntVector, TextureResolution)
+		SHADER_PARAMETER_RDG_TEXTURE_SRV(Texture3D, InputTextureLHS)
+		SHADER_PARAMETER_RDG_TEXTURE_SRV(Texture3D, InputTextureRHS)
+		// Output
+		SHADER_PARAMETER_RDG_TEXTURE_UAV(RWTexture3D<float3>, RWOutputTexture)
+	END_SHADER_PARAMETER_STRUCT()
+
+	static int32 GetThreadGroupSize3D() { return 4; }
+};
+
+IMPLEMENT_GLOBAL_SHADER(FMaxTexture3D, "/Engine/Private/HeterogeneousVolumes/HeterogeneousVolumesPreshadingPipeline.usf", "MaxTexture3D", SF_Compute);
+
+void MaxTexture3D(
+	FRDGBuilder& GraphBuilder,
+	const FViewInfo& View,
+	FRDGTextureRef InputTextureLHS,
+	FRDGTextureRef InputTextureRHS,
+	FRDGTextureRef OutputTexture
+)
+{
+	FRDGTextureDesc OutputTextureDesc = OutputTexture->Desc;
+	const FIntVector TextureResolution(OutputTextureDesc.Extent.X, OutputTextureDesc.Extent.Y, OutputTextureDesc.Depth);
+	FMaxTexture3D::FParameters* PassParameters = GraphBuilder.AllocParameters<FMaxTexture3D::FParameters>();
+	{
+		PassParameters->TextureResolution = TextureResolution;
+		PassParameters->InputTextureLHS = GraphBuilder.CreateSRV(InputTextureLHS);
+		PassParameters->InputTextureRHS = GraphBuilder.CreateSRV(InputTextureRHS);
+		PassParameters->RWOutputTexture = GraphBuilder.CreateUAV(OutputTexture);
+	}
+
+	uint32 GroupCountX = FMath::DivideAndRoundUp(TextureResolution.X, FMaxTexture3D::GetThreadGroupSize3D());
+	uint32 GroupCountY = FMath::DivideAndRoundUp(TextureResolution.Y, FMaxTexture3D::GetThreadGroupSize3D());
+	uint32 GroupCountZ = FMath::DivideAndRoundUp(TextureResolution.Z, FMaxTexture3D::GetThreadGroupSize3D());
+	FIntVector GroupCount(GroupCountX, GroupCountY, GroupCountZ);
+
+	TShaderRef<FMaxTexture3D> ComputeShader = View.ShaderMap->GetShader<FMaxTexture3D>();
+	FComputeShaderUtils::AddPass(
+		GraphBuilder,
+		RDG_EVENT_NAME("MaxTexture3D"),
+		ComputeShader,
+		PassParameters,
+		GroupCount);
+}
+
 void RenderLightingCacheWithPreshadingCompute(
 	FRDGBuilder& GraphBuilder,
@@ -1250,28 +1318,359 @@
 	ECVF_RenderThreadSafe
 );
 
-void RenderWithPreshading(
+void CollectRenderableHeterogeneousVolumes(
 	FRDGBuilder& GraphBuilder,
-	// Scene data
-	const FSceneTextures& SceneTextures,
-	FScene* Scene,
-	FViewInfo& View, int32 ViewIndex,
-	// Shadow data
-	TArray<FVisibleLightInfo, SceneRenderingAllocator>& VisibleLightInfos,
-	const FVirtualShadowMapArray& VirtualShadowMapArray,
+	const FScene* Scene,
+	const FViewInfo& View,
+	TArray<const IHeterogeneousVolumeInterface*>& HeterogeneousVolumeInterfaces,
+	TArray<const FMaterialRenderProxy*>& MaterialRenderProxies
+)
+{
+	for (int32 MeshBatchIndex = 0; MeshBatchIndex < View.HeterogeneousVolumesMeshBatches.Num(); ++MeshBatchIndex)
+	{
+		const FMeshBatch* Mesh = View.HeterogeneousVolumesMeshBatches[MeshBatchIndex].Mesh;
+		const FPrimitiveSceneProxy* PrimitiveSceneProxy = View.HeterogeneousVolumesMeshBatches[MeshBatchIndex].Proxy;
+		const FMaterialRenderProxy* MaterialRenderProxy = Mesh->MaterialRenderProxy;
+		if (!ShouldRenderMeshBatchWithHeterogeneousVolumes(Mesh, PrimitiveSceneProxy, View.GetFeatureLevel()))
+		{
+			continue;
+		}
+
+		for (int32 VolumeIndex = 0; VolumeIndex < Mesh->Elements.Num(); ++VolumeIndex)
+		{
+			const IHeterogeneousVolumeInterface* HeterogeneousVolumeInterface = (IHeterogeneousVolumeInterface*)Mesh->Elements[VolumeIndex].UserData;
+			//check(HeterogeneousVolumeInterface != nullptr);
+			if (HeterogeneousVolumeInterface == nullptr)
+			{
+				continue;
+			}
+
+			// Append
+			HeterogeneousVolumeInterfaces.Add(HeterogeneousVolumeInterface);
+			MaterialRenderProxies.Add(MaterialRenderProxy);
+		}
+	}
+}
+
+void BuildPreshadingUniformBufferParameters(
+	FRDGBuilder& GraphBuilder,
+	const FScene* Scene,
+	const FViewInfo& View,
+	const IHeterogeneousVolumeInterface* HeterogeneousVolumeInterface,
+	const FMaterialRenderProxy* MaterialRenderProxy,
+	// Volume data
+	const FIntVector& BakingResolution,
+	TRefCountPtr<IPooledRenderTarget> ExtinctionTexture,
+	TRefCountPtr<IPooledRenderTarget> EmissionTexture,
+	TRefCountPtr<IPooledRenderTarget> AlbedoTexture,
+	const FIntVector& MajorantResolution,
+	TRefCountPtr<IPooledRenderTarget> MajorantTexture,
+	FPreshadingUniformBufferParameters& PreshadingUniformBufferParameters
+)
+{
 	// Object data
+	// TODO: Convert to relative-local space
+	//FVector3f ViewOriginHigh = FDFVector3(View.ViewMatrices.GetViewOrigin()).High;
+	//FMatrix44f RelativeLocalToWorld = FDFMatrix::MakeToRelativeWorldMatrix(ViewOriginHigh, HeterogeneousVolumeInterface->GetLocalToWorld()).M;
+	FMatrix InstanceToLocal = HeterogeneousVolumeInterface->GetInstanceToLocal();
+	FMatrix LocalToWorld = HeterogeneousVolumeInterface->GetLocalToWorld();
+	PreshadingUniformBufferParameters.LocalToWorld = FMatrix44f(InstanceToLocal * LocalToWorld);
+	PreshadingUniformBufferParameters.WorldToLocal = PreshadingUniformBufferParameters.LocalToWorld.Inverse();
+
+	FMatrix LocalToInstance = InstanceToLocal.Inverse();
+	const FBoxSphereBounds LocalBoxSphereBounds = HeterogeneousVolumeInterface->GetLocalBounds();
+	FBoxSphereBounds InstanceBoxSphereBounds = LocalBoxSphereBounds.TransformBy(LocalToInstance);
+	PreshadingUniformBufferParameters.LocalBoundsOrigin = FVector3f(InstanceBoxSphereBounds.Origin);
+	PreshadingUniformBufferParameters.LocalBoundsExtent = FVector3f(InstanceBoxSphereBounds.BoxExtent);
+
+	// Volume data
+	PreshadingUniformBufferParameters.VolumeResolution = BakingResolution;
+	PreshadingUniformBufferParameters.ExtinctionTexture = GraphBuilder.RegisterExternalTexture(ExtinctionTexture);
+	PreshadingUniformBufferParameters.EmissionTexture = GraphBuilder.RegisterExternalTexture(EmissionTexture);
+	PreshadingUniformBufferParameters.AlbedoTexture = GraphBuilder.RegisterExternalTexture(AlbedoTexture);
+	PreshadingUniformBufferParameters.TextureSampler = TStaticSamplerState<SF_Trilinear, AM_Clamp, AM_Clamp, AM_Clamp>::GetRHI();
+
+	PreshadingUniformBufferParameters.MajorantResolution = MajorantResolution;
+	PreshadingUniformBufferParameters.MajorantTexture = GraphBuilder.RegisterExternalTexture(MajorantTexture);
+	PreshadingUniformBufferParameters.PointTextureSampler = TStaticSamplerState<SF_Point, AM_Clamp, AM_Clamp, AM_Clamp>::GetRHI();
+}
+
+void BuildPreshadingUniformBuffer(
+	FRDGBuilder& GraphBuilder,
+	const FScene* Scene,
+	const FViewInfo& View,
 	const IHeterogeneousVolumeInterface* HeterogeneousVolumeInterface,
 	const FMaterialRenderProxy* MaterialRenderProxy,
-	const FPersistentPrimitiveIndex& PersistentPrimitiveIndex,
-	const FBoxSphereBounds LocalBoxSphereBounds,
-	// Transmittance acceleration
-	FRDGTextureRef LightingCacheTexture,
-	// Output
-	FRDGTextureRef& HeterogeneousVolumeRadiance
+	// Volume data
+	const FIntVector& BakingResolution,
+	FRDGTextureRef ExtinctionTexture,
+	FRDGTextureRef EmissionTexture,
+	FRDGTextureRef AlbedoTexture,
+	const FIntVector& MajorantResolution,
+	FRDGTextureRef MajorantTexture,
+	TRDGUniformBufferRef<FPreshadingUniformBufferParameters>& PreshadingUniformBuffer
 )
 {
-	CSV_SCOPED_TIMING_STAT_EXCLUSIVE(Effects);
+	// Create Preshading UniformBuffer
+	FPreshadingUniformBufferParameters* PreshadingUniformBufferParameters = GraphBuilder.AllocParameters<FPreshadingUniformBufferParameters>();
+	{
+		// Object data
+		// TODO: Convert to relative-local space
+		//FVector3f ViewOriginHigh = FDFVector3(View.ViewMatrices.GetViewOrigin()).High;
+		//FMatrix44f RelativeLocalToWorld = FDFMatrix::MakeToRelativeWorldMatrix(ViewOriginHigh, HeterogeneousVolumeInterface->GetLocalToWorld()).M;
+		FMatrix InstanceToLocal = HeterogeneousVolumeInterface->GetInstanceToLocal();
+		FMatrix LocalToWorld = HeterogeneousVolumeInterface->GetLocalToWorld();
+		PreshadingUniformBufferParameters->LocalToWorld = FMatrix44f(InstanceToLocal * LocalToWorld);
+		PreshadingUniformBufferParameters->WorldToLocal = PreshadingUniformBufferParameters->LocalToWorld.Inverse();
 
+		FMatrix LocalToInstance = InstanceToLocal.Inverse();
+		const FBoxSphereBounds LocalBoxSphereBounds = HeterogeneousVolumeInterface->GetLocalBounds();
+		FBoxSphereBounds InstanceBoxSphereBounds = LocalBoxSphereBounds.TransformBy(LocalToInstance);
+		PreshadingUniformBufferParameters->LocalBoundsOrigin = FVector3f(InstanceBoxSphereBounds.Origin);
+		PreshadingUniformBufferParameters->LocalBoundsExtent = FVector3f(InstanceBoxSphereBounds.BoxExtent);
+
+		// Volume data
+		PreshadingUniformBufferParameters->VolumeResolution = BakingResolution;
+		PreshadingUniformBufferParameters->ExtinctionTexture = ExtinctionTexture;
+		PreshadingUniformBufferParameters->EmissionTexture = EmissionTexture;
+		PreshadingUniformBufferParameters->AlbedoTexture = AlbedoTexture;
+		PreshadingUniformBufferParameters->TextureSampler = TStaticSamplerState<SF_Trilinear, AM_Clamp, AM_Clamp, AM_Clamp>::GetRHI();
+
+		PreshadingUniformBufferParameters->MajorantResolution = MajorantResolution;
+		PreshadingUniformBufferParameters->MajorantTexture = MajorantTexture;
+		PreshadingUniformBufferParameters->PointTextureSampler = TStaticSamplerState<SF_Point, AM_Clamp, AM_Clamp, AM_Clamp>::GetRHI();
+	}
+	PreshadingUniformBuffer = GraphBuilder.CreateUniformBuffer(PreshadingUniformBufferParameters);
+}
+
+void BuildEmptyPreshadingUniformBuffer(
+	FRDGBuilder& GraphBuilder,
+	TRDGUniformBufferRef<FPreshadingUniformBufferParameters>& PreshadingUniformBuffer
+)
+{
+	FPreshadingUniformBufferParameters* PreshadingUniformBufferParameters = GraphBuilder.AllocParameters<FPreshadingUniformBufferParameters>();
+	{
+		PreshadingUniformBufferParameters->LocalToWorld = FMatrix44f::Identity;
+		PreshadingUniformBufferParameters->WorldToLocal = FMatrix44f::Identity;
+		PreshadingUniformBufferParameters->LocalBoundsOrigin = FVector3f::ZeroVector;
+		PreshadingUniformBufferParameters->LocalBoundsExtent = FVector3f::ZeroVector;
+
+		PreshadingUniformBufferParameters->VolumeResolution = FIntVector::ZeroValue;
+		PreshadingUniformBufferParameters->ExtinctionTexture = GraphBuilder.RegisterExternalTexture(GSystemTextures.BlackDummy);
+		PreshadingUniformBufferParameters->EmissionTexture = GraphBuilder.RegisterExternalTexture(GSystemTextures.BlackDummy);
+		PreshadingUniformBufferParameters->AlbedoTexture = GraphBuilder.RegisterExternalTexture(GSystemTextures.BlackDummy);
+		PreshadingUniformBufferParameters->TextureSampler = TStaticSamplerState<SF_Trilinear, AM_Clamp, AM_Clamp, AM_Clamp>::GetRHI();
+
+		PreshadingUniformBufferParameters->MajorantResolution = FIntVector::ZeroValue;
+		PreshadingUniformBufferParameters->MajorantTexture = GraphBuilder.RegisterExternalTexture(GSystemTextures.BlackDummy);
+		PreshadingUniformBufferParameters->PointTextureSampler = TStaticSamplerState<SF_Point, AM_Clamp, AM_Clamp, AM_Clamp>::GetRHI();
+	}
+	PreshadingUniformBuffer = GraphBuilder.CreateUniformBuffer(PreshadingUniformBufferParameters);
+}
+
+void BuildMajorantTexture(
+	FRDGBuilder& GraphBuilder,
+	const FScene* Scene,
+	const FViewInfo& View,
+	const IHeterogeneousVolumeInterface* HeterogeneousVolumeInterface,
+	const FMaterialRenderProxy* MaterialRenderProxy,
+	FIntVector& MajorantResolution,
+	FRDGTextureRef MajorantHistoryTexture,
+	FRDGTextureRef& MajorantTexture
+)
+{
+	RDG_EVENT_SCOPE(GraphBuilder, "Majorant Update");
+
+	// Collect material and object information
+	const FPrimitiveSceneProxy* PrimitiveSceneProxy = HeterogeneousVolumeInterface->GetPrimitiveSceneProxy();
+	const FPrimitiveSceneInfo* PrimitiveSceneInfo = PrimitiveSceneProxy->GetPrimitiveSceneInfo();
+	const FPersistentPrimitiveIndex PersistentPrimitiveIndex = PrimitiveSceneInfo->GetPersistentIndex();
+	const FBoxSphereBounds LocalBoxSphereBounds = HeterogeneousVolumeInterface->GetLocalBounds();
+
+	// Create baked material grids
+	const uint32 NumMips = 1;
+	const EPixelFormat PixelFormat = HeterogeneousVolumes::GetPreshadingPixelFormat();
+	MajorantResolution = FIntVector(64, 64, 64);
+
+	FRDGTextureDesc MajorantDesc = FRDGTextureDesc::Create3D(
+		MajorantResolution,
+		PixelFormat,
+		FClearValueBinding::Black,
+		TexCreate_ShaderResource | TexCreate_UAV | TexCreate_3DTiling,
+		NumMips
+	);
+	FRDGTextureRef ExtinctionTexture = GraphBuilder.CreateTexture(MajorantDesc, TEXT("HeterogeneousVolumes.ExtinctionTexture"));
+
+	FRDGTextureDesc EmptyDesc = FRDGTextureDesc::Create3D(
+		FIntVector(1, 1, 1),
+		PixelFormat,
+		FClearValueBinding::Black,
+		TexCreate_ShaderResource | TexCreate_UAV | TexCreate_3DTiling,
+		NumMips
+	);
+	FRDGTextureRef DummyTexture = GraphBuilder.CreateTexture(EmptyDesc, TEXT("HeterogeneousVolumes.DummyTexture"));
+
+	AddClearUAVPass(GraphBuilder, GraphBuilder.CreateUAV(ExtinctionTexture), FLinearColor::Black);
+	AddClearUAVPass(GraphBuilder, GraphBuilder.CreateUAV(DummyTexture), FLinearColor::White);
+
+	// Sample data
+	bool bJitter = true;
+	bool bLooseBounds = true;
+	bool bEvaluateAlbedo = false;
+	bool bEvaluateEmission = false;
+	ComputeHeterogeneousVolumeBakeMaterial(
+		GraphBuilder,
+		// Scene data
+		Scene,
+		View,
+		// Object data
+		HeterogeneousVolumeInterface,
+		MaterialRenderProxy,
+		PersistentPrimitiveIndex,
+		LocalBoxSphereBounds,
+		// Sample data
+		bJitter,
+		bLooseBounds,
+		bEvaluateAlbedo,
+		bEvaluateEmission,
+		// Volume data
+		MajorantResolution,
+		// Output
+		ExtinctionTexture,
+		DummyTexture,
+		DummyTexture
+	);
+
+	// Update majorant
+	if (MajorantHistoryTexture)
+	{
+		MajorantTexture = GraphBuilder.CreateTexture(MajorantDesc, TEXT("HeterogeneousVolumes.MajorantTexture"));
+		MaxTexture3D(
+			GraphBuilder,
+			View,
+			MajorantHistoryTexture,
+			ExtinctionTexture,
+			MajorantTexture
+		);
+	}
+	else
+	{
+		MajorantTexture = ExtinctionTexture;
+	}
+}
+
+static TAutoConsoleVariable<int32> CVarHeterogeneousVolumesMajorantResolution(
+	TEXT("r.HeterogeneousVolumes.Majorant.Resolution"),
+	64,
+	TEXT("Determines the underlying resolution for the majorant acceleration texture (Default = 64)\n")
+	TEXT("When a value of 0 is used, the majorant corresponds to the native resolution of the volume"),
+	ECVF_RenderThreadSafe
+);
+
+FIntVector GetMajorantResolution(const IHeterogeneousVolumeInterface* HeterogeneousVolumeInterface)
+{
+	int32 Value = CVarHeterogeneousVolumesMajorantResolution.GetValueOnRenderThread();
+	return (Value == 0) ? HeterogeneousVolumeInterface->GetVoxelResolution() : FIntVector(Value);
+}
+
+void BuildMajorantTexture(
+	FRDGBuilder& GraphBuilder,
+	const FScene* Scene,
+	const FViewInfo& View,
+	const IHeterogeneousVolumeInterface* HeterogeneousVolumeInterface,
+	const FMaterialRenderProxy* MaterialRenderProxy,
+	FIntVector& MajorantResolution,
+	TRefCountPtr<IPooledRenderTarget> MajorantHistoryTexture,
+	TRefCountPtr<IPooledRenderTarget>& MajorantTexture
+)
+{
+	RDG_EVENT_SCOPE(GraphBuilder, "Majorant Update");
+
+	// Collect material and object information
+	const FPrimitiveSceneProxy* PrimitiveSceneProxy = HeterogeneousVolumeInterface->GetPrimitiveSceneProxy();
+	const FPrimitiveSceneInfo* PrimitiveSceneInfo = PrimitiveSceneProxy->GetPrimitiveSceneInfo();
+	const FPersistentPrimitiveIndex PersistentPrimitiveIndex = PrimitiveSceneInfo->GetPersistentIndex();
+	const FBoxSphereBounds LocalBoxSphereBounds = HeterogeneousVolumeInterface->GetLocalBounds();
+
+	// Create baked material grids
+	const uint32 NumMips = 1;
+	const EPixelFormat PixelFormat = HeterogeneousVolumes::GetPreshadingPixelFormat();
+	MajorantResolution = GetMajorantResolution(HeterogeneousVolumeInterface);
+
+	FRHITextureCreateDesc MajorantDesc = FRHITextureCreateDesc::Create3D(TEXT("HeterogeneousVolumes.MajorantTexture"), MajorantResolution, PixelFormat)
+		.SetFlags(ETextureCreateFlags::ShaderResource | ETextureCreateFlags::UAV);
+	TRefCountPtr<IPooledRenderTarget> ExtinctionTexture = CreateRenderTarget(GraphBuilder.RHICmdList.CreateTexture(MajorantDesc), MajorantDesc.DebugName);
+
+	FRDGTextureDesc EmptyDesc = FRDGTextureDesc::Create3D(
+		FIntVector(1, 1, 1),
+		PixelFormat,
+		FClearValueBinding::Black,
+		TexCreate_ShaderResource | TexCreate_UAV | TexCreate_3DTiling,
+		NumMips
+	);
+	FRDGTextureRef DummyTexture = GraphBuilder.CreateTexture(EmptyDesc, TEXT("HeterogeneousVolumes.DummyTexture"));
+
+	AddClearUAVPass(GraphBuilder, GraphBuilder.CreateUAV(GraphBuilder.RegisterExternalTexture(ExtinctionTexture)), FLinearColor::Black);
+	AddClearUAVPass(GraphBuilder, GraphBuilder.CreateUAV(DummyTexture), FLinearColor::White);
+
+	// Sample data
+	bool bJitter = true;
+	bool bLooseBounds = true;
+	bool bEvaluateAlbedo = false;
+	bool bEvaluateEmission = false;
+	ComputeHeterogeneousVolumeBakeMaterial(
+		GraphBuilder,
+		// Scene data
+		Scene,
+		View,
+		// Object data
+		HeterogeneousVolumeInterface,
+		MaterialRenderProxy,
+		PersistentPrimitiveIndex,
+		LocalBoxSphereBounds,
+		// Sample data
+		bJitter,
+		bLooseBounds,
+		bEvaluateAlbedo,
+		bEvaluateEmission,
+		// Volume data
+		MajorantResolution,
+		// Output
+		GraphBuilder.RegisterExternalTexture(ExtinctionTexture),
+		DummyTexture,
+		DummyTexture
+	);
+
+	// Update majorant
+	if (MajorantHistoryTexture && MajorantTexture)
+	{
+		MaxTexture3D(
+			GraphBuilder,
+			View,
+			GraphBuilder.RegisterExternalTexture(MajorantHistoryTexture),
+			GraphBuilder.RegisterExternalTexture(ExtinctionTexture),
+			GraphBuilder.RegisterExternalTexture(MajorantTexture)
+		);
+	}
+	else
+	{
+		MajorantTexture = ExtinctionTexture;
+	}
+}
+
+void BuildPreshadingGrid(
+	FRDGBuilder& GraphBuilder,
+	const FScene* Scene,
+	const FViewInfo& View,
+	const IHeterogeneousVolumeInterface* HeterogeneousVolumeInterface,
+	const FMaterialRenderProxy* MaterialRenderProxy,
+	TRefCountPtr<IPooledRenderTarget>& MajorantHistoryRT,
+	TRDGUniformBufferRef<FPreshadingUniformBufferParameters>& PreshadingUniformBuffer
+)
+{
+	check(View.ViewState);
+
+	// Determine baking resolution
 	FIntVector VolumeResolution = HeterogeneousVolumes::GetVolumeResolution(HeterogeneousVolumeInterface);
 	FIntVector BakingResolution = VolumeResolution;
 	if (CVarHeterogeneousVolumesPreshadingCLOD.GetValueOnAnyThread() == 1)
@@ -1290,6 +1689,12 @@
 		BakingResolution = HeterogeneousVolumes::GetLightingCacheResolution(HeterogeneousVolumeInterface, LODValue);
 	}
 
+	// Collect material and object information
+	const FPrimitiveSceneProxy* PrimitiveSceneProxy = HeterogeneousVolumeInterface->GetPrimitiveSceneProxy();
+	const FPrimitiveSceneInfo* PrimitiveSceneInfo = PrimitiveSceneProxy->GetPrimitiveSceneInfo();
+	const FPersistentPrimitiveIndex PersistentPrimitiveIndex = PrimitiveSceneInfo->GetPersistentIndex();
+	const FBoxSphereBounds LocalBoxSphereBounds = HeterogeneousVolumeInterface->GetLocalBounds();
+
 	// Create baked material grids
 	const uint32 NumMips = 1;
 	const EPixelFormat PixelFormat = HeterogeneousVolumes::GetPreshadingPixelFormat();
@@ -1301,60 +1706,213 @@
 		TexCreate_ShaderResource | TexCreate_UAV | TexCreate_3DTiling,
 		NumMips
 	);
-
 	FRDGTextureRef ExtinctionTexture = GraphBuilder.CreateTexture(BakedMaterialDesc, TEXT("HeterogeneousVolumes.ExtinctionTexture"));
 	FRDGTextureRef EmissionTexture = GraphBuilder.CreateTexture(BakedMaterialDesc, TEXT("HeterogeneousVolumes.EmissionTexture"));
 	FRDGTextureRef AlbedoTexture = GraphBuilder.CreateTexture(BakedMaterialDesc, TEXT("HeterogeneousVolumes.AlbedoTexture"));
 
-	// Material baking
+	bool bJitter = HeterogeneousVolumes::ShouldJitter();
+	bool bLooseBounds = false;
+	bool bEvaluateAlbedo = true;
+	bool bEvaluateEmission = true;
+	ComputeHeterogeneousVolumeBakeMaterial(
+		GraphBuilder,
+		// Scene data
+		Scene,
+		View,
+		// Object data
+		HeterogeneousVolumeInterface,
+		MaterialRenderProxy,
+		PersistentPrimitiveIndex,
+		LocalBoxSphereBounds,
+		// Sample data
+		bJitter,
+		bLooseBounds,
+		bEvaluateAlbedo,
+		bEvaluateEmission,
+		// Volume data
+		BakingResolution,
+		// Output
+		ExtinctionTexture,
+		EmissionTexture,
+		AlbedoTexture
+	);
+
+	FRDGTextureRef MajorantHistory = MajorantHistoryRT ? GraphBuilder.RegisterExternalTexture(MajorantHistoryRT) : nullptr;
+
+	FIntVector MajorantResolution;
+	FRDGTextureRef MajorantTexture;
+	BuildMajorantTexture(
+		GraphBuilder,
+		Scene,
+		View,
+		HeterogeneousVolumeInterface,
+		MaterialRenderProxy,
+		MajorantResolution,
+		MajorantHistory,
+		MajorantTexture
+	);
+
+	// Cache result
+	GraphBuilder.QueueTextureExtraction(MajorantTexture, &MajorantHistoryRT);
+
+	BuildPreshadingUniformBuffer(
+		GraphBuilder,
+		Scene,
+		View,
+		HeterogeneousVolumeInterface,
+		MaterialRenderProxy,
+		BakingResolution,
+		ExtinctionTexture,
+		EmissionTexture,
+		AlbedoTexture,
+		MajorantResolution,
+		MajorantTexture,
+		PreshadingUniformBuffer
+	);
+}
+
+void BuildPreshadingGrid(
+	FRDGBuilder& GraphBuilder,
+	const FScene* Scene,
+	const FViewInfo& View,
+	const IHeterogeneousVolumeInterface* HeterogeneousVolumeInterface,
+	const FMaterialRenderProxy* MaterialRenderProxy,
+	TRefCountPtr<IPooledRenderTarget>& PersistentMajorantRT,
+	FPathTracingHeterogeneousVolumesCallableInfo& CallableInfo,
+	FPreshadingUniformBufferParameters& PreshadingUniformBufferParameters
+)
+{
+	// Determine baking resolution
+	FIntVector VolumeResolution = HeterogeneousVolumes::GetVolumeResolution(HeterogeneousVolumeInterface);
+	FIntVector BakingResolution = VolumeResolution;
+	if (CVarHeterogeneousVolumesPreshadingCLOD.GetValueOnAnyThread() == 1)
 	{
-		RDG_EVENT_SCOPE(GraphBuilder, "Material Baking");
-		SCOPE_CYCLE_COUNTER(STATGROUP_HeterogeneousVolumesMaterialBaking);
+		// Determine baking voxel resolution
+		float LODFactor = HeterogeneousVolumes::CalcLODFactor(View, HeterogeneousVolumeInterface);
+		BakingResolution = FIntVector(BakingResolution / LODFactor);
+		BakingResolution.X = FMath::Clamp(BakingResolution.X, 1, 1024);
+		BakingResolution.Y = FMath::Clamp(BakingResolution.Y, 1, 1024);
+		BakingResolution.Z = FMath::Clamp(BakingResolution.Z, 1, 1024);
 
-		ComputeHeterogeneousVolumeBakeMaterial(
-			GraphBuilder,
-			// Scene data
-			Scene,
-			View,
-			// Object data
-			HeterogeneousVolumeInterface,
-			MaterialRenderProxy,
-			PersistentPrimitiveIndex,
-			LocalBoxSphereBounds,
-			// Volume data
-			BakingResolution,
-			// Output
-			ExtinctionTexture,
-			EmissionTexture,
-			AlbedoTexture
-		);
+	}
+	else if (CVarHeterogeneousVolumesPreshadingCLOD.GetValueOnAnyThread() == 2)
+	{
+		HeterogeneousVolumes::FLODValue LODValue = HeterogeneousVolumes::CalcLOD(View, HeterogeneousVolumeInterface);
+		BakingResolution = HeterogeneousVolumes::GetLightingCacheResolution(HeterogeneousVolumeInterface, LODValue);
 	}
 
-	// Create Preshading UniformBuffer
-	FPreshadingUniformBufferParameters* PreshadingUniformBufferParameters = GraphBuilder.AllocParameters<FPreshadingUniformBufferParameters>();
-	{
+	// Collect material and object information
+	const FPrimitiveSceneProxy* PrimitiveSceneProxy = HeterogeneousVolumeInterface->GetPrimitiveSceneProxy();
+	const FPrimitiveSceneInfo* PrimitiveSceneInfo = PrimitiveSceneProxy->GetPrimitiveSceneInfo();
+	const FPersistentPrimitiveIndex PersistentPrimitiveIndex = PrimitiveSceneInfo->GetPersistentIndex();
+	const FBoxSphereBounds LocalBoxSphereBounds = HeterogeneousVolumeInterface->GetLocalBounds();
+
+	// Create baked material grids
+	const uint32 NumMips = 1;
+	const EPixelFormat PixelFormat = HeterogeneousVolumes::GetPreshadingPixelFormat();
+
+	FRHITextureCreateDesc ExtinctionMaterialDesc = FRHITextureCreateDesc::Create3D(TEXT("HeterogeneousVolumes.ExtinctionTexture"), BakingResolution, PixelFormat)
+		.SetFlags(ETextureCreateFlags::ShaderResource | ETextureCreateFlags::UAV);
+	CallableInfo.ExtinctionRT = CreateRenderTarget(GraphBuilder.RHICmdList.CreateTexture(ExtinctionMaterialDesc), ExtinctionMaterialDesc.DebugName);
+
+	FRHITextureCreateDesc EmissionMaterialDesc = FRHITextureCreateDesc::Create3D(TEXT("HeterogeneousVolumes.EmissionTexture"), BakingResolution, PixelFormat)
+		.SetFlags(ETextureCreateFlags::ShaderResource | ETextureCreateFlags::UAV);
+	CallableInfo.EmissionRT = CreateRenderTarget(GraphBuilder.RHICmdList.CreateTexture(EmissionMaterialDesc), EmissionMaterialDesc.DebugName);
+
+	FRHITextureCreateDesc AlbedoMaterialDesc = FRHITextureCreateDesc::Create3D(TEXT("HeterogeneousVolumes.AlbedoTexture"), BakingResolution, PixelFormat)
+		.SetFlags(ETextureCreateFlags::ShaderResource | ETextureCreateFlags::UAV);
+	CallableInfo.AlbedoRT = CreateRenderTarget(GraphBuilder.RHICmdList.CreateTexture(AlbedoMaterialDesc), AlbedoMaterialDesc.DebugName);
+
+	bool bJitter = HeterogeneousVolumes::ShouldJitter();
+	bool bLooseBounds = false;
+	bool bEvaluateAlbedo = true;
+	bool bEvaluateEmission = true;
+	ComputeHeterogeneousVolumeBakeMaterial(
+		GraphBuilder,
+		// Scene data
+		Scene,
+		View,
 		// Object data
-		// TODO: Convert to relative-local space
-		//FVector3f ViewOriginHigh = FDFVector3(View.ViewMatrices.GetViewOrigin()).High;
-		//FMatrix44f RelativeLocalToWorld = FDFMatrix::MakeToRelativeWorldMatrix(ViewOriginHigh, HeterogeneousVolumeInterface->GetLocalToWorld()).M;
-		FMatrix InstanceToLocal = HeterogeneousVolumeInterface->GetInstanceToLocal();
-		FMatrix LocalToWorld = HeterogeneousVolumeInterface->GetLocalToWorld();
-		PreshadingUniformBufferParameters->LocalToWorld = FMatrix44f(InstanceToLocal * LocalToWorld);
-		PreshadingUniformBufferParameters->WorldToLocal = PreshadingUniformBufferParameters->LocalToWorld.Inverse();
+		HeterogeneousVolumeInterface,
+		MaterialRenderProxy,
+		PersistentPrimitiveIndex,
+		LocalBoxSphereBounds,
+		// Sample data
+		bJitter,
+		bLooseBounds,
+		bEvaluateAlbedo,
+		bEvaluateEmission,
+		// Volume data
+		BakingResolution,
+		// Output
+		CallableInfo.ExtinctionRT,
+		CallableInfo.EmissionRT,
+		CallableInfo.AlbedoRT
+	);
 
-		FMatrix LocalToInstance = InstanceToLocal.Inverse();
-		FBoxSphereBounds InstanceBoxSphereBounds = LocalBoxSphereBounds.TransformBy(LocalToInstance);
-		PreshadingUniformBufferParameters->LocalBoundsOrigin = FVector3f(InstanceBoxSphereBounds.Origin);
-		PreshadingUniformBufferParameters->LocalBoundsExtent = FVector3f(InstanceBoxSphereBounds.BoxExtent);
+	//TRefCountPtr<IPooledRenderTarget> MajorantHistoryRT = View.ViewState ? CallableInfo.MajorantRT : nullptr;
+	TRefCountPtr<IPooledRenderTarget> MajorantHistoryRT = View.ViewState ? PersistentMajorantRT : nullptr;
+	FIntVector MajorantResolution;
+	BuildMajorantTexture(
+		GraphBuilder,
+		Scene,
+		View,
+		HeterogeneousVolumeInterface,
+		MaterialRenderProxy,
+		MajorantResolution,
+		MajorantHistoryRT,
+		PersistentMajorantRT
+	);
 
-		// Volume data
-		PreshadingUniformBufferParameters->VolumeResolution = VolumeResolution;
-		PreshadingUniformBufferParameters->ExtinctionTexture = ExtinctionTexture;
-		PreshadingUniformBufferParameters->EmissionTexture = EmissionTexture;
-		PreshadingUniformBufferParameters->AlbedoTexture = AlbedoTexture;
-		PreshadingUniformBufferParameters->TextureSampler = TStaticSamplerState<SF_Trilinear, AM_Clamp, AM_Clamp, AM_Clamp>::GetRHI();
-	}
-	TRDGUniformBufferRef<FPreshadingUniformBufferParameters> PreshadingUniformBuffer = GraphBuilder.CreateUniformBuffer(PreshadingUniformBufferParameters);
+	BuildPreshadingUniformBufferParameters(
+		GraphBuilder,
+		Scene,
+		View,
+		HeterogeneousVolumeInterface,
+		MaterialRenderProxy,
+		BakingResolution,
+		CallableInfo.ExtinctionRT,
+		CallableInfo.EmissionRT,
+		CallableInfo.AlbedoRT,
+		MajorantResolution,
+		PersistentMajorantRT,
+		PreshadingUniformBufferParameters
+	);
+}
 
+void RenderWithPreshading(
+	FRDGBuilder& GraphBuilder,
+	// Scene data
+	const FSceneTextures& SceneTextures,
+	FScene* Scene,
+	FViewInfo& View, int32 ViewIndex,
+	// Shadow data
+	TArray<FVisibleLightInfo, SceneRenderingAllocator>& VisibleLightInfos,
+	const FVirtualShadowMapArray& VirtualShadowMapArray,
+	// Object data
+	const IHeterogeneousVolumeInterface* HeterogeneousVolumeInterface,
+	const FMaterialRenderProxy* MaterialRenderProxy,
+	const FPersistentPrimitiveIndex& PersistentPrimitiveIndex,
+	const FBoxSphereBounds LocalBoxSphereBounds,
+	// Transmittance acceleration
+	FRDGTextureRef LightingCacheTexture,
+	// Output
+	FRDGTextureRef& HeterogeneousVolumeRadiance
+)
+{
+	CSV_SCOPED_TIMING_STAT_EXCLUSIVE(Effects);
+
+	TRefCountPtr<IPooledRenderTarget> MajorantHistoryRT;
+	TRDGUniformBufferRef<FPreshadingUniformBufferParameters> PreshadingUniformBuffer;
+	BuildPreshadingGrid(
+		GraphBuilder,
+		Scene,
+		View,
+		HeterogeneousVolumeInterface,
+		MaterialRenderProxy,
+		MajorantHistoryRT,
+		PreshadingUniformBuffer
+	);
+
 	RenderWithPreshadingCompute(
 		GraphBuilder,
diff --git a/Engine/Source/Runtime/Renderer/Private/PathTracing.cpp b/Engine/Source/Runtime/Renderer/Private/PathTracing.cpp
--- a/Engine/Source/Runtime/Renderer/Private/PathTracing.cpp
+++ b/Engine/Source/Runtime/Renderer/Private/PathTracing.cpp
@@ -518,6 +518,16 @@ TAutoConsoleVariable<bool> CVarPathTracingHeterogeneousVolumesRebuildEveryFrame(
 	ECVF_RenderThreadSafe
 );
 
+TAutoConsoleVariable<int32> CVarPathTracingHeterogeneousVolumesPipeline(
+	TEXT("r.PathTracing.HeterogeneousVolumes.Pipeline"),
+	0,
+	TEXT("Determines which pipeline to invoke when rendering Heterogeneous Volumes. The first method resamples the volume into a global voxel buffer. The second method isolates each volume, and invokes an individual callable shader to perform rendering options.\n")
+	TEXT("0: Global voxel buffers (default)\n")
+	TEXT("1: Callable shading\n")
+	TEXT("Requires r.PathTracing.Experimental=true to modify.\n"),
+	ECVF_RenderThreadSafe
+);
+
 TAutoConsoleVariable<bool> CVarPathTracingCameraMediumTracking(
 	TEXT("r.PathTracing.CameraMediumTracking"),
 	true,
@@ -807,6 +817,9 @@ struct FPathTracingState
 	FAdaptiveOrthoGridParameterCache AdaptiveOrthoGridParameterCache;
 	FAdaptiveFrustumGridParameterCache AdaptiveFrustumGridParameterCache;
 
+	TMap<FHeterogeneousVolumesPrimitiveInstanceKey, TRefCountPtr<IPooledRenderTarget> > HeterogeneousVolumesPersistentMajorantRT;
+	TArray<FPathTracingHeterogeneousVolumesCallableInfo> HeterogeneousVolumesCallableInfo;
+
 	// Texture holding onto the precomputed atmosphere data
 	TRefCountPtr<IPooledRenderTarget> AtmosphereOpticalDepthLUT;
 	FAtmosphereConfig LastAtmosphereConfig;
@@ -1261,7 +1274,8 @@ class FPathTracingRG : public FGlobalShader
 	class FSubstrateComplexSpecialMaterial : SHADER_PERMUTATION_BOOL("PATH_TRACER_USE_SUBSTRATE_SPECIAL_COMPLEX_MATERIAL");
 	class FUseSER : SHADER_PERMUTATION_BOOL("PATH_TRACER_USE_SER");
 	class FNeedTMinWorkaround : SHADER_PERMUTATION_BOOL("NEED_TMIN_WORKAROUND");
-	using FPermutationDomain = TShaderPermutationDomain<FKernelType, FAdaptiveSampling, FCloudShader, FSubstrateComplexSpecialMaterial, FUseSER, FNeedTMinWorkaround>;
+	class FHeterogeneousVolumesPipeline : SHADER_PERMUTATION_INT("PATH_TRACER_HETEROGENEOUS_VOLUMES_PIPELINE", 2);
+	using FPermutationDomain = TShaderPermutationDomain<FKernelType, FAdaptiveSampling, FCloudShader, FSubstrateComplexSpecialMaterial, FUseSER, FNeedTMinWorkaround, FHeterogeneousVolumesPipeline>;
 
 	static bool ShouldCompilePermutation(const FGlobalShaderPermutationParameters& Parameters)
 	{
@@ -1274,6 +1288,11 @@ class FPathTracingRG : public FGlobalShader
 			{
 				return false;
 			}
+
+			if (PermutationVector.Get<FHeterogeneousVolumesPipeline>())
+			{
+				return false;
+			}
 		}
 
 		if (PermutationVector.Get<FCloudShader>())
@@ -1370,6 +1389,16 @@ class FPathTracingRG : public FGlobalShader
 		// Heterogeneous volumes adaptive voxel grid
 		SHADER_PARAMETER_RDG_UNIFORM_BUFFER(FOrthoVoxelGridUniformBufferParameters, OrthoGridUniformBuffer)
 		SHADER_PARAMETER_RDG_UNIFORM_BUFFER(FFrustumVoxelGridUniformBufferParameters, FrustumGridUniformBuffer)
+		// NOTE: Remove in separate submission
+		//SHADER_PARAMETER_RDG_UNIFORM_BUFFER(FPreshadingUniformBufferParameters, PreshadingUniformBuffer)
+
+        // Heterogeneous volumes callable interface
+		SHADER_PARAMETER(int32, HeterogeneousVolumesNum)
+		SHADER_PARAMETER_RDG_BUFFER_SRV(Buffer<int>, HeterogeneousVolumesCallableShaderIdIntersect)
+		SHADER_PARAMETER_RDG_BUFFER_SRV(Buffer<int>, HeterogeneousVolumesCallableShaderIdGetDensityBounds)
+		SHADER_PARAMETER_RDG_BUFFER_SRV(Buffer<int>, HeterogeneousVolumesCallableShaderIdGetDensity)
+		SHADER_PARAMETER_RDG_BUFFER_SRV(Buffer<int>, HeterogeneousVolumesCallableShaderIdSampleDistance)
+		SHADER_PARAMETER_RDG_BUFFER_SRV(Buffer<int>, HeterogeneousVolumesCallableShaderIdGetTransmittance)
 
 		// scene decals
 		SHADER_PARAMETER_RDG_UNIFORM_BUFFER(FRayTracingDecals, DecalParameters)
@@ -1798,6 +1827,108 @@ private:
 };
 IMPLEMENT_SHADER_TYPE(, FPathTracingVolumetricCloudMaterial, TEXT("/Engine/Private/PathTracing/PathTracingVolumetricCloudMaterialShader.usf"), TEXT("PathTracingVolumetricCloudMaterialShader"), SF_RayCallable);
 
+enum class ECallableType
+{
+	Intersect,
+	GetDensityBounds,
+	GetDensity,
+	SampleDistance,
+	GetTransmittance
+};
+
+template <enum ECallableType>
+class FPathTracingHeterogeneousVolumesMaterial : public FGlobalShader
+//class FPathTracingHeterogeneousVolumesMaterial : public FMaterialShader
+{
+	//DECLARE_SHADER_TYPE(FPathTracingHeterogeneousVolumesMaterial, Material);
+	DECLARE_GLOBAL_SHADER(FPathTracingHeterogeneousVolumesMaterial)
+public:
+	FPathTracingHeterogeneousVolumesMaterial() = default;
+
+	FPathTracingHeterogeneousVolumesMaterial(const ShaderMetaType::CompiledShaderInitializerType& Initializer)
+		//: FMaterialShader(Initializer)
+		: FGlobalShader(Initializer)
+	{
+		PreshadingUniformBuffer.Bind(Initializer.ParameterMap, TEXT("PreshadingUniformBuffer"));
+	}
+
+	//static bool ShouldCompilePermutation(const FMaterialShaderPermutationParameters& Parameters)
+	static bool ShouldCompilePermutation(const FGlobalShaderPermutationParameters& Parameters)
+	{
+		return PathTracing::ShouldCompilePathTracingShadersForProject(Parameters.Platform) &&
+			ShouldCompileRayTracingCallableShadersForProject(Parameters.Platform)
+			//Parameters.MaterialParameters.bIsUsedWithHeterogeneousVolumes &&
+			//Parameters.MaterialParameters.MaterialDomain == MD_Volume
+			;
+	}
+
+	//static void ModifyCompilationEnvironment(const FMaterialShaderPermutationParameters& Parameters, FShaderCompilerEnvironment& OutEnvironment)
+	static void ModifyCompilationEnvironment(const FGlobalShaderPermutationParameters& Parameters, FShaderCompilerEnvironment& OutEnvironment)
+	{
+		//FMeshMaterialShader::ModifyCompilationEnvironment(Parameters, OutEnvironment);
+		//check(Parameters.MaterialParameters.MaterialDomain == MD_Volume);
+		FGlobalShader::ModifyCompilationEnvironment(Parameters, OutEnvironment);
+	}
+
+	static bool ValidateCompiledResult(EShaderPlatform Platform, const FShaderParameterMap& ParameterMap, TArray<FString>& OutError)
+	{
+		if (ParameterMap.ContainsParameterAllocation(FSceneTextureUniformParameters::FTypeInfo::GetStructMetadata()->GetShaderVariableName()))
+		{
+			OutError.Add(TEXT("Ray tracing callable shaders cannot read from the SceneTexturesStruct."));
+			return false;
+		}
+
+		for (const auto& It : ParameterMap.GetParameterMap())
+		{
+			const FParameterAllocation& ParamAllocation = It.Value;
+			if (ParamAllocation.Type != EShaderParameterType::UniformBuffer
+				&& ParamAllocation.Type != EShaderParameterType::LooseData)
+			{
+				OutError.Add(FString::Printf(TEXT("Invalid ray tracing shader parameter '%s'. Only uniform buffers and loose data parameters are supported."), *(It.Key)));
+				return false;
+			}
+		}
+
+		return true;
+	}
+
+	static ERayTracingPayloadType GetRayTracingPayloadType(const int32 PermutationId)
+	{
+		// TODO: This isn't the payload we use, but the logic in RayTracingMaterialHitShaders.cpp needs to assume a consistent payload ID for all callable shaders
+		// The simplest solution is probably to remove FDecalShaderPayload and use the PT payload everywhere
+		return ERayTracingPayloadType::Decals;
+	}
+
+	static const FShaderBindingLayout* GetShaderBindingLayout(const FShaderPermutationParameters& Parameters)
+	{
+		return RayTracing::GetShaderBindingLayout(Parameters.Platform);
+	}
+
+	void GetShaderBindings(
+		const FViewInfo& View,
+		const TUniformBufferRef<FPreshadingUniformBufferParameters>& PreshadingUniformBufferParameters,
+		FMeshDrawSingleShaderBindings& ShaderBindings) const
+	{
+		//FMaterialShader::GetShaderBindings(Scene, FeatureLevel, MaterialRenderProxy, Material, ShaderBindings);
+
+		ShaderBindings.Add(GetUniformBufferParameter<FViewUniformShaderParameters>(), View.ViewUniformBuffer);
+		ShaderBindings.Add(GetUniformBufferParameter<FPreshadingUniformBufferParameters>(), PreshadingUniformBufferParameters);
+	}
+
+private:
+	LAYOUT_FIELD(FShaderUniformBufferParameter, PreshadingUniformBuffer);
+};
+
+using FPathTracingHeterogeneousVolumesIntersectCS = FPathTracingHeterogeneousVolumesMaterial<ECallableType::Intersect>;
+using FPathTracingHeterogeneousVolumesGetDensityBoundsCS = FPathTracingHeterogeneousVolumesMaterial<ECallableType::GetDensityBounds>;
+using FPathTracingHeterogeneousVolumesGetDensityCS = FPathTracingHeterogeneousVolumesMaterial<ECallableType::GetDensity>;
+using FPathTracingHeterogeneousVolumesSampleDistanceCS = FPathTracingHeterogeneousVolumesMaterial<ECallableType::SampleDistance>;
+using FPathTracingHeterogeneousVolumesGetTransmittanceCS = FPathTracingHeterogeneousVolumesMaterial<ECallableType::GetTransmittance>;
+IMPLEMENT_SHADER_TYPE(template<>, FPathTracingHeterogeneousVolumesIntersectCS, TEXT("/Engine/Private/PathTracing/PathTracingHeterogeneousVolumeMaterialShader.usf"), TEXT("PathTracingHeterogeneousVolumesIntersectCS"), SF_RayCallable);
+IMPLEMENT_SHADER_TYPE(template<>, FPathTracingHeterogeneousVolumesGetDensityBoundsCS, TEXT("/Engine/Private/PathTracing/PathTracingHeterogeneousVolumeMaterialShader.usf"), TEXT("PathTracingHeterogeneousVolumesGetDensityBoundsCS"), SF_RayCallable);
+IMPLEMENT_SHADER_TYPE(template<>, FPathTracingHeterogeneousVolumesGetDensityCS, TEXT("/Engine/Private/PathTracing/PathTracingHeterogeneousVolumeMaterialShader.usf"), TEXT("PathTracingHeterogeneousVolumesGetDensityCS"), SF_RayCallable);
+IMPLEMENT_SHADER_TYPE(template<>, FPathTracingHeterogeneousVolumesSampleDistanceCS, TEXT("/Engine/Private/PathTracing/PathTracingHeterogeneousVolumeMaterialShader.usf"), TEXT("PathTracingHeterogeneousVolumesSampleDistanceCS"), SF_RayCallable);
+IMPLEMENT_SHADER_TYPE(template<>, FPathTracingHeterogeneousVolumesGetTransmittanceCS, TEXT("/Engine/Private/PathTracing/PathTracingHeterogeneousVolumeMaterialShader.usf"), TEXT("PathTracingHeterogeneousVolumesGetTransmittanceCS"), SF_RayCallable);
 
 class FPathTracingVolumetricCloudMaterialVisualize : public FGlobalShader
 {
@@ -2013,6 +2144,258 @@ void PreparePathTracingCloudMaterial(FRDGBuilder& GraphBuilder, FScene* Scene, T
 	}
 }
 
+template <class CallableShaderType>
+void BindCallableShader(
+	FScene* Scene,
+	const FViewInfo& View,
+	const FMaterialRenderProxy* MaterialRenderProxy,
+	const FMaterial* MaterialResource,
+	const FMaterialShaderMap* MaterialShaderMap,
+	TUniformBufferRef<FPreshadingUniformBufferParameters> PreshadingUniformBuffer,
+	int32& CallableShaderIndex
+)
+{
+	//auto CallableShader = MaterialShaderMap->GetShader<CallableShaderType>();
+	auto CallableShader = GetGlobalShaderMap(Scene->GetFeatureLevel())->GetShader<CallableShaderType>();
+	if (!CallableShader.IsValid())
+	{
+		return;
+	}
+
+	CallableShaderIndex = Scene->RayTracingSBT.NumCallableShaderSlots++;
+
+	FRayTracingShaderCommand& Command = Scene->RayTracingSBT.CallableCommands.AddDefaulted_GetRef();
+	Command.SetShader(CallableShader);
+	Command.SlotInScene = CallableShaderIndex;
+
+	// Bind resource to shader
+	FMeshDrawSingleShaderBindings SingleShaderBindings = Command.ShaderBindings.GetSingleShaderBindings(SF_RayCallable);
+	CallableShader->GetShaderBindings(View, PreshadingUniformBuffer, SingleShaderBindings);
+}
+
+void PreparePathTracingHeterogeneousVolumeMaterial(
+	FRDGBuilder& GraphBuilder,
+	FScene* Scene,
+	const FViewInfo& View,
+	FVolumetricMeshBatch& VolumetricMeshBatch,
+	int32 VolumeInstanceIndex,
+	const FPreshadingUniformBufferParameters& PreshadingUniformBufferParameters,
+	FPathTracingHeterogeneousVolumesCallableInfo& CallableInfo
+)
+{
+	check(View.ViewState);
+
+	if (!ShouldCompileRayTracingCallableShadersForProject(Scene->GetShaderPlatform()))
+	{
+		return;
+	}
+
+	if (!ShouldRenderHeterogeneousVolumes(Scene))
+	{
+		return;
+	}
+
+	const FPrimitiveSceneProxy* PrimitiveSceneProxy = VolumetricMeshBatch.Proxy;
+	if (!PrimitiveSceneProxy->IsHeterogeneousVolume())
+	{
+		return;
+	}
+
+	const FMaterialRenderProxy* MaterialRenderProxy = VolumetricMeshBatch.Mesh->MaterialRenderProxy;
+	if (MaterialRenderProxy == nullptr)
+	{
+		return;
+	}
+
+	const FMaterial* MaterialResource = &MaterialRenderProxy->GetMaterialWithFallback(Scene->GetFeatureLevel(), MaterialRenderProxy);
+	if (MaterialResource->GetMaterialDomain() != MD_Volume)
+	{
+		return;
+	}
+
+	const IHeterogeneousVolumeInterface* HeterogeneousVolumeInterface = (IHeterogeneousVolumeInterface*)VolumetricMeshBatch.Mesh->Elements[VolumeInstanceIndex].UserData;
+	if (HeterogeneousVolumeInterface == nullptr)
+	{
+		return;
+	}
+
+	// Add preshading uniform buffer resource to SBT
+	TUniformBufferRef<FPreshadingUniformBufferParameters> PreshadingUniformBuffer = CreateUniformBufferImmediate(PreshadingUniformBufferParameters, EUniformBufferUsage::UniformBuffer_SingleFrame);
+	Scene->RayTracingSBT.TransientUniformBuffers.Add(PreshadingUniformBuffer);
+
+	const FMaterialShaderMap* MaterialShaderMap = MaterialResource->GetRenderingThreadShaderMap();
+	BindCallableShader<FPathTracingHeterogeneousVolumesIntersectCS>(
+		Scene,
+		View,
+		MaterialRenderProxy,
+		MaterialResource,
+		MaterialShaderMap,
+		PreshadingUniformBuffer,
+		CallableInfo.IntersectShaderId
+	);
+
+	BindCallableShader<FPathTracingHeterogeneousVolumesGetDensityBoundsCS>(
+		Scene,
+		View,
+		MaterialRenderProxy,
+		MaterialResource,
+		MaterialShaderMap,
+		PreshadingUniformBuffer,
+		CallableInfo.GetDensityBoundsShaderId
+	);
+
+	BindCallableShader<FPathTracingHeterogeneousVolumesGetDensityCS>(
+		Scene,
+		View,
+		MaterialRenderProxy,
+		MaterialResource,
+		MaterialShaderMap,
+		PreshadingUniformBuffer,
+		CallableInfo.GetDensityShaderId
+	);
+
+	BindCallableShader<FPathTracingHeterogeneousVolumesSampleDistanceCS>(
+		Scene,
+		View,
+		MaterialRenderProxy,
+		MaterialResource,
+		MaterialShaderMap,
+		PreshadingUniformBuffer,
+		CallableInfo.SampleDistanceShaderId
+	);
+
+	BindCallableShader<FPathTracingHeterogeneousVolumesGetTransmittanceCS>(
+		Scene,
+		View,
+		MaterialRenderProxy,
+		MaterialResource,
+		MaterialShaderMap,
+		PreshadingUniformBuffer,
+		CallableInfo.GetTransmittanceShaderId
+	);
+
+	// TODO: check that callable shader id index corresponds to volume id
+}
+
+enum class EPathTracingHeterogeneousVolumesPipeline
+{
+	Megavoxel,
+	Callable
+};
+
+EPathTracingHeterogeneousVolumesPipeline GetPathTracingHeterogeneousVolumesPipeline()
+{
+	EPathTracingHeterogeneousVolumesPipeline Pipeline = EPathTracingHeterogeneousVolumesPipeline::Megavoxel;
+
+	const bool bUseExperimental = CVarPathTracingExperimental.GetValueOnRenderThread();
+	if (bUseExperimental &&
+		CVarPathTracingHeterogeneousVolumesPipeline.GetValueOnRenderThread() == 1)
+	{
+		Pipeline = EPathTracingHeterogeneousVolumesPipeline::Callable;
+	}
+	return Pipeline;
+}
+
+static TAutoConsoleVariable<int32> CVarHeterogeneousVolumesMajorantFlush(
+	TEXT("r.HeterogeneousVolumes.Majorant.Flush"),
+	0,
+	TEXT("Flushes the majorant cache (Default = 0)"),
+	ECVF_RenderThreadSafe
+);
+
+void PreparePathTracingHeterogeneousVolumes(
+	FRDGBuilder& GraphBuilder,
+	FScene* Scene,
+	const TArrayView<FViewInfo>& Views
+)
+{
+	if (!ShouldCompileRayTracingCallableShadersForProject(Scene->GetShaderPlatform()))
+	{
+		return;
+	}
+
+	if (GetPathTracingHeterogeneousVolumesPipeline() == EPathTracingHeterogeneousVolumesPipeline::Megavoxel)
+	{
+		return;
+	}
+
+	for (int32 ViewIndex = 0; ViewIndex < Views.Num(); ViewIndex++)
+	{
+		FViewInfo& View = Views[ViewIndex];
+		if (View.ViewState == nullptr)
+		{
+			continue;
+		}
+
+		FPathTracingState* PathTracingState = GetPathTracingStateFromView(View);
+
+		// Empty shader ids
+		PathTracingState->HeterogeneousVolumesCallableInfo.Empty();
+
+		if (CVarHeterogeneousVolumesMajorantFlush.GetValueOnAnyThread())
+		{
+			PathTracingState->HeterogeneousVolumesPersistentMajorantRT.Empty();
+		}
+
+		if (ShouldRenderHeterogeneousVolumesForView(View))
+		{
+			int32 VolumeId = 0;
+			for (int32 MeshBatchIndex = 0; MeshBatchIndex < View.HeterogeneousVolumesMeshBatches.Num(); ++MeshBatchIndex)
+			{
+				FVolumetricMeshBatch VolumetricMeshBatch = View.HeterogeneousVolumesMeshBatches[MeshBatchIndex];
+				const FMeshBatch* Mesh = VolumetricMeshBatch.Mesh;
+				const FPrimitiveSceneProxy* PrimitiveSceneProxy = VolumetricMeshBatch.Proxy;
+				const FMaterialRenderProxy* MaterialRenderProxy = Mesh->MaterialRenderProxy;
+				if (!ShouldRenderMeshBatchWithHeterogeneousVolumes(Mesh, PrimitiveSceneProxy, View.GetFeatureLevel()))
+				{
+					continue;
+				}
+
+				for (int32 VolumeInstanceIndex = 0; VolumeInstanceIndex < Mesh->Elements.Num(); ++VolumeInstanceIndex)
+				{
+					const IHeterogeneousVolumeInterface* HeterogeneousVolumeInterface = (IHeterogeneousVolumeInterface*)Mesh->Elements[VolumeInstanceIndex].UserData;
+					//check(HeterogeneousVolumeInterface != nullptr);
+					if (HeterogeneousVolumeInterface == nullptr)
+					{
+						continue;
+					}
+
+					FHeterogeneousVolumesPrimitiveInstanceKey Key;
+					Key.PersistentPrimitiveIndex = PrimitiveSceneProxy->GetPrimitiveSceneInfo()->GetPersistentIndex().Index;
+					Key.InstanceIndex = VolumeInstanceIndex;
+
+					TRefCountPtr<IPooledRenderTarget>& PersistentMajorantRT = PathTracingState->HeterogeneousVolumesPersistentMajorantRT.FindOrAdd(Key);
+					FPathTracingHeterogeneousVolumesCallableInfo& CallableInfo = PathTracingState->HeterogeneousVolumesCallableInfo.AddDefaulted_GetRef();
+
+					FPreshadingUniformBufferParameters PreshadingUniformBufferParameters;
+					BuildPreshadingGrid(GraphBuilder, Scene, View, HeterogeneousVolumeInterface, MaterialRenderProxy,
+						PersistentMajorantRT,
+						CallableInfo,
+						PreshadingUniformBufferParameters);
+					PreparePathTracingHeterogeneousVolumeMaterial(
+						GraphBuilder,
+						Scene,
+						View,
+						VolumetricMeshBatch,
+						VolumeInstanceIndex,
+						PreshadingUniformBufferParameters,
+						CallableInfo);
+				}
+			}
+		}
+
+		// Initialize shader ids if there are no volumes..
+		if (PathTracingState->HeterogeneousVolumesCallableInfo.IsEmpty())
+		{
+			FPathTracingHeterogeneousVolumesCallableInfo& CallableInfo = PathTracingState->HeterogeneousVolumesCallableInfo.AddDefaulted_GetRef();
+			CallableInfo.IntersectShaderId = -1;
+			CallableInfo.GetDensityBoundsShaderId = -1;
+			CallableInfo.GetDensityShaderId = -1;
+			CallableInfo.SampleDistanceShaderId = -1;
+			CallableInfo.GetTransmittanceShaderId = -1;
+		}
+	}
+}
 
 class FPathTracingBuildAdaptiveErrorTextureCS : public FGlobalShader
 {
@@ -3296,6 +3679,7 @@ static FPathTracingRG::FPermutationDomain GetPathTracingRGPermutation(const FVie
 	Out.Set<FPathTracingRG::FSubstrateComplexSpecialMaterial>(bHasComplexSpecialRenderPath);
 	Out.Set<FPathTracingRG::FUseSER>(bUseSER);
 	Out.Set<FPathTracingRG::FNeedTMinWorkaround>(GPUNeedsTMinWorkaround());
+	Out.Set<FPathTracingRG::FHeterogeneousVolumesPipeline>(static_cast<int32>(GetPathTracingHeterogeneousVolumesPipeline()));
 	return Out;
 }
 
@@ -3503,6 +3887,21 @@ DECLARE_GPU_STAT_NAMED(PathTracingPost, TEXT("Path Tracing Post"));
 DECLARE_GPU_STAT_NAMED(PathTracingCopy, TEXT("Path Tracing Copy"));
 #endif
 
+TAutoConsoleVariable<int32> CVarPathTracingVolumeIndex(
+	TEXT("r.PathTracing.VolumeIndex"),
+	0,
+	TEXT("Heterogeneous volume index for rendering"),
+	ECVF_RenderThreadSafe
+);
+
+BEGIN_SHADER_PARAMETER_STRUCT(FHeterogeneousVolumesCallableShadersParameters, )
+	RDG_BUFFER_ACCESS(IntersectShaderIds, ERHIAccess::CopyDest)
+	RDG_BUFFER_ACCESS(GetDensityBoundsShaderIds, ERHIAccess::CopyDest)
+	RDG_BUFFER_ACCESS(GetDensityShaderIds, ERHIAccess::CopyDest)
+	RDG_BUFFER_ACCESS(SampleDistanceShaderIds, ERHIAccess::CopyDest)
+	RDG_BUFFER_ACCESS(GetTransmittanceShaderIds, ERHIAccess::CopyDest)
+END_SHADER_PARAMETER_STRUCT()
+
 void FDeferredShadingSceneRenderer::RenderPathTracing(
 	FRDGBuilder& GraphBuilder,
 	const FViewInfo& View,
@@ -3930,8 +4329,120 @@ void FDeferredShadingSceneRenderer::RenderPathTracing(
 	// Declare heterogeneous volume buffers
 	TRDGUniformBufferRef<FOrthoVoxelGridUniformBufferParameters> OrthoGridUniformBuffer;
 	TRDGUniformBufferRef<FFrustumVoxelGridUniformBufferParameters> FrustumGridUniformBuffer;
+	// NOTE: Remove in separate submission
+	TRDGUniformBufferRef<FPreshadingUniformBufferParameters> PreshadingUniformBuffer = nullptr;
 	bool bCreateVolumeGrids = false;
 
+	// Upload heterogeneous volume callable interface
+	bool bCallableShaderHeterogeneousVolumesPipeline = GetPathTracingHeterogeneousVolumesPipeline() == EPathTracingHeterogeneousVolumesPipeline::Callable;
+	int32 ElementCount = bCallableShaderHeterogeneousVolumesPipeline ? FMath::Max(PathTracingState->HeterogeneousVolumesCallableInfo.Num(), 1) : 1;
+	FRDGBufferDesc ShaderIdsDesc = FRDGBufferDesc::CreateBufferDesc(sizeof(int32), ElementCount);
+	FRDGBufferRef IntersectShaderIdsBuffer = GraphBuilder.CreateBuffer(ShaderIdsDesc, TEXT("IntersectShaderId"));
+	FRDGBufferRef GetDensityBoundsShaderIdsBuffer = GraphBuilder.CreateBuffer(ShaderIdsDesc, TEXT("GetDensityBoundsShaderId"));
+	FRDGBufferRef GetDensityShaderIdsBuffer = GraphBuilder.CreateBuffer(ShaderIdsDesc, TEXT("GetDensityShaderId"));
+	FRDGBufferRef SampleDistanceShaderIdsBuffer = GraphBuilder.CreateBuffer(ShaderIdsDesc, TEXT("SampleDistanceShaderId"));
+	FRDGBufferRef GetTransmittanceShaderIdsBuffer = GraphBuilder.CreateBuffer(ShaderIdsDesc, TEXT("GetTransmittanceShaderId"));
+	if (bCallableShaderHeterogeneousVolumesPipeline)
+	{
+		FHeterogeneousVolumesCallableShadersParameters* HeterogeneousVolumesCallableShadersParameters = GraphBuilder.AllocParameters<FHeterogeneousVolumesCallableShadersParameters>();
+		{
+			HeterogeneousVolumesCallableShadersParameters->IntersectShaderIds = IntersectShaderIdsBuffer;
+			HeterogeneousVolumesCallableShadersParameters->GetDensityBoundsShaderIds = GetDensityBoundsShaderIdsBuffer;
+			HeterogeneousVolumesCallableShadersParameters->GetDensityShaderIds = GetDensityShaderIdsBuffer;
+			HeterogeneousVolumesCallableShadersParameters->SampleDistanceShaderIds = SampleDistanceShaderIdsBuffer;
+			HeterogeneousVolumesCallableShadersParameters->GetTransmittanceShaderIds = GetTransmittanceShaderIdsBuffer;
+		}
+
+		GraphBuilder.AddPass(
+			RDG_EVENT_NAME("UploadShaderIds"),
+			HeterogeneousVolumesCallableShadersParameters,
+			ERDGPassFlags::Copy | ERDGPassFlags::NeverCull,
+			[
+				ShaderIdsDesc,
+				HeterogeneousVolumesCallableShadersParameters,
+				PathTracingState
+			](FRHICommandList& RHICmdList)
+			{
+				int32 SizeInBytes = ShaderIdsDesc.NumElements * ShaderIdsDesc.BytesPerElement;
+
+				// Intersect
+				{
+					// TODO: Convert to SoA style
+					TArray<int32> IntersectShaderIds;
+					IntersectShaderIds.Reserve(ShaderIdsDesc.NumElements);
+					for (uint32 i = 0; i < ShaderIdsDesc.NumElements; ++i)
+					{
+						IntersectShaderIds.Add(PathTracingState->HeterogeneousVolumesCallableInfo[i].IntersectShaderId);
+					}
+
+					void* Dest = RHICmdList.LockBuffer(HeterogeneousVolumesCallableShadersParameters->IntersectShaderIds->GetRHI(), 0, SizeInBytes, RLM_WriteOnly);
+					FPlatformMemory::Memcpy(Dest, IntersectShaderIds.GetData(), SizeInBytes);
+					RHICmdList.UnlockBuffer(HeterogeneousVolumesCallableShadersParameters->IntersectShaderIds->GetRHI());
+				}
+
+				// GetDensityBounds
+				{
+					// TODO: Convert to SoA style
+					TArray<int32> GetDensityBoundsShaderId;
+					GetDensityBoundsShaderId.Reserve(ShaderIdsDesc.NumElements);
+					for (uint32 i = 0; i < ShaderIdsDesc.NumElements; ++i)
+					{
+						GetDensityBoundsShaderId.Add(PathTracingState->HeterogeneousVolumesCallableInfo[i].GetDensityBoundsShaderId);
+					}
+
+					void* Dest = RHICmdList.LockBuffer(HeterogeneousVolumesCallableShadersParameters->GetDensityBoundsShaderIds->GetRHI(), 0, SizeInBytes, RLM_WriteOnly);
+					FPlatformMemory::Memcpy(Dest, GetDensityBoundsShaderId.GetData(), SizeInBytes);
+					RHICmdList.UnlockBuffer(HeterogeneousVolumesCallableShadersParameters->GetDensityBoundsShaderIds->GetRHI());
+				}
+
+				// GetDensity
+				{
+					// TODO: Convert to SoA style
+					TArray<int32> GetDensityShaderId;
+					GetDensityShaderId.Reserve(ShaderIdsDesc.NumElements);
+					for (uint32 i = 0; i < ShaderIdsDesc.NumElements; ++i)
+					{
+						GetDensityShaderId.Add(PathTracingState->HeterogeneousVolumesCallableInfo[i].GetDensityShaderId);
+					}
+
+					void* Dest = RHICmdList.LockBuffer(HeterogeneousVolumesCallableShadersParameters->GetDensityShaderIds->GetRHI(), 0, SizeInBytes, RLM_WriteOnly);
+					FPlatformMemory::Memcpy(Dest, GetDensityShaderId.GetData(), SizeInBytes);
+					RHICmdList.UnlockBuffer(HeterogeneousVolumesCallableShadersParameters->GetDensityShaderIds->GetRHI());
+				}
+
+				// SampleDistance
+				{
+					// TODO: Convert to SoA style
+					TArray<int32> SampleDistanceShaderId;
+					SampleDistanceShaderId.Reserve(ShaderIdsDesc.NumElements);
+					for (uint32 i = 0; i < ShaderIdsDesc.NumElements; ++i)
+					{
+						SampleDistanceShaderId.Add(PathTracingState->HeterogeneousVolumesCallableInfo[i].SampleDistanceShaderId);
+					}
+
+					void* Dest = RHICmdList.LockBuffer(HeterogeneousVolumesCallableShadersParameters->SampleDistanceShaderIds->GetRHI(), 0, SizeInBytes, RLM_WriteOnly);
+					FPlatformMemory::Memcpy(Dest, SampleDistanceShaderId.GetData(), SizeInBytes);
+					RHICmdList.UnlockBuffer(HeterogeneousVolumesCallableShadersParameters->SampleDistanceShaderIds->GetRHI());
+				}
+
+				//  GetTransmittance
+				{
+					// TODO: Convert to SoA style
+					TArray<int32> GetTransmittanceShaderId;
+					GetTransmittanceShaderId.Reserve(ShaderIdsDesc.NumElements);
+					for (uint32 i = 0; i < ShaderIdsDesc.NumElements; ++i)
+					{
+						GetTransmittanceShaderId.Add(PathTracingState->HeterogeneousVolumesCallableInfo[i].GetTransmittanceShaderId);
+					}
+
+					void* Dest = RHICmdList.LockBuffer(HeterogeneousVolumesCallableShadersParameters->GetTransmittanceShaderIds->GetRHI(), 0, SizeInBytes, RLM_WriteOnly);
+					FPlatformMemory::Memcpy(Dest, GetTransmittanceShaderId.GetData(), SizeInBytes);
+					RHICmdList.UnlockBuffer(HeterogeneousVolumesCallableShadersParameters->GetTransmittanceShaderIds->GetRHI());
+				}
+			}
+		);
+	}
+
 	// Prepare radiance buffer (will be shared with display pass)
 	FRDGTexture* RadianceTexture = nullptr;
 	FRDGTexture* VarianceTexture = nullptr;
@@ -4036,13 +4547,18 @@ void FDeferredShadingSceneRenderer::RenderPathTracing(
 		if (bNeedsMoreRays)
 		{
 			RDG_EVENT_SCOPE(GraphBuilder, "Path Tracing Compute (%d x %d)", DispatchResX, DispatchResY);
+			// Combined interface
 			const bool bForceRebuild = CVarPathTracingHeterogeneousVolumesRebuildEveryFrame.GetValueOnRenderThread();
 			bCreateVolumeGrids = bForceRebuild ||
 				!PathTracingState->AdaptiveFrustumGridParameterCache.TopLevelGridBuffer ||
 				!PathTracingState->AdaptiveOrthoGridParameterCache.TopLevelGridBuffer;
 			if (bCreateVolumeGrids)
 			{
+				bool bMegavoxelPipeline = GetPathTracingHeterogeneousVolumesPipeline() == EPathTracingHeterogeneousVolumesPipeline::Megavoxel;
+
 				FVoxelGridBuildOptions BuildOptions;
+				BuildOptions.bBuildFrustumGrid &= bMegavoxelPipeline;
+				BuildOptions.bBuildOrthoGrid &= bMegavoxelPipeline;
 				BuildOrthoVoxelGrid(GraphBuilder, Scene, Views, VisibleLightInfos, BuildOptions, OrthoGridUniformBuffer);
 				BuildFrustumVoxelGrid(GraphBuilder, Scene, Views[0], BuildOptions, FrustumGridUniformBuffer);
 			}
@@ -4059,6 +4575,45 @@ void FDeferredShadingSceneRenderer::RenderPathTracing(
 				);
 			}
 
+			// Preshading interface
+			// NOTE: Removing in separate submission
+			if (false)
+			{
+				RDG_EVENT_SCOPE(GraphBuilder, "Heterogeneous Volumes: Preshading Pipeline");
+
+				TArray<const IHeterogeneousVolumeInterface*> HeterogeneousVolumeInterfaces;
+				TArray<const FMaterialRenderProxy*> MaterialRenderProxies;
+				CollectRenderableHeterogeneousVolumes(
+					GraphBuilder,
+					Scene,
+					View,
+					HeterogeneousVolumeInterfaces,
+					MaterialRenderProxies
+				);
+
+				if (!HeterogeneousVolumeInterfaces.IsEmpty())
+				{
+					FHeterogeneousVolumesPrimitiveInstanceKey Key;
+					Key.PersistentPrimitiveIndex = 0;
+					Key.InstanceIndex = 0;
+
+					TRefCountPtr<IPooledRenderTarget>& PersistentMajorantRT = PathTracingState->HeterogeneousVolumesPersistentMajorantRT.FindOrAdd(Key);
+					BuildPreshadingGrid(
+						GraphBuilder,
+						Scene,
+						View,
+						HeterogeneousVolumeInterfaces[0],
+						MaterialRenderProxies[0],
+						PersistentMajorantRT,
+						PreshadingUniformBuffer
+					);
+				}
+				else
+				{
+					BuildEmptyPreshadingUniformBuffer(GraphBuilder, PreshadingUniformBuffer);
+				}
+			}
+
 			FRDGTexture* CloudAccelerationMap = PathTracingState->CloudAccelerationMap.IsValid() ? GraphBuilder.RegisterExternalTexture(PathTracingState->CloudAccelerationMap, TEXT("PathTracer.CloudAccelerationMap")) : nullptr;
 
 			const bool bEnableClouds = (Config.PathTracingData.VolumeFlags & PATH_TRACER_VOLUME_ENABLE_CLOUDS) != 0 && CloudAccelerationMap != nullptr;
@@ -4303,6 +4858,19 @@ void FDeferredShadingSceneRenderer::RenderPathTracing(
 							// Heterogeneous volume bindings
 							PassParameters->OrthoGridUniformBuffer = OrthoGridUniformBuffer;
 							PassParameters->FrustumGridUniformBuffer = FrustumGridUniformBuffer;
+							// NOTE: Remove in separate submission
+							//PassParameters->PreshadingUniformBuffer = PreshadingUniformBuffer;
+
+							// Heterogeneous volume callable interface
+							if (bCallableShaderHeterogeneousVolumesPipeline)
+							{
+								PassParameters->HeterogeneousVolumesNum = PathTracingState->HeterogeneousVolumesCallableInfo.Num();
+								PassParameters->HeterogeneousVolumesCallableShaderIdIntersect = GraphBuilder.CreateSRV(IntersectShaderIdsBuffer, PF_R32_SINT);
+								PassParameters->HeterogeneousVolumesCallableShaderIdGetDensityBounds = GraphBuilder.CreateSRV(GetDensityBoundsShaderIdsBuffer, PF_R32_SINT);
+								PassParameters->HeterogeneousVolumesCallableShaderIdGetDensity = GraphBuilder.CreateSRV(GetDensityShaderIdsBuffer, PF_R32_SINT);
+								PassParameters->HeterogeneousVolumesCallableShaderIdSampleDistance = GraphBuilder.CreateSRV(SampleDistanceShaderIdsBuffer, PF_R32_SINT);
+								PassParameters->HeterogeneousVolumesCallableShaderIdGetTransmittance = GraphBuilder.CreateSRV(GetTransmittanceShaderIdsBuffer, PF_R32_SINT);
+							}
 
 							PassParameters->TilePixelOffset.X = TileX;
 							PassParameters->TilePixelOffset.Y = TileY + CurrentGPU;
diff --git a/Engine/Source/Runtime/Renderer/Private/PathTracing.h b/Engine/Source/Runtime/Renderer/Private/PathTracing.h
--- a/Engine/Source/Runtime/Renderer/Private/PathTracing.h
+++ b/Engine/Source/Runtime/Renderer/Private/PathTracing.h
@@ -33,6 +33,7 @@ RENDERER_API FRHIRayTracingShader* GetGPULightmassDefaultHiddenHitShader(const F
 
 void PreparePathTracingRTPSO();
 void PreparePathTracingCloudMaterial(FRDGBuilder& GraphBuilder, FScene* Scene, TArrayView<FViewInfo> Views);
+void PreparePathTracingHeterogeneousVolumes(FRDGBuilder& GraphBuilder, FScene* Scene, const TArrayView<FViewInfo>& Views);
 
 namespace PathTracing
 {
