[HWRT] Fix stale LightingChannelMask in FPrimitiveRayTracingData and cached FRayTracingGeometryInstances.
- Need to update cached value when SetLightingChannels_GameThread(...) is called.

#Synced-CL 51376035

diff --git a/Engine/Source/Runtime/Engine/Private/PrimitiveSceneProxy.cpp b/Engine/Source/Runtime/Engine/Private/PrimitiveSceneProxy.cpp
--- a/Engine/Source/Runtime/Engine/Private/PrimitiveSceneProxy.cpp
+++ b/Engine/Source/Runtime/Engine/Private/PrimitiveSceneProxy.cpp
@@ -1225,6 +1225,8 @@ void FPrimitiveSceneProxy::SetLightingChannels_GameThread(FLightingChannels Ligh
 	{
 		PrimitiveSceneProxy->LightingChannelMask = LocalLightingChannelMask;
 
+		Scene->UpdateCachedRayTracingState(PrimitiveSceneProxy, EPrimitiveDirtyRayTracingState::Default | EPrimitiveDirtyRayTracingState::LightingChannelMask);
+
 		if (PrimitiveSceneInfo)
 		{
 			Scene->RequestUniformBufferUpdate(*PrimitiveSceneInfo);
diff --git a/Engine/Source/Runtime/Engine/Public/SceneInterface.h b/Engine/Source/Runtime/Engine/Public/SceneInterface.h
--- a/Engine/Source/Runtime/Engine/Public/SceneInterface.h
+++ b/Engine/Source/Runtime/Engine/Public/SceneInterface.h
@@ -104,9 +104,10 @@ enum class EPrimitiveDirtyRayTracingState
 	None = 0,
 	MeshCommandsAndInstances = 1 << 0,
 	DrawInGame = 1 << 1,
+	LightingChannelMask = 1 << 2,
 
 	Default = MeshCommandsAndInstances,
-	All = MeshCommandsAndInstances | DrawInGame
+	All = MeshCommandsAndInstances | DrawInGame | LightingChannelMask
 };
 ENUM_CLASS_FLAGS(EPrimitiveDirtyRayTracingState);
 
diff --git a/Engine/Source/Runtime/Renderer/Private/RendererScene.cpp b/Engine/Source/Runtime/Renderer/Private/RendererScene.cpp
--- a/Engine/Source/Runtime/Renderer/Private/RendererScene.cpp
+++ b/Engine/Source/Runtime/Renderer/Private/RendererScene.cpp
@@ -4236,6 +4236,11 @@ void FScene::UpdateCachedRayTracingState(FPrimitiveSceneProxy* SceneProxy, EPrim
 				PrimitiveRayTracingDatas[PrimitiveIndex].bDrawInGame = SceneProxy->IsDrawnInGame();
 			}
 
+			if (EnumHasAnyFlags(PrimitiveDirtyState, EPrimitiveDirtyRayTracingState::LightingChannelMask))
+			{
+				PrimitiveRayTracingDatas[PrimitiveIndex].LightingChannelMask = SceneProxy->GetLightingChannelMask();
+			}
+
 			if (EnumHasAnyFlags(PrimitiveDirtyState, EPrimitiveDirtyRayTracingState::MeshCommandsAndInstances))
 			{
 				PrimitiveRayTracingDatas[PrimitiveIndex].bCachedRaytracingDataDirty = true;
