[Front Layer Translucency] Rename r.Lumen.TranslucencyReflections.FrontLayer.DepthThreshold to r.FrontLayerTranslucency.DepthThreshold.

#Synced-CL 51399018

diff --git a/Engine/Source/Runtime/Renderer/Private/FrontLayerTranslucency.cpp b/Engine/Source/Runtime/Renderer/Private/FrontLayerTranslucency.cpp
--- a/Engine/Source/Runtime/Renderer/Private/FrontLayerTranslucency.cpp
+++ b/Engine/Source/Runtime/Renderer/Private/FrontLayerTranslucency.cpp
@@ -43,13 +43,15 @@ FAutoConsoleVariableRef CVarLumenTranslucencyReflectionsFrontLayerAllowed(
 	ECVF_Scalability | ECVF_RenderThreadSafe
 );
 
-static TAutoConsoleVariable<float> CVarLumenFrontLayerDepthThreshold(
-	TEXT("r.Lumen.TranslucencyReflections.FrontLayer.DepthThreshold"),
+static TAutoConsoleVariable<float> CVarFrontLayerTranslucencyDepthThreshold(
+	TEXT("r.FrontLayerTranslucency.DepthThreshold"),
 	1024.0f,
-	TEXT("Depth test threshold used to determine whether the fragments being rendered match the single layer that reflections were calculated for. In float ULP units."),
+	TEXT("Depth test threshold used to determine whether the fragments being rendered match the single layer of translucency for which lighting was calculated in a dedicated pass. In float ULP units."),
 	ECVF_Scalability | ECVF_RenderThreadSafe
 );
 
+FAutoConsoleVariableDeprecated CVarLumenFrontLayerDepthThreshold(TEXT("r.Lumen.TranslucencyReflections.FrontLayer.DepthThreshold"), TEXT("r.FrontLayerTranslucency.DepthThreshold"), TEXT("5.8"));
+
 bool IsVSMTranslucentHighQualityEnabled();
 
 namespace Lumen
@@ -671,7 +673,7 @@ FFrontLayerTranslucencyData FDeferredShadingSceneRenderer::RenderFrontLayerTrans
 			View.FrontLayerTranslucency.GBufferC = Out.GBufferC;
 			View.FrontLayerTranslucency.SceneDepth = Out.SceneDepth;
 
-			View.FrontLayerTranslucency.RelativeDepthThreshold = CVarLumenFrontLayerDepthThreshold.GetValueOnRenderThread();
+			View.FrontLayerTranslucency.RelativeDepthThreshold = CVarFrontLayerTranslucencyDepthThreshold.GetValueOnRenderThread();
 		}
 	}
 	return Out;
