Path Tracer: Implement support for "Visible in Reflections" flag on primitives

Remove hair specific ray masks from the path tracer. In its place, use ray flags to skip procedural primitives when far from the origin. This achieves a similar effect of improving performance for rays originating far from the origin reducing performance when traversing hair.

Implement support for the flag in "Primary Rays" mode of path tracing debug for consistency.

Refactor instance mask logic a bit to consolidate identical code fragments.

#Synced-CL 51391139

diff --git a/Engine/Shaders/Private/PathTracing/PathTracingCommon.ush b/Engine/Shaders/Private/PathTracing/PathTracingCommon.ush
--- a/Engine/Shaders/Private/PathTracing/PathTracingCommon.ush
+++ b/Engine/Shaders/Private/PathTracing/PathTracingCommon.ush
@@ -59,6 +59,11 @@
 // In ordinary scenes, the most likely occurance of this problem is when combining hair and volumetric atmosphere as rays from several Km can travel back towards the hair and cause long delays or even TDRs
 #define PATH_TRACING_SKIP_HAIR_DISTANCE      1e6
 
+uint GetPathTracingHairRayFlags(float3 RayOrigin)
+{
+    return length2(RayOrigin) < Pow2(PATH_TRACING_SKIP_HAIR_DISTANCE) ? 0 : RAY_FLAG_SKIP_PROCEDURAL_PRIMITIVES;
+}
+
 #define SUBSTRATE_PATH_TRACING_SPECULAR_PROFILE	(SUBSTRATE_ENABLED && PLATFORM_ENABLES_SUBSTRATE_SPECULAR_PROFILE)
 #define SUBSTRATE_PATH_TRACING_GLINTS 			(SUBSTRATE_ENABLED && PLATFORM_ENABLES_SUBSTRATE_GLINTS)
 
diff --git a/Engine/Shaders/Private/PathTracing/PathTracingCore.ush b/Engine/Shaders/Private/PathTracing/PathTracingCore.ush
--- a/Engine/Shaders/Private/PathTracing/PathTracingCore.ush
+++ b/Engine/Shaders/Private/PathTracing/PathTracingCore.ush
@@ -214,6 +214,36 @@ struct FPathState
 		return FirstScatterType == PATHTRACER_SCATTER_DIFFUSE || FirstScatterType == PATHTRACER_SCATTER_SPECULAR;
 	}
 
+    uint GetInstanceInclusionMaskOpaque()
+    {
+        uint FirstScatterType = GetFirstScatterType();
+        if (FirstScatterType == PATHTRACER_SCATTER_CAMERA ||
+            FirstScatterType == PATHTRACER_SCATTER_REFRACT)
+        {
+            return PATHTRACER_MASK_CAMERA_OPAQUE;
+        }
+        if (FirstScatterType == PATHTRACER_SCATTER_SPECULAR)
+        {
+            return PATHTRACER_MASK_REFLECTION_OPAQUE;
+        }
+        return PATHTRACER_MASK_INDIRECT_OPAQUE;
+    }
+
+    uint GetInstanceInclusionMaskTranslucent()
+    {
+        uint FirstScatterType = GetFirstScatterType();
+        if (FirstScatterType == PATHTRACER_SCATTER_CAMERA ||
+            FirstScatterType == PATHTRACER_SCATTER_REFRACT)
+        {
+            return PATHTRACER_MASK_CAMERA_TRANSLUCENT;
+        }
+        if (FirstScatterType == PATHTRACER_SCATTER_SPECULAR)
+        {
+            return PATHTRACER_MASK_REFLECTION_TRANSLUCENT;
+        }
+        return PATHTRACER_MASK_INDIRECT_TRANSLUCENT;
+    }
+
 	bool UpdateScatterType(uint ScatterType)
 	{
 		// Keep track of the first non-refractive scatter type, and leave it "locked" beyond this
@@ -421,7 +451,7 @@ void EvaluateDecals(FPathState PathState, float HitT, inout FDecalShaderPayload
 		DecalRay.TMin = 0.0f;
 		DecalRay.TMax = HitT + MeshDecalBias;
 
-		const uint InstanceInclusionMask = PathState.GetFirstScatterType() == PATHTRACER_SCATTER_CAMERA ? (PATHTRACER_MASK_CAMERA | PATHTRACER_MASK_CAMERA_TRANSLUCENT) : (PATHTRACER_MASK_INDIRECT | PATHTRACER_MASK_INDIRECT_TRANSLUCENT);
+        const uint InstanceInclusionMask = PathState.GetInstanceInclusionMaskOpaque() | PathState.GetInstanceInclusionMaskTranslucent();
 		const uint MissShaderIndex = 0; // TODO
 
 		for (;;)
@@ -475,17 +505,11 @@ FPackedPathTracingPayload TraceTransparentRay(inout FPathState PathState, int Bo
 	const bool bLastBounce = Bounce == MaxBounces;
 #if PATH_TRACER_KERNEL_TYPE == PATH_TRACER_KERNEL_CAMERA
 	// backface culling on camera rays only to match raster but still allow refraction and indirect to work properly
-	const uint RayFlags = RAY_FLAG_CULL_BACK_FACING_TRIANGLES;
+	const uint RayFlags = RAY_FLAG_CULL_BACK_FACING_TRIANGLES | GetPathTracingHairRayFlags(PathState.Ray.Origin);
 #else
-	const uint RayFlags = 0;
+	const uint RayFlags = GetPathTracingHairRayFlags(PathState.Ray.Origin);
 #endif
-	const bool bUseCameraRayType = !PathState.HasMadeContributionScatter();
-	uint InstanceInclusionMask = bUseCameraRayType ? PATHTRACER_MASK_CAMERA : PATHTRACER_MASK_INDIRECT;
-	if (length2(PathState.Ray.Origin) < Pow2(PATH_TRACING_SKIP_HAIR_DISTANCE))
-	{
-		InstanceInclusionMask |= bUseCameraRayType ? PATHTRACER_MASK_HAIR_CAMERA : PATHTRACER_MASK_HAIR_INDIRECT;
-	}
-
+    const uint InstanceInclusionMaskOpaque = PathState.GetInstanceInclusionMaskOpaque();
 
 	const uint MissShaderIndex = 0;
 	float3 RISRandSample = PathState.RandSequence.Get3D();
@@ -513,7 +537,7 @@ FPackedPathTracingPayload TraceTransparentRay(inout FPathState PathState, int Bo
 		NvTraceRayHitObject(
 			TLAS,
 			RayFlags,
-			InstanceInclusionMask,
+			InstanceInclusionMaskOpaque,
 			RAY_TRACING_SHADER_SLOT_MATERIAL,
 			RAY_TRACING_NUM_SHADER_SLOTS,
 			MissShaderIndex,
@@ -530,7 +554,7 @@ FPackedPathTracingPayload TraceTransparentRay(inout FPathState PathState, int Bo
 	TraceRay(
 		TLAS,
 		RayFlags,
-		InstanceInclusionMask,
+		InstanceInclusionMaskOpaque,
 		RAY_TRACING_SHADER_SLOT_MATERIAL,
 		RAY_TRACING_NUM_SHADER_SLOTS,
 		MissShaderIndex,
@@ -557,7 +581,7 @@ FPackedPathTracingPayload TraceTransparentRay(inout FPathState PathState, int Bo
 				TraceRay(
 					TLAS,
 					RayFlags,
-					InstanceInclusionMask,
+					InstanceInclusionMaskOpaque,
 					RAY_TRACING_SHADER_SLOT_MATERIAL,
 					RAY_TRACING_NUM_SHADER_SLOTS,
 					MissShaderIndex,
@@ -577,8 +601,6 @@ FPackedPathTracingPayload TraceTransparentRay(inout FPathState PathState, int Bo
 		ApplyDecal(DecalPayload, OpaquePayload, PathState.Ray.Direction, PathState.GetPathRoughness());
 	}
 
-	InstanceInclusionMask = bUseCameraRayType ? PATHTRACER_MASK_CAMERA_TRANSLUCENT : PATHTRACER_MASK_INDIRECT_TRANSLUCENT;
-
 	if (OpaquePayload.IsHit())
 	{
 		PathState.Ray.TMax = asfloat(asuint(OpaquePayload.HitT) - 1);
@@ -599,13 +621,15 @@ FPackedPathTracingPayload TraceTransparentRay(inout FPathState PathState, int Bo
 			PackedPayload.SetSceneDepth(dot(TranslatedWorldPos, ViewZ));
 		}
 
+        const uint InstanceInclusionMaskTranslucent = PathState.GetInstanceInclusionMaskTranslucent();
+
 #if PATH_TRACER_USE_SER 
 		{
 			NvHitObject Hit;
 			NvTraceRayHitObject(
 				TLAS,
 				RayFlags,
-				InstanceInclusionMask,
+				InstanceInclusionMaskTranslucent,
 				RAY_TRACING_SHADER_SLOT_MATERIAL,
 				RAY_TRACING_NUM_SHADER_SLOTS,
 				MissShaderIndex,
@@ -622,7 +646,7 @@ FPackedPathTracingPayload TraceTransparentRay(inout FPathState PathState, int Bo
 		TraceRay(
 			TLAS,
 			RayFlags,
-			InstanceInclusionMask,
+			InstanceInclusionMaskTranslucent,
 			RAY_TRACING_SHADER_SLOT_MATERIAL,
 			RAY_TRACING_NUM_SHADER_SLOTS,
 			MissShaderIndex,
@@ -1031,8 +1055,8 @@ float3 TraceTransparentVisibilityRay(FRayDesc Ray, float PathRoughness, float Pr
 		return 1.0;
 	}
 
-	const uint RayFlags = RAY_FLAG_ACCEPT_FIRST_HIT_AND_END_SEARCH | RAY_FLAG_SKIP_CLOSEST_HIT_SHADER;
-	const uint InstanceInclusionMask = PATHTRACER_MASK_SHADOW | (length2(Ray.Origin) < Pow2(PATH_TRACING_SKIP_HAIR_DISTANCE) ? PATHTRACER_MASK_HAIR_SHADOW : 0);
+    const uint RayFlags = RAY_FLAG_ACCEPT_FIRST_HIT_AND_END_SEARCH | RAY_FLAG_SKIP_CLOSEST_HIT_SHADER | GetPathTracingHairRayFlags(Ray.Origin);
+	const uint InstanceInclusionMask = PATHTRACER_MASK_SHADOW;
 	const uint RayContributionToHitGroupIndex = RAY_TRACING_SHADER_SLOT_SHADOW;
 	const uint MultiplierForGeometryContributionToShaderIndex = RAY_TRACING_NUM_SHADER_SLOTS;
 
@@ -1117,11 +1141,8 @@ FProbeResult TraceSSSProbeRay(FRayDesc Ray, inout int InterfaceCounter)
 		// TODO: use smaller payload + flag to skip shading work? (only want smooth+geo normals)
 		//       however supporting blockers would require some shading ...
 		//       
-		const uint RayFlags = 0;
-		const uint InstanceInclusionMask = PATHTRACER_MASK_ALL
-			- PATHTRACER_MASK_HAIR_SHADOW
-			- PATHTRACER_MASK_HAIR_INDIRECT
-			- PATHTRACER_MASK_HAIR_CAMERA; // Ignore hair strands for SSS rays
+		const uint RayFlags = RAY_FLAG_SKIP_PROCEDURAL_PRIMITIVES; // ignore procedural hair
+		const uint InstanceInclusionMask = PATHTRACER_MASK_SSS_RANDOM_WALK;
 		const uint MissShaderIndex = 0;
 		TraceRay(
 			TLAS,
diff --git a/Engine/Shaders/Private/PathTracing/PathTracingDebug.usf b/Engine/Shaders/Private/PathTracing/PathTracingDebug.usf
--- a/Engine/Shaders/Private/PathTracing/PathTracingDebug.usf
+++ b/Engine/Shaders/Private/PathTracing/PathTracingDebug.usf
@@ -79,10 +79,11 @@ float3 TraceShadowRay(FRayDesc Ray, float PathRoughness, float PreviousPathRough
 		Ray.TMin = POSITIVE_INFINITY;
 		Ray.TMax = POSITIVE_INFINITY;
 	}
-	const uint InstanceInclusionMask = PATHTRACER_MASK_SHADOW | (length2(Ray.Origin) < Pow2(PATH_TRACING_SKIP_HAIR_DISTANCE) ? PATHTRACER_MASK_HAIR_SHADOW : 0);
+	const uint RayFlags = RAY_FLAG_ACCEPT_FIRST_HIT_AND_END_SEARCH | RAY_FLAG_SKIP_CLOSEST_HIT_SHADER | GetPathTracingHairRayFlags(Ray.Origin);
+	const uint InstanceInclusionMask = PATHTRACER_MASK_SHADOW;
 	TraceRay(
 		TLAS,
-		RAY_FLAG_ACCEPT_FIRST_HIT_AND_END_SEARCH | RAY_FLAG_SKIP_CLOSEST_HIT_SHADER,
+		RayFlags,
 		InstanceInclusionMask,
 		RAY_TRACING_SHADER_SLOT_SHADOW,
 		RAY_TRACING_NUM_SHADER_SLOTS,
@@ -261,6 +262,36 @@ struct FPathState
 		return FirstScatterType == PATHTRACER_SCATTER_DIFFUSE || FirstScatterType == PATHTRACER_SCATTER_SPECULAR;
 	}
 
+    uint GetInstanceInclusionMaskOpaque()
+    {
+        uint FirstScatterType = GetFirstScatterType();
+        if (FirstScatterType == PATHTRACER_SCATTER_CAMERA ||
+            FirstScatterType == PATHTRACER_SCATTER_REFRACT)
+        {
+            return PATHTRACER_MASK_CAMERA_OPAQUE;
+        }
+        if (FirstScatterType == PATHTRACER_SCATTER_SPECULAR)
+        {
+            return PATHTRACER_MASK_REFLECTION_OPAQUE;
+        }
+        return PATHTRACER_MASK_INDIRECT_OPAQUE;
+    }
+
+    uint GetInstanceInclusionMaskTranslucent()
+    {
+        uint FirstScatterType = GetFirstScatterType();
+        if (FirstScatterType == PATHTRACER_SCATTER_CAMERA ||
+            FirstScatterType == PATHTRACER_SCATTER_REFRACT)
+        {
+            return PATHTRACER_MASK_CAMERA_TRANSLUCENT;
+        }
+        if (FirstScatterType == PATHTRACER_SCATTER_SPECULAR)
+        {
+            return PATHTRACER_MASK_REFLECTION_TRANSLUCENT;
+        }
+        return PATHTRACER_MASK_INDIRECT_TRANSLUCENT;
+    }
+
 	bool UpdateScatterType(uint ScatterType)
 	{
 		// Keep track of the first non-refractive scatter type, and leave it "locked" beyond this
@@ -387,17 +418,12 @@ float3 AnisoDebugColor(float InAniso)
 FPackedPathTracingPayload TraceTransparentRay(inout FPathState PathState, int Bounce)
 {
 #if PATH_TRACER_KERNEL_TYPE == PATH_TRACER_KERNEL_CAMERA
-	const uint RayFlags = RAY_FLAG_CULL_BACK_FACING_TRIANGLES;
+	const uint RayFlags = RAY_FLAG_CULL_BACK_FACING_TRIANGLES | GetPathTracingHairRayFlags(PathState.Ray.Origin);
 #else
-	const uint RayFlags = 0;
+	const uint RayFlags = GetPathTracingHairRayFlags(PathState.Ray.Origin);
 #endif
 	const uint MissShaderIndex = 0;
-	const bool bUseCameraRayType = !PathState.HasMadeContributionScatter();
-	uint InstanceInclusionMask = bUseCameraRayType ? PATHTRACER_MASK_CAMERA : PATHTRACER_MASK_INDIRECT;
-	if (length2(PathState.Ray.Origin) < Pow2(PATH_TRACING_SKIP_HAIR_DISTANCE))
-	{
-		InstanceInclusionMask |= bUseCameraRayType ? PATHTRACER_MASK_HAIR_CAMERA : PATHTRACER_MASK_HAIR_INDIRECT;
-	}
+    const uint InstanceInclusionMaskOpaque = PathState.GetInstanceInclusionMaskOpaque();
 	const float2 RandSample = PathState.RandSequence.Get2D();
 
 	FRISContext HitSample = InitRISContext(RandSample.x);
@@ -417,7 +443,7 @@ FPackedPathTracingPayload TraceTransparentRay(inout FPathState PathState, int Bo
 		NvTraceRayHitObject(
 			TLAS,
 			RayFlags,
-			InstanceInclusionMask,
+			InstanceInclusionMaskOpaque,
 			RAY_TRACING_SHADER_SLOT_MATERIAL,
 			RAY_TRACING_NUM_SHADER_SLOTS,
 			MissShaderIndex,
@@ -433,7 +459,7 @@ FPackedPathTracingPayload TraceTransparentRay(inout FPathState PathState, int Bo
 	TraceRay(
 		TLAS,
 		RayFlags,
-		InstanceInclusionMask,
+		InstanceInclusionMaskOpaque,
 		RAY_TRACING_SHADER_SLOT_MATERIAL,
 		RAY_TRACING_NUM_SHADER_SLOTS,
 		MissShaderIndex,
@@ -448,8 +474,6 @@ FPackedPathTracingPayload TraceTransparentRay(inout FPathState PathState, int Bo
 		PathState.Ray.TMax = asfloat(asuint(OpaquePackedPayload.HitT) - 1);
 	}
 
-	InstanceInclusionMask = bUseCameraRayType ? PATHTRACER_MASK_CAMERA_TRANSLUCENT : PATHTRACER_MASK_INDIRECT_TRANSLUCENT;
-
 	for (;;)
 	{
 #if PATH_TRACER_INCLUDE_TRANSLUCENT
@@ -466,13 +490,14 @@ FPackedPathTracingPayload TraceTransparentRay(inout FPathState PathState, int Bo
 			PackedPayload.SetSceneDepth(dot(TranslatedWorldPos, ViewZ));
 		}
 
+        const uint InstanceInclusionMaskTranslucent = PathState.GetInstanceInclusionMaskTranslucent();
 #if PATH_TRACER_USE_SER 
 		{
 			NvHitObject Hit;
 			NvTraceRayHitObject(
 				TLAS,
 				RayFlags,
-				InstanceInclusionMask,
+				InstanceInclusionMaskTranslucent,
 				RAY_TRACING_SHADER_SLOT_MATERIAL,
 				RAY_TRACING_NUM_SHADER_SLOTS,
 				MissShaderIndex,
@@ -489,7 +514,7 @@ FPackedPathTracingPayload TraceTransparentRay(inout FPathState PathState, int Bo
 		TraceRay(
 			TLAS,
 			RayFlags,
-			InstanceInclusionMask,
+			InstanceInclusionMaskTranslucent,
 			RAY_TRACING_SHADER_SLOT_MATERIAL,
 			RAY_TRACING_NUM_SHADER_SLOTS,
 			MissShaderIndex,
@@ -938,7 +963,6 @@ FPathTracingResult PathTracingDebugCore(FRayDesc Ray, uint2 LaunchIndex)
 	FPathTracingResult Result = (FPathTracingResult) 0;
 	const uint RayFlags = DebugMode == PATH_TRACER_DEBUG_VIZ_HITKIND ? 0 : RAY_FLAG_CULL_BACK_FACING_TRIANGLES;
 	const uint MissShaderIndex = 0;
-	uint InstanceInclusionMask = PATHTRACER_MASK_CAMERA | PATHTRACER_MASK_HAIR_CAMERA;
 	int TransparentHitCount = 0;
 
 	const float RandSample = RandomSequenceCreate(uint3(LaunchIndex, View.StateFrameIndex), 0, 1).Get1D();
@@ -955,7 +979,7 @@ FPathTracingResult PathTracingDebugCore(FRayDesc Ray, uint2 LaunchIndex)
 	TraceRay(
 		TLAS,
 		RayFlags,
-		InstanceInclusionMask,
+		PATHTRACER_MASK_CAMERA_OPAQUE,
 		RAY_TRACING_SHADER_SLOT_MATERIAL,
 		RAY_TRACING_NUM_SHADER_SLOTS,
 		MissShaderIndex,
@@ -1001,9 +1025,6 @@ FPathTracingResult PathTracingDebugCore(FRayDesc Ray, uint2 LaunchIndex)
 		return Result;
 	}
 
-
-	InstanceInclusionMask = PATHTRACER_MASK_CAMERA_TRANSLUCENT;
-
 	for (;;)
 	{
 #if PATH_TRACER_INCLUDE_TRANSLUCENT
@@ -1024,7 +1045,7 @@ FPathTracingResult PathTracingDebugCore(FRayDesc Ray, uint2 LaunchIndex)
 		TraceRay(
 			TLAS,
 			RayFlags,
-			InstanceInclusionMask,
+			PATHTRACER_MASK_CAMERA_TRANSLUCENT,
 			RAY_TRACING_SHADER_SLOT_MATERIAL,
 			RAY_TRACING_NUM_SHADER_SLOTS,
 			MissShaderIndex,
diff --git a/Engine/Shaders/Private/PathTracing/PathTracingInitExtinctionCoefficient.usf b/Engine/Shaders/Private/PathTracing/PathTracingInitExtinctionCoefficient.usf
--- a/Engine/Shaders/Private/PathTracing/PathTracingInitExtinctionCoefficient.usf
+++ b/Engine/Shaders/Private/PathTracing/PathTracingInitExtinctionCoefficient.usf
@@ -8,9 +8,9 @@ RWBuffer<float> RWStartingMediumData;
 
 RAY_TRACING_ENTRY_RAYGEN(PathTracingInitExtinctionCoefficientRG)
 {
-	const uint RayFlags = 0;
+    const uint RayFlags = RAY_FLAG_SKIP_PROCEDURAL_PRIMITIVES;
 	const uint MissShaderIndex = 0;
-	const uint InstanceInclusionMask = PATHTRACER_MASK_CAMERA | PATHTRACER_MASK_CAMERA_TRANSLUCENT; // Ignore hair strands, only trace what camera rays can see
+	const uint InstanceInclusionMask = PATHTRACER_MASK_CAMERA_OPAQUE | PATHTRACER_MASK_CAMERA_TRANSLUCENT; // Only trace what camera rays can see
 
 	FRayDesc UpRay;
 	UpRay.Origin = PrimaryView.TranslatedWorldCameraOrigin;
diff --git a/Engine/Shaders/Private/RayTracing/RayTracingPrimaryRays.usf b/Engine/Shaders/Private/RayTracing/RayTracingPrimaryRays.usf
--- a/Engine/Shaders/Private/RayTracing/RayTracingPrimaryRays.usf
+++ b/Engine/Shaders/Private/RayTracing/RayTracingPrimaryRays.usf
@@ -299,7 +299,7 @@ RAY_TRACING_ENTRY_RAYGEN(RayTracingPrimaryRaysRGS)
 			ApplyPositionBias(ReflectionRay, Payload.WorldNormal, MaxNormalBias);
 
 			const uint ReflectionRayFlags = RAY_FLAG_CULL_BACK_FACING_TRIANGLES | (ForceOpaqueRays ? RAY_FLAG_FORCE_OPAQUE : 0);
-			const uint ReflectionInstanceInclusionMask = (RAY_TRACING_MASK_OPAQUE | RAY_TRACING_MASK_TRANSLUCENT | RAY_TRACING_MASK_HAIR_STRANDS);
+			const uint ReflectionInstanceInclusionMask = (RAY_TRACING_MASK_OPAQUE_REFLECTION | RAY_TRACING_MASK_TRANSLUCENT | RAY_TRACING_MASK_HAIR_STRANDS);
 			const bool bReflectionRayTraceSkyLightContribution = false;
 			const bool bReflectionDecoupleSampleGeneration = true;
 			const bool bReflectionEnableSkyLightContribution = bSkyLightAffectReflection;
diff --git a/Engine/Shaders/Shared/PathTracingDefinitions.h b/Engine/Shaders/Shared/PathTracingDefinitions.h
--- a/Engine/Shaders/Shared/PathTracingDefinitions.h
+++ b/Engine/Shaders/Shared/PathTracingDefinitions.h
@@ -27,17 +27,17 @@
 #define PATHTRACER_FLAG_CAST_CLOUD_SHADOW_MASK              (1 << 11)
 #define PATHTRACER_FLAG_HAS_RECT_TEXTURE_MASK				(1 << 12)
 
-#define PATHTRACER_MASK_CAMERA								0x01			// opaque and alpha tested meshes and particles as a whole (primary ray) excluding hairs
-#define PATHTRACER_MASK_HAIR_CAMERA							0x02			// For primary ray tracing against hair
-#define PATHTRACER_MASK_SHADOW								0x04			// Whether the geometry is visible for shadow rays
-#define PATHTRACER_MASK_HAIR_SHADOW							0x08			// Whether hair is visible for shadow rays
-#define PATHTRACER_MASK_INDIRECT							0x10			// opaque and alpha tested meshes and particles as a whole (indirect ray) excluding hairs
-#define PATHTRACER_MASK_HAIR_INDIRECT						0x20			// For indirect ray tracing against hair
-#define PATHTRACER_MASK_CAMERA_TRANSLUCENT					0x40			
-#define PATHTRACER_MASK_INDIRECT_TRANSLUCENT				0x80			
+#define PATHTRACER_MASK_CAMERA_OPAQUE                       0x01            // Visible to Camera and BlendMode=Opaque+Masked 
+#define PATHTRACER_MASK_CAMERA_TRANSLUCENT                  0x02            // Visible to Camera and BlendMode=Translucent+others (all blend modes outside of opaque and masked)
+#define PATHTRACER_MASK_SHADOW                              0x04            // Casts shadows: all blend modes
+#define PATHTRACER_MASK_INDIRECT_OPAQUE                     0x08            // Affects Indirect and BlendMode=Opaque+Masked 
+#define PATHTRACER_MASK_INDIRECT_TRANSLUCENT                0x10            // Affects Indirect and BlendMode=Translucent+others
+#define PATHTRACER_MASK_REFLECTION_OPAQUE                   0x20            // Visible in Reflections and BlendMode=Opaque+Masked
+#define PATHTRACER_MASK_REFLECTION_TRANSLUCENT              0x40			// Visible in Reflections and BlendMode=Translucent+others
+#define PATHTRACER_MASK_SSS_RANDOM_WALK                     0x80            // Visible to SSS random walks
 
-#define PATHTRACER_MASK_IGNORE								0x00			// used when mapping general tracing mask to path tracing mask
-#define PATHTRACER_MASK_ALL									0xFF
+#define PATHTRACER_MASK_IGNORE                              0x00            // used when mapping general tracing mask to path tracing mask
+#define PATHTRACER_MASK_ALL                                 0xFF
 
 // Constants for light contribution types (AOV decomposition of the image)
 // Leaving all constants enabled creates the beauty image, but turning off some bits allows
diff --git a/Engine/Source/Runtime/Renderer/Private/RayTracing/RayTracingDebug.cpp b/Engine/Source/Runtime/Renderer/Private/RayTracing/RayTracingDebug.cpp
--- a/Engine/Source/Runtime/Renderer/Private/RayTracing/RayTracingDebug.cpp
+++ b/Engine/Source/Runtime/Renderer/Private/RayTracing/RayTracingDebug.cpp
@@ -2094,6 +2094,10 @@ void FDeferredShadingSceneRenderer::RayTracingDisplayPicking(const FRayTracingPi
 	{
 		Writer.DrawLine(FText::FromString(TEXT("   RAY_TRACING_MASK_OPAQUE_FP_WORLD_SPACE")), 10, FColor::Yellow);
 	}
+	if (PickingFeedback.Mask & RAY_TRACING_MASK_OPAQUE_REFLECTION)
+	{
+		Writer.DrawLine(FText::FromString(TEXT("   RAY_TRACING_MASK_OPAQUE_REFLECTION")), 10, FColor::Yellow);
+	}
 
 	Writer.DrawLine(FText::FromString(FString::Printf(TEXT("ContributionToHitGroup: %u"), PickingFeedback.InstanceContributionToHitGroupIndex)), 10, FColor::Yellow);
 	{
diff --git a/Engine/Source/Runtime/Renderer/Private/RayTracing/RayTracingInstanceMask.cpp b/Engine/Source/Runtime/Renderer/Private/RayTracing/RayTracingInstanceMask.cpp
--- a/Engine/Source/Runtime/Renderer/Private/RayTracing/RayTracingInstanceMask.cpp
+++ b/Engine/Source/Runtime/Renderer/Private/RayTracing/RayTracingInstanceMask.cpp
@@ -26,18 +26,17 @@ enum class ERayTracingInstanceMaskType : uint8
 	TranslucentShadow,
 	ThinShadow,
 
-	// Mask types for reflection rays
-	OpaqueVisibleInReflections,
-
 	// Geometry specific ray types
 	HairStrands,
 
 	// Special purpose ray types
 	FirstPersonWorldSpaceRepresentation,
 
-	// path tracing specific mask type
+	// Ray types according to illumination type
 	VisibleInPrimaryRay,
-	VisibleInIndirectRay
+	VisibleInIndirectRay,
+	VisibleInReflectionRay,
+
 };
 
 } // anonymous namespace
@@ -49,7 +48,7 @@ static uint8 ComputeRayTracingInstanceMask(ERayTracingInstanceMaskType MaskType,
 		switch (MaskType)
 		{
 		case ERayTracingInstanceMaskType::Opaque:
-			return RAY_TRACING_MASK_OPAQUE;
+			return RAY_TRACING_MASK_OPAQUE | RAY_TRACING_MASK_OPAQUE_REFLECTION;
 		case ERayTracingInstanceMaskType::Translucent:
 			return RAY_TRACING_MASK_TRANSLUCENT;
 		case ERayTracingInstanceMaskType::OpaqueShadow:
@@ -58,8 +57,6 @@ static uint8 ComputeRayTracingInstanceMask(ERayTracingInstanceMaskType MaskType,
 			return RAY_TRACING_MASK_TRANSLUCENT_SHADOW;
 		case ERayTracingInstanceMaskType::ThinShadow:
 			return RAY_TRACING_MASK_THIN_SHADOW;
-		case ERayTracingInstanceMaskType::OpaqueVisibleInReflections:
-			return RAY_TRACING_MASK_OPAQUE_REFLECTION;
 		case ERayTracingInstanceMaskType::HairStrands:
 			return RAY_TRACING_MASK_HAIR_STRANDS | RAY_TRACING_MASK_THIN_SHADOW;
 		case ERayTracingInstanceMaskType::FirstPersonWorldSpaceRepresentation:
@@ -68,6 +65,8 @@ static uint8 ComputeRayTracingInstanceMask(ERayTracingInstanceMaskType MaskType,
 			return 0; // There is no distinct notion of primary ray visibility for ray tracing
 		case ERayTracingInstanceMaskType::VisibleInIndirectRay:
 			return RAY_TRACING_MASK_OPAQUE | RAY_TRACING_MASK_TRANSLUCENT | RAY_TRACING_MASK_HAIR_STRANDS;
+		case ERayTracingInstanceMaskType::VisibleInReflectionRay:
+			return RAY_TRACING_MASK_OPAQUE_REFLECTION;
 		default:
 			break;
 		}
@@ -78,25 +77,23 @@ static uint8 ComputeRayTracingInstanceMask(ERayTracingInstanceMaskType MaskType,
 		switch (MaskType)
 		{
 		case ERayTracingInstanceMaskType::Opaque:
-			return PATHTRACER_MASK_CAMERA | PATHTRACER_MASK_INDIRECT;
+			return PATHTRACER_MASK_CAMERA_OPAQUE | PATHTRACER_MASK_INDIRECT_OPAQUE | PATHTRACER_MASK_REFLECTION_OPAQUE | PATHTRACER_MASK_SSS_RANDOM_WALK;
 		case ERayTracingInstanceMaskType::Translucent:
-			return PATHTRACER_MASK_CAMERA_TRANSLUCENT | PATHTRACER_MASK_INDIRECT_TRANSLUCENT;
+			return PATHTRACER_MASK_CAMERA_TRANSLUCENT | PATHTRACER_MASK_INDIRECT_TRANSLUCENT | PATHTRACER_MASK_REFLECTION_TRANSLUCENT | PATHTRACER_MASK_SSS_RANDOM_WALK;
 		case ERayTracingInstanceMaskType::OpaqueShadow:
-			return PATHTRACER_MASK_SHADOW;
 		case ERayTracingInstanceMaskType::TranslucentShadow:
-			return PATHTRACER_MASK_SHADOW;
 		case ERayTracingInstanceMaskType::ThinShadow:
-			return PATHTRACER_MASK_HAIR_SHADOW;
-		case ERayTracingInstanceMaskType::OpaqueVisibleInReflections:
-			return PATHTRACER_MASK_IGNORE;
+			return PATHTRACER_MASK_SHADOW;
 		case ERayTracingInstanceMaskType::FirstPersonWorldSpaceRepresentation:
 			return PATHTRACER_MASK_IGNORE;
 		case ERayTracingInstanceMaskType::HairStrands:
-			return PATHTRACER_MASK_HAIR_CAMERA | PATHTRACER_MASK_HAIR_SHADOW | PATHTRACER_MASK_HAIR_INDIRECT;
+			return PATHTRACER_MASK_CAMERA_OPAQUE | PATHTRACER_MASK_INDIRECT_OPAQUE | PATHTRACER_MASK_REFLECTION_OPAQUE | PATHTRACER_MASK_SHADOW;
 		case ERayTracingInstanceMaskType::VisibleInPrimaryRay:
-			return PATHTRACER_MASK_CAMERA | PATHTRACER_MASK_HAIR_CAMERA | PATHTRACER_MASK_CAMERA_TRANSLUCENT;
+			return PATHTRACER_MASK_CAMERA_OPAQUE | PATHTRACER_MASK_CAMERA_TRANSLUCENT;
 		case ERayTracingInstanceMaskType::VisibleInIndirectRay:
-			return PATHTRACER_MASK_INDIRECT | PATHTRACER_MASK_HAIR_INDIRECT | PATHTRACER_MASK_INDIRECT_TRANSLUCENT;
+			return PATHTRACER_MASK_INDIRECT_OPAQUE | PATHTRACER_MASK_INDIRECT_TRANSLUCENT;
+		case ERayTracingInstanceMaskType::VisibleInReflectionRay:
+			return PATHTRACER_MASK_REFLECTION_OPAQUE | PATHTRACER_MASK_REFLECTION_TRANSLUCENT;
 		default:
 			break;
 		}
@@ -132,9 +129,51 @@ struct FSceneProxyRayTracingMaskInfo
 	bool bVisibleToCamera = false;
 	bool bVisibleToShadow = false;
 	bool bVisibleToIndirect = false;
-	bool bVisibleInReflections = false;
+	bool bVisibleToReflections = false;
 	bool bIsFirstPersonWorldSpaceRepresentation = false;
 	ERayTracingType RayTracingType = ERayTracingType::RayTracing;
+
+    void UpdateMask(uint8* Mask) const
+    {
+        if (!bVisibleToCamera)
+        {
+            // If the object is not visible to camera, remove all direct visibility bits.
+            *Mask &= ~ComputeRayTracingInstanceMask(ERayTracingInstanceMaskType::VisibleInPrimaryRay, RayTracingType);
+        }
+
+        if (!bVisibleToIndirect)
+        {
+            // If the object does not affect indirect lighting, remove all indirect bits.
+            *Mask &= ~ComputeRayTracingInstanceMask(ERayTracingInstanceMaskType::VisibleInIndirectRay, RayTracingType);
+        }
+
+        if (!bVisibleToReflections)
+        {
+            // If the object does not affect reflections, remove all reflection bits.
+            *Mask &= ~ComputeRayTracingInstanceMask(ERayTracingInstanceMaskType::VisibleInReflectionRay, RayTracingType);
+        }
+
+        if (!bVisibleToShadow)
+        {
+            // Not casting shadows, remove any set shadow flags
+            *Mask &= ~(
+                ComputeRayTracingInstanceMask(ERayTracingInstanceMaskType::OpaqueShadow, RayTracingType) |
+                ComputeRayTracingInstanceMask(ERayTracingInstanceMaskType::TranslucentShadow, RayTracingType) |
+                ComputeRayTracingInstanceMask(ERayTracingInstanceMaskType::ThinShadow, RayTracingType)
+            );
+        }
+
+        if (bIsFirstPersonWorldSpaceRepresentation)
+        {
+            const uint8 OpaqueMaskType = 
+                ComputeRayTracingInstanceMask(ERayTracingInstanceMaskType::Opaque, RayTracingType) |
+                ComputeRayTracingInstanceMask(ERayTracingInstanceMaskType::OpaqueShadow, RayTracingType);
+            const bool bIsOpaque = (*Mask & OpaqueMaskType) != 0;
+            // Tag world space representations of first person meshes so rays originating from first person meshes can skip them.
+            // We currently only support opaque world space representations of first person objects, so set the mask to 0 otherwise.
+            *Mask = bIsOpaque ? ComputeRayTracingInstanceMask(ERayTracingInstanceMaskType::FirstPersonWorldSpaceRepresentation, RayTracingType) : 0;
+        }
+    }
 };
 
 } // anonymous namespace
@@ -166,25 +205,12 @@ static FSceneProxyRayTracingMaskInfo GetSceneProxyRayTracingMaskInfo(const FPrim
 		MaskInfo.bVisibleToIndirect = PrimitiveSceneProxy.AffectsIndirectLightingWhileHidden();
 	}
 
-	MaskInfo.bVisibleInReflections = MaskInfo.bVisibleToIndirect && PrimitiveSceneProxy.IsVisibleInReflections();
+	MaskInfo.bVisibleToReflections = MaskInfo.bVisibleToIndirect && PrimitiveSceneProxy.IsVisibleInReflections();
 	MaskInfo.bIsFirstPersonWorldSpaceRepresentation = PrimitiveSceneProxy.IsFirstPersonWorldSpaceRepresentation();
 
 	return MaskInfo;
 }
 
-static uint8 ApplyFirstPersonRayTracingInstanceMaskFlag(uint8 Mask, ERayTracingType RayTracingType, bool bIsFirstPersonWorldSpaceRepresentation)
-{
-	if (bIsFirstPersonWorldSpaceRepresentation)
-	{
-		const bool bIsOpaque = (Mask & (ComputeRayTracingInstanceMask(ERayTracingInstanceMaskType::Opaque, RayTracingType) | ComputeRayTracingInstanceMask(ERayTracingInstanceMaskType::OpaqueShadow, RayTracingType))) != 0;
-
-		// Tag world space representations of first person meshes so rays originating from first person meshes can skip them.
-		// We currently only support opaque world space representations of first person objects, so set the mask to 0 otherwise.
-		Mask = bIsOpaque ? ComputeRayTracingInstanceMask(ERayTracingInstanceMaskType::FirstPersonWorldSpaceRepresentation, RayTracingType) : 0;
-	}
-	return Mask;
-}
-
 FRayTracingMaskAndFlags BuildRayTracingInstanceMaskAndFlags(const FRayTracingInstance& Instance, const FPrimitiveSceneProxy& PrimitiveSceneProxy)
 {
 	FSceneProxyRayTracingMaskInfo MaskInfo = GetSceneProxyRayTracingMaskInfo(PrimitiveSceneProxy);
@@ -223,6 +249,7 @@ FRayTracingMaskAndFlags BuildRayTracingInstanceMaskAndFlags(const FRayTracingIns
 			bAllSegmentsReverseCulling &= MeshBatch.ReverseCulling;
 		}
 	}
+    MaskInfo.bVisibleToShadow = bAnySegmentsCastShadow;
 
 	// Run AHS for alpha masked and meshes with only some sections casting shadows, which require per mesh section filtering in AHS
 	Result.bForceOpaque = bAllSegmentsOpaque && (bAllSegmentsCastShadow || !bAnySegmentsCastShadow);
@@ -246,34 +273,7 @@ FRayTracingMaskAndFlags BuildRayTracingInstanceMaskAndFlags(const FRayTracingIns
 		Result.bAllSegmentsTranslucent = false;
 	}
 
-	if (!MaskInfo.bVisibleToCamera)
-	{
-		// If the object is not visible to camera, remove all direct visibility bits.
-		Result.Mask &= ~ComputeRayTracingInstanceMask(ERayTracingInstanceMaskType::VisibleInPrimaryRay, RayTracingType);
-	}
-
-	if (!MaskInfo.bVisibleToIndirect)
-	{
-		// If the object does not affect indirect lighting, remove all indirect bits.
-		Result.Mask &= ~ComputeRayTracingInstanceMask(ERayTracingInstanceMaskType::VisibleInIndirectRay, RayTracingType);
-	}
-
-	if (!bAnySegmentsCastShadow)
-	{
-		// Not casting shadows, remove any set shadow flags
-		Result.Mask &= ~(
-			ComputeRayTracingInstanceMask(ERayTracingInstanceMaskType::OpaqueShadow, RayTracingType) |
-			ComputeRayTracingInstanceMask(ERayTracingInstanceMaskType::TranslucentShadow, RayTracingType) |
-			ComputeRayTracingInstanceMask(ERayTracingInstanceMaskType::ThinShadow, RayTracingType)
-		);
-	}
-
-	if (MaskInfo.bVisibleInReflections && (Result.Mask & ComputeRayTracingInstanceMask(ERayTracingInstanceMaskType::Opaque, RayTracingType)))
-	{
-		Result.Mask |= ComputeRayTracingInstanceMask(ERayTracingInstanceMaskType::OpaqueVisibleInReflections, RayTracingType);
-	}
-
-	Result.Mask = ApplyFirstPersonRayTracingInstanceMaskFlag(Result.Mask, RayTracingType, MaskInfo.bIsFirstPersonWorldSpaceRepresentation);
+    MaskInfo.UpdateMask(&Result.Mask);
 
 	return Result;
 }
@@ -308,41 +308,11 @@ void SetupRayTracingMeshCommandMaskAndStatus(FRayTracingMeshCommand& MeshCommand
 	MeshCommand.bNaniteRayTracing = (Nanite::GetRayTracingMode() != Nanite::ERayTracingMode::Fallback) && PrimitiveSceneProxy->IsNaniteMesh();
 
 	FSceneProxyRayTracingMaskInfo MaskInfo = GetSceneProxyRayTracingMaskInfo(*PrimitiveSceneProxy);
+    MaskInfo.RayTracingType = RayTracingType;
+    MaskInfo.bVisibleToShadow &= MeshCommand.bCastRayTracedShadows;
 
-	MeshCommand.InstanceMask = ApplyFirstPersonRayTracingInstanceMaskFlag(MeshCommand.InstanceMask, RayTracingType, MaskInfo.bIsFirstPersonWorldSpaceRepresentation);
-
-	// TODO: This should be done once all mesh commands for a mesh are combined (similar to BuildRayTracingInstanceMaskAndFlags above)
-	// TODO: Why is this logic not applied in the ray tracing case?
-	if (RayTracingType == ERayTracingType::PathTracing ||
-		RayTracingType == ERayTracingType::LightMapTracing)
-	{
-	
-		if (!MaskInfo.bVisibleToCamera)
-		{
-			// If the object is not visible to camera, remove all direct visibility bits.
-			MeshCommand.InstanceMask &= ~ComputeRayTracingInstanceMask(ERayTracingInstanceMaskType::VisibleInPrimaryRay, RayTracingType);
-		}
-
-		if (!MaskInfo.bVisibleToIndirect)
-		{
-			// If the object does not affect indirect lighting, remove all indirect bits.
-			MeshCommand.InstanceMask &= ~ComputeRayTracingInstanceMask(ERayTracingInstanceMaskType::VisibleInIndirectRay, RayTracingType);
-		}
-
-		if (!MaskInfo.bVisibleToShadow || !MeshCommand.bCastRayTracedShadows)
-		{
-			// Not casting shadows, remove any set shadow flags
-			MeshCommand.InstanceMask &= ~(
-				ComputeRayTracingInstanceMask(ERayTracingInstanceMaskType::OpaqueShadow, RayTracingType) |
-				ComputeRayTracingInstanceMask(ERayTracingInstanceMaskType::TranslucentShadow, RayTracingType) |
-				ComputeRayTracingInstanceMask(ERayTracingInstanceMaskType::ThinShadow, RayTracingType));
-		}
-	}
-
-	if (MaskInfo.bVisibleInReflections && (MeshCommand.InstanceMask & ComputeRayTracingInstanceMask(ERayTracingInstanceMaskType::Opaque, RayTracingType)))
-	{
-		MeshCommand.InstanceMask |= ComputeRayTracingInstanceMask(ERayTracingInstanceMaskType::OpaqueVisibleInReflections, RayTracingType);
-	}
+    // TODO: Should this be done once all mesh commands for a mesh are combined? (similar to BuildRayTracingInstanceMaskAndFlags above)
+    MaskInfo.UpdateMask(&MeshCommand.InstanceMask);
 }
 
 #endif
