[ShaderPrint] Fix incorrect indirect args buffer usage in InternalDrawView_Primitives.
- DrawIndexedPrimitiveIndirect(...) requires FRHIDrawIndexedIndirectParameters.
- also set StartInstanceLocation in ShaderDrawDebugCopyCS and BuildIndirectDrawArgsCS.

#Synced-CL 51462794

diff --git a/Engine/Shaders/Private/ShaderPrintDraw.usf b/Engine/Shaders/Private/ShaderPrintDraw.usf
--- a/Engine/Shaders/Private/ShaderPrintDraw.usf
+++ b/Engine/Shaders/Private/ShaderPrintDraw.usf
@@ -270,12 +270,10 @@ void CompactStateBufferCS(uint3 DispatchThreadId : SV_DispatchThreadID)
 [numthreads(1, 1, 1)]
 void BuildIndirectDrawArgsCS()
 {
-	uint SymbolCount = min(ShaderPrintData.MaxSymbolCount, SymbolsBuffer[SHADER_PRINT_COUNTER_OFFSET_SYMBOL]);
-
-	RWIndirectDrawArgsBuffer[0] = 6;
-	RWIndirectDrawArgsBuffer[1] = SymbolCount * 2; // Main glyph + drop shadow
-	RWIndirectDrawArgsBuffer[2] = 0;
-	RWIndirectDrawArgsBuffer[3] = 0;
+	const uint SymbolCount = min(ShaderPrintData.MaxSymbolCount, SymbolsBuffer[SHADER_PRINT_COUNTER_OFFSET_SYMBOL]);
+	const uint IndexCount = 6;
+	const uint InstanceCount = SymbolCount * 2; // Main glyph + drop shadow;
+	WriteIndirectIndexedDrawArgs(RWIndirectDrawArgsBuffer, 0, IndexCount, InstanceCount);
 }
 
 
diff --git a/Engine/Shaders/Private/ShaderPrintDrawPrimitive.usf b/Engine/Shaders/Private/ShaderPrintDrawPrimitive.usf
--- a/Engine/Shaders/Private/ShaderPrintDrawPrimitive.usf
+++ b/Engine/Shaders/Private/ShaderPrintDrawPrimitive.usf
@@ -126,11 +126,8 @@ void ShaderDrawDebugCopyCS(uint3 DispatchThreadId : SV_DispatchThreadID)
 	const uint Index    = bLine ? SHADER_PRINT_COUNTER_OFFSET_LINE : SHADER_PRINT_COUNTER_OFFSET_TRIANGLE;
 	const uint MaxCount = bLine ? ShaderPrintData.MaxLineCount : ShaderPrintData.MaxTriangleCount;
 
-	const uint IndirectArgInstanceCount = min(MaxCount, ElementBuffer[Index]);
-
-	RWIndirectArgs[0] = bLine ? 2 : 3; // Line or Triangle primitive
-	RWIndirectArgs[1] = IndirectArgInstanceCount;
-	RWIndirectArgs[2] = 0;
-	RWIndirectArgs[3] = 0;
+	const uint IndexCount = bLine ? 2 : 3; // Line or Triangle primitive
+	const uint InstanceCount = min(MaxCount, ElementBuffer[Index]);
+	WriteIndirectIndexedDrawArgs(RWIndirectArgs, 0, IndexCount, InstanceCount);
 }
 #endif
\ No newline at end of file
diff --git a/Engine/Source/Runtime/Renderer/Private/ShaderPrint.cpp b/Engine/Source/Runtime/Renderer/Private/ShaderPrint.cpp
--- a/Engine/Source/Runtime/Renderer/Private/ShaderPrint.cpp
+++ b/Engine/Source/Runtime/Renderer/Private/ShaderPrint.cpp
@@ -609,7 +609,6 @@ namespace ShaderPrint
 			SHADER_PARAMETER(FMatrix44f, TranslatedWorldToClip)
 			SHADER_PARAMETER_STRUCT_REF(FViewUniformShaderParameters, View)
 			SHADER_PARAMETER_RDG_BUFFER_SRV(Buffer<uint>, ShaderDrawDebugPrimitive)
-			RDG_BUFFER_ACCESS(IndirectBuffer, ERHIAccess::IndirectArgs)
 		END_SHADER_PARAMETER_STRUCT()
 
 		class FPrimitiveType : SHADER_PERMUTATION_INT("PERMUTATION_PRIMITIVE_TYPE", 2);
@@ -675,6 +674,7 @@ namespace ShaderPrint
 	BEGIN_SHADER_PARAMETER_STRUCT(FShaderDrawVSPSParameters , )
 		SHADER_PARAMETER_STRUCT_INCLUDE(FShaderDrawDebugVS::FParameters, VS)
 		SHADER_PARAMETER_STRUCT_INCLUDE(FShaderDrawDebugPS::FParameters, PS)
+		RDG_BUFFER_ACCESS(IndirectBuffer, ERHIAccess::IndirectArgs)
 	END_SHADER_PARAMETER_STRUCT()
 
 
@@ -1086,7 +1086,7 @@ namespace ShaderPrint
 	{
 		FGlobalShaderMap* GlobalShaderMap = GetGlobalShaderMap(GMaxRHIFeatureLevel);
 
-		FRDGBufferRef IndirectBuffer = GraphBuilder.CreateBuffer(FRDGBufferDesc::CreateIndirectDesc<FRHIDrawIndirectParameters>(1), TEXT("ShaderDraw.IndirectBuffer"), ERDGBufferFlags::None);
+		FRDGBufferRef IndirectBuffer = GraphBuilder.CreateBuffer(FRDGBufferDesc::CreateIndirectDesc<FRHIDrawIndexedIndirectParameters>(1), TEXT("ShaderDraw.IndirectBuffer"), ERDGBufferFlags::None);
 		{
 			FShaderDrawDebugCopyCS::FParameters* Parameters = GraphBuilder.AllocParameters<FShaderDrawDebugCopyCS::FParameters>();
 			Parameters->ElementBuffer = ShaderPrintPrimitiveBufferSRV;
@@ -1128,13 +1128,13 @@ namespace ShaderPrint
 		PassParameters->PS.bCheckerboardEnabled = bLines ? 1u : 0u;
 		PassParameters->PS.bDrawOccludedLines = CVarDrawOccludedLines.GetValueOnRenderThread() != 0 ? 1 : 0;
 		PassParameters->VS.ShaderDrawDebugPrimitive = ShaderPrintPrimitiveBufferSRV;
-		PassParameters->VS.IndirectBuffer = IndirectBuffer;
 		PassParameters->VS.Common = ShaderPrintData.UniformBuffer;
+		PassParameters->IndirectBuffer = IndirectBuffer;
 
 		ValidateShaderParameters(PixelShader, PassParameters->PS);
-		ClearUnusedGraphResources(PixelShader, &PassParameters->PS, { IndirectBuffer });
+		ClearUnusedGraphResources(PixelShader, &PassParameters->PS);
 		ValidateShaderParameters(VertexShader, PassParameters->VS);
-		ClearUnusedGraphResources(VertexShader, &PassParameters->VS, { IndirectBuffer });
+		ClearUnusedGraphResources(VertexShader, &PassParameters->VS);
 
 		const FIntRect Viewport = UnscaledViewRect;
 		GraphBuilder.AddPass(
@@ -1144,7 +1144,7 @@ namespace ShaderPrint
 			[VertexShader, PixelShader, PassParameters, IndirectBuffer, Viewport, bLines](FRDGAsyncTask, FRHICommandList& RHICmdList)
 			{
 				// Marks the indirect draw parameter as used by the pass, given it's not used directly by any of the shaders.
-				PassParameters->VS.IndirectBuffer->MarkResourceAsUsed();
+				PassParameters->IndirectBuffer->MarkResourceAsUsed();
 
 				FGraphicsPipelineStateInitializer GraphicsPSOInit;
 				RHICmdList.ApplyCachedRenderTargets(GraphicsPSOInit);
@@ -1162,7 +1162,7 @@ namespace ShaderPrint
 				SetShaderParameters(RHICmdList, PixelShader, PixelShader.GetPixelShader(), PassParameters->PS);
 
 				// Marks the indirect draw parameter as used by the pass, given it's not used directly by any of the shaders.
-				FRHIBuffer* IndirectBufferRHI = PassParameters->VS.IndirectBuffer->GetIndirectRHICallBuffer();
+				FRHIBuffer* IndirectBufferRHI = PassParameters->IndirectBuffer->GetIndirectRHICallBuffer();
 				check(IndirectBufferRHI != nullptr);
 				RHICmdList.DrawIndexedPrimitiveIndirect(bLines ? GLineIndexBuffer.IndexBufferRHI : GTriangleIndexBuffer.IndexBufferRHI, IndirectBufferRHI, 0);
 			});
