Force rect light atlas to reset when changing the atlas filtering type.

#Synced-CL 51465991

diff --git a/Engine/Source/Runtime/Renderer/Private/RectLightTextureManager.cpp b/Engine/Source/Runtime/Renderer/Private/RectLightTextureManager.cpp
--- a/Engine/Source/Runtime/Renderer/Private/RectLightTextureManager.cpp
+++ b/Engine/Source/Runtime/Renderer/Private/RectLightTextureManager.cpp
@@ -1710,14 +1710,16 @@ void UpdateAtlasTexture(FRDGBuilder& GraphBuilder, const ERHIFeatureLevel::Type
 	// Force update by resetting the atlas layout
 	static int32 CachedMaxAtlasResolution = CVarRectLightTextureResolution.GetValueOnRenderThread();
 	static EPixelFormat CachedAtlasFormat = GetRectLightAtlasFormat();
+	static int32 CachedFilteringMethod = CVarRectLighFilterQuality.GetValueOnRenderThread();
 	const EPixelFormat AtlasFormat = GetRectLightAtlasFormat();
-	const bool bForceUpdate = CVarRectLighForceUpdate.GetValueOnRenderThread() > 0 || CachedMaxAtlasResolution != CVarRectLightTextureResolution.GetValueOnRenderThread() || CachedAtlasFormat != AtlasFormat;	
+	const bool bForceUpdate = CVarRectLighForceUpdate.GetValueOnRenderThread() > 0 || CachedMaxAtlasResolution != CVarRectLightTextureResolution.GetValueOnRenderThread() || CachedAtlasFormat != AtlasFormat || CachedFilteringMethod != CVarRectLighFilterQuality.GetValueOnRenderThread();
 	const bool bResetOnChange = CVarRectLighResetOnChange.GetValueOnRenderThread() > 0;
 	const bool bClearAtlas = CVarRectLighClearAtlas.GetValueOnRenderThread() > 0;
 	auto ResetAtlas = [&]()
 	{
 		CachedMaxAtlasResolution = CVarRectLightTextureResolution.GetValueOnRenderThread();
 		CachedAtlasFormat = AtlasFormat;
+		CachedFilteringMethod = CVarRectLighFilterQuality.GetValueOnRenderThread();
 		GRectLightTextureManager.bHasPendingAdds = true;
 		GRectLightTextureManager.AtlasLayout = FAtlasLayout(FIntPoint(CachedMaxAtlasResolution, CachedMaxAtlasResolution));
 		GRectLightTextureManager.DeletedSlots.Empty();
