Substrate - Toon - fixed runtime render target load/bind issue when running without substrate enabled.

#rnx
#Synced-CL 51456025

diff --git a/Engine/Source/Runtime/Renderer/Private/LightRendering.cpp b/Engine/Source/Runtime/Renderer/Private/LightRendering.cpp
--- a/Engine/Source/Runtime/Renderer/Private/LightRendering.cpp
+++ b/Engine/Source/Runtime/Renderer/Private/LightRendering.cpp
@@ -2749,7 +2749,7 @@ static FDeferredLightPS::FParameters GetDeferredLightPSParameters(
 
 	// PS - Render Targets
 	Out.RenderTargets[0] = FRenderTargetBinding(SceneColorTexture, ERenderTargetLoadAction::ELoad);
-	if (SubstrateTileType >= ESubstrateTileType::EComplex && IsExperimentalToonUnifiedDiffuseEnabled())
+	if (SubstrateTileType != ESubstrateTileType::ECount && SubstrateTileType >= ESubstrateTileType::EComplex && IsExperimentalToonUnifiedDiffuseEnabled())
 	{
 		const FSubstrateSceneData& SubstrateSceneData = *View.GetSceneTextures().SubstrateSceneData;
 		Out.RenderTargets[1] = FRenderTargetBinding(SubstrateSceneData.SubstrateSceneToonData, ERenderTargetLoadAction::ELoad);
@@ -2787,7 +2787,7 @@ uint32 SetupLightGraphicsPSOState(
 	EShaderPlatform ShaderPlatform)
 {
 	GraphicsPSOInit.PrimitiveType = PT_TriangleList;
-	if (SubstrateTileMaterialType >= ESubstrateTileType::EComplex && IsExperimentalToonUnifiedDiffuseEnabled())
+	if (SubstrateTileMaterialType != ESubstrateTileType::ECount && SubstrateTileMaterialType >= ESubstrateTileType::EComplex && IsExperimentalToonUnifiedDiffuseEnabled())
 	{
 		GraphicsPSOInit.BlendState = TStaticBlendState<
 			CW_RGBA, BO_Add, BF_One, BF_One, BO_Add, BF_One, BF_One,
