Substrate - Toon - Fixed bottom half end of the diffuse ramp not being accounted for with Unified Diffuse.

#rnx
#Synced-CL 51456025

diff --git a/Engine/Shaders/Private/DeferredLightPixelShaders.usf b/Engine/Shaders/Private/DeferredLightPixelShaders.usf
--- a/Engine/Shaders/Private/DeferredLightPixelShaders.usf
+++ b/Engine/Shaders/Private/DeferredLightPixelShaders.usf
@@ -273,7 +273,7 @@ void DeferredLightPixelMain(
 	const float2 PixelPos = SVPos.xy;
 	OutColor = 0;
 #if SUBSTRATE_OUTPUT_TOON_DATA
-	OutToonData = float4(0.0f, 0.0f, 0.0f, -1.0f);
+	OutToonData = float4(0.0f, 0.0f, 0.0f, 0.0f);
 #endif
 #if SUBSTRATE_OPAQUE_ROUGH_REFRACTION_ENABLED
 	OutOpaqueRoughRefractionSceneColor = 0;
diff --git a/Engine/Shaders/Private/Substrate/SubstrateEvaluation.ush b/Engine/Shaders/Private/Substrate/SubstrateEvaluation.ush
--- a/Engine/Shaders/Private/Substrate/SubstrateEvaluation.ush
+++ b/Engine/Shaders/Private/Substrate/SubstrateEvaluation.ush
@@ -1411,7 +1411,7 @@ FSubstrateEvaluateResult SubstrateEvaluateBSDFCommon(FSubstrateBSDFContext BSDFC
 			const bool bApplyDiffuseRamp	= SUBSTRATE_LIGHTPASS_APPLIES_TOON_DIFFUSE_RAMP;
 			const bool bApplyDiffuseShadow	= SUBSTRATE_LIGHTPASS_APPLIES_TOON_DIFFUSE_RAMP || (!SUBSTRATE_LIGHTPASS_APPLIES_TOON_DIFFUSE_RAMP && !ToonEvalSettings.bUnifiedDiffuseIncludeShadow);
 
-			Sample.ToonData					= (AreaLight.NoL*0.5+0.5) * (bApplyDiffuseShadow ? 1.0f : TransmittanceShadow);
+			Sample.ToonData					= (BSDFContext.Context.NoL * 0.5 + 0.5) * (bApplyDiffuseShadow ? 1.0f : TransmittanceShadow);
 
 			// Diffuse
 			const float3 DiffuseRamp		= bApplyDiffuseRamp ? EvaluateToonProfileDiffuseRamp(ToonProfileId, BSDFContext.Context.NoL + ToonProfileDiffuseDither) : 1.0f;
diff --git a/Engine/Source/Runtime/Renderer/Private/Substrate/Substrate.cpp b/Engine/Source/Runtime/Renderer/Private/Substrate/Substrate.cpp
--- a/Engine/Source/Runtime/Renderer/Private/Substrate/Substrate.cpp
+++ b/Engine/Source/Runtime/Renderer/Private/Substrate/Substrate.cpp
@@ -614,7 +614,7 @@ void InitialiseSubstrateFrameSceneData(FRDGBuilder& GraphBuilder, FScene& Scene,
 		if (IsExperimentalToonUnifiedDiffuseEnabled())
 		{
 			Out.SubstrateSceneToonData = GraphBuilder.CreateTexture(
-				FRDGTextureDesc::Create2D(SceneTextureExtent, PF_FloatRGBA, FClearValueBinding(FLinearColor(0.0f, 0.0f, 0.0f, -1.0f)), TexCreate_ShaderResource | TexCreate_UAV | TexCreate_RenderTargetable), TEXT("Substrate.SubstrateSceneToonData"));
+				FRDGTextureDesc::Create2D(SceneTextureExtent, PF_FloatRGBA, FClearValueBinding(FLinearColor(0.0f, 0.0f, 0.0f, 0.0f)), TexCreate_ShaderResource | TexCreate_UAV | TexCreate_RenderTargetable), TEXT("Substrate.SubstrateSceneToonData"));
 			AddClearRenderTargetPass(GraphBuilder, Out.SubstrateSceneToonData, Out.SubstrateSceneToonData->Desc.ClearValue.GetClearColor());
 		}
 
