Impove perf when updating Rect light atlas.

When available, we wave quad passing to propagate filteirng/MIP data. This is done:
* When inserting the texture into the atlas
* When generate the filtering for all MIPs
This reduce the number of MIPs pass.
This save roughly 0.1ms / 20% on XSX when the atlas is half refreshed.

Change CVar value mapping and shader code to match convention. Default behavior remains unchanged.

#Synced-CL 51490518

#Synced-CL 51490518

diff --git a/Engine/Shaders/Private/RectLightAtlas.usf b/Engine/Shaders/Private/RectLightAtlas.usf
--- a/Engine/Shaders/Private/RectLightAtlas.usf
+++ b/Engine/Shaders/Private/RectLightAtlas.usf
@@ -79,6 +79,23 @@ void MainVS(
 
 ////////////////////////////////////////////////////////////////////////////////////////////////
 // Add texture
+
+#define PERMUTATION_WAVE_MIP (PLATFORM_SUPPORTS_SM6_0_WAVE_OPERATIONS)
+
+#if PERMUTATION_WAVE_MIP
+void FillMip(RWTexture2D<float4> OutAtlasMip, uint2 InCoord, float4 InColor, bool bWrite)
+{
+	float Out = InColor;
+	float3 MipColorH = InColor + QuadReadAcrossX(InColor).xyz;
+	float3 MipColor1 = MipColorH + QuadReadAcrossY(MipColorH).xyz;
+	if (all((InCoord & 1) == 0) && bWrite)
+	{
+		const uint2 Mip1Coord = uint2(InCoord)>>1;
+		OutAtlasMip[Mip1Coord] = float4(MipColor1 * 0.25, 0);
+	} 
+}
+#endif
+
 #if SHADER_ADD_TEXTURE
 
 int InTextureMIPBias;
@@ -100,6 +117,10 @@ SamplerState InSampler5;
 SamplerState InSampler6;
 SamplerState InSampler7;
 
+#if PERMUTATION_WAVE_MIP
+RWTexture2D<float4> OutAtlasTextureMip1;
+#endif
+
 void MainPS(
 	in float4 InPosition : SV_Position,
 	in float2 InCoord : RECT_COORD,
@@ -123,6 +144,10 @@ void MainPS(
 	// Ensure there is no NaN value
 	Color = -min(-Color, 0);
 
+	#if PERMUTATION_WAVE_MIP
+	FillMip(OutAtlasTextureMip1, InPosition.xy, Color, true /*bWrite*/);
+	#endif
+
 	OutTexture = Color;
 }
 
@@ -159,11 +184,16 @@ void MainPS(
 
 float KernelSize;
 uint FilterQuality;
+uint bWriteMIP1;
 int2 SrcAtlasResolution;
 Buffer<uint4> SrcSlotBuffer;
 Texture2D<float4> SourceAtlasTexture;
 SamplerState SourceAtlasSampler;
 
+#if PERMUTATION_WAVE_MIP
+RWTexture2D<float4> OutAtlasTextureMip1;
+#endif
+
 void MainPS(
 	in float4 InPosition : SV_Position,
 	in float2 InCoord : RECT_COORD,
@@ -172,7 +202,7 @@ void MainPS(
 	out float4 OutTexture : SV_Target0)
 {
 
-	if (FilterQuality)
+	if (FilterQuality == 0)
 	{
 		// 2x2 box filtering
 		const float2 SrcCoord = uint2(InCoord) * 2.f;
@@ -197,7 +227,13 @@ void MainPS(
 		const float4 Color2 = SourceAtlasTexture.Load(uint3(Coord01, 0));
 		const float4 Color3 = SourceAtlasTexture.Load(uint3(Coord11, 0));
 
-		OutTexture = (W00*Color0 + W10*Color1 + W01*Color2 + W11*Color3) * InvW;
+		const float4 MipColor = (W00*Color0 + W10*Color1 + W01*Color2 + W11*Color3) * InvW;
+
+		#if PERMUTATION_WAVE_MIP
+		FillMip(OutAtlasTextureMip1, InPosition.xy, MipColor, bWriteMIP1);
+		#endif
+
+		OutTexture = MipColor;
 	}
 	else
 	{
diff --git a/Engine/Source/Runtime/Renderer/Private/RectLightTextureManager.cpp b/Engine/Source/Runtime/Renderer/Private/RectLightTextureManager.cpp
--- a/Engine/Source/Runtime/Renderer/Private/RectLightTextureManager.cpp
+++ b/Engine/Source/Runtime/Renderer/Private/RectLightTextureManager.cpp
@@ -66,7 +66,7 @@ static TAutoConsoleVariable<int32> CVarRectLighClearAtlas(
 
 static TAutoConsoleVariable<int32> CVarRectLighFilterQuality(
 	TEXT("r.RectLightAtlas.FilterQuality"),
-	1,
+	0,
 	TEXT("Define the filtering quality used for filtering texture (0:Box filter, 1:Gaussian filter)."),
 	ECVF_RenderThreadSafe);
 
@@ -106,6 +106,11 @@ namespace RectLightAtlas
 static const uint32 InvalidSlotIndex = ~0u;
 static const FIntPoint InvalidOrigin = FIntPoint(-1, -1);
 
+static bool SupportWaveQuadPassing(EShaderPlatform In)
+{
+	return GRHISupportsWaveOperations && RHISupportsWaveOperations(In);
+}
+
 static FIntPoint GetSourceResolution(const FRHITexture* Tex, const FVector4f& ScaleOffset)
 {
 	FIntPoint Out(1, 1);
@@ -576,6 +581,7 @@ class FRectAtlasAddTexturePS : public FGlobalShader
 		SHADER_PARAMETER_SAMPLER(SamplerState, InSampler6)
 		SHADER_PARAMETER_SAMPLER(SamplerState, InSampler7)
 		SHADER_PARAMETER_STRUCT_INCLUDE(FRectLightAtlasVS::FParameters, VS)
+		SHADER_PARAMETER_RDG_TEXTURE_UAV(RWTexture2D<float4>, OutAtlasTextureMip1)
 		RENDER_TARGET_BINDING_SLOTS()
 	END_SHADER_PARAMETER_STRUCT()
 
@@ -604,6 +610,9 @@ static void AddSlotsPass(
 	FRDGBufferRef SlotBuffer		= CreateSlotBuffer(GraphBuilder, Slots, TEXT("RectLight.AtlasSlotBuffer"));
 	FRDGBufferRef ScaleOffsetBuffer = CreateScaleOffsetBuffer(GraphBuilder, Slots, TEXT("RectLight.AtlasScaleOffsetBuffer"));
 
+	// If compiler support quad pass, use it to generate first MIP during atlas copy
+	const bool bUseQuadPassing = SupportWaveQuadPassing(ShaderMap->GetPlatform());
+
 	// Batch new slots into several passes
 	const uint32 SlotCountPerPass = 8u;
 	const uint32 PassCount = FMath::DivideAndRoundUp(uint32(Slots.Num()), SlotCountPerPass);
@@ -658,6 +667,10 @@ static void AddSlotsPass(
 		Parameters->VS.SlotBuffer = GraphBuilder.CreateSRV(SlotBuffer, PF_R16G16B16A16_UINT);
 		Parameters->VS.ScaleOffsetBuffer = GraphBuilder.CreateSRV(ScaleOffsetBuffer, PF_A32B32G32R32F);
 		Parameters->VS.bHasScaleOffset   = true;
+		if (bUseQuadPassing)
+		{
+			Parameters->OutAtlasTextureMip1 = GraphBuilder.CreateUAV(FRDGTextureUAVDesc(OutAtlas, 1));
+		}
 		Parameters->RenderTargets[0] = FRenderTargetBinding(OutAtlas, ERenderTargetLoadAction::ELoad, 0);
 
 		TShaderMapRef<FRectLightAtlasVS> VertexShader(ShaderMap);
@@ -820,11 +833,13 @@ class FRectAtlasFilterTexturePS : public FGlobalShader
 	BEGIN_SHADER_PARAMETER_STRUCT(FParameters, )
 		SHADER_PARAMETER(FIntPoint, SrcAtlasResolution)
 		SHADER_PARAMETER(uint32, FilterQuality)
+		SHADER_PARAMETER(uint32, bWriteMIP1)
 		SHADER_PARAMETER_RDG_BUFFER_SRV(Buffer, DstSlotBuffer)
 		SHADER_PARAMETER_RDG_BUFFER_SRV(Buffer, SrcSlotBuffer)
 		SHADER_PARAMETER_SAMPLER(SamplerState, SourceAtlasSampler)
 		SHADER_PARAMETER_RDG_TEXTURE_SRV(Texture2D<float4>, SourceAtlasTexture)
 		SHADER_PARAMETER_STRUCT_INCLUDE(FRectLightAtlasVS::FParameters, VS)
+		SHADER_PARAMETER_RDG_TEXTURE_UAV(RWTexture2D<float4>, OutAtlasTextureMip1)
 		RENDER_TARGET_BINDING_SLOTS()
 	END_SHADER_PARAMETER_STRUCT()
 
@@ -848,8 +863,19 @@ static void FilterSlotsPass(
 {
 	FRDGBufferSRVRef DummyScaleOffsetBuffer = GraphBuilder.CreateSRV(GSystemTextures.GetDefaultBuffer(GraphBuilder, 16u, 1u), PF_A32B32G32R32F);
 
+	// If QuadPass is available generate two MIPs at a time when using box filtering
+	const uint32 FilteringQuality = FMath::Clamp(CVarRectLighFilterQuality.GetValueOnRenderThread(), 0, 1);
+	const bool bUseQuadPassing = SupportWaveQuadPassing(ShaderMap->GetPlatform()) && FilteringQuality == 0;
+
+	FRDGTextureUAVRef DummyTextureUAV = nullptr; 
+	if (bUseQuadPassing)
+	{
+		FRDGTextureRef DummyTexture = GraphBuilder.CreateTexture(FRDGTextureDesc::Create2D(FIntPoint(1, 1), InAtlas->Desc.Format, FClearValueBinding::Black, TexCreate_UAV), TEXT("RectLight.DummyAtlasMIP"));
+		DummyTextureUAV = GraphBuilder.CreateUAV(DummyTexture);
+	}
+
 	const uint32 MipCount = InAtlas->Desc.NumMips; // FMath::Log2(FMath::Min(InAtlas->Desc.Extent.X, InAtlas->Desc.Extent.Y));
-	for (uint32 DstMip = 1; DstMip < MipCount; ++DstMip)
+	for (uint32 DstMip = (bUseQuadPassing ? 2u : 1u); DstMip < MipCount; DstMip += (bUseQuadPassing ? 2u : 1u))
 	{
 		const uint32 SrcMip = DstMip-1;
 
@@ -885,7 +911,7 @@ static void FilterSlotsPass(
 			FRDGBufferRef DstSlotBuffer = CreateSlotBuffer(GraphBuilder, DstMIPSlots, TEXT("RectLight.DstMIPSlotBuffer"));
 
 			FRectAtlasFilterTexturePS::FParameters* Parameters = GraphBuilder.AllocParameters<FRectAtlasFilterTexturePS::FParameters>();
-			Parameters->FilterQuality		= FMath::Clamp(CVarRectLighFilterQuality.GetValueOnRenderThread(), 0, 1);
+			Parameters->FilterQuality		= FilteringQuality;
 			Parameters->SrcAtlasResolution	= ToMIP(InAtlas->Desc.Extent, SrcMip);
 			Parameters->DstSlotBuffer		= GraphBuilder.CreateSRV(DstSlotBuffer, PF_R16G16B16A16_UINT);
 			Parameters->SrcSlotBuffer		= GraphBuilder.CreateSRV(SrcSlotBuffer, PF_R16G16B16A16_UINT);
@@ -894,8 +920,13 @@ static void FilterSlotsPass(
 			Parameters->VS.AtlasResolution	= ToMIP(InAtlas->Desc.Extent, DstMip);
 			Parameters->VS.SlotBufferOffset = 0;
 			Parameters->VS.SlotBuffer		= GraphBuilder.CreateSRV(DstSlotBuffer, PF_R16G16B16A16_UINT);
-			Parameters->VS.ScaleOffsetBuffer = DummyScaleOffsetBuffer;
-			Parameters->VS.bHasScaleOffset   = false;
+			Parameters->VS.ScaleOffsetBuffer= DummyScaleOffsetBuffer;
+			Parameters->VS.bHasScaleOffset  = false;
+			Parameters->bWriteMIP1			= (bUseQuadPassing && DstMip+1 < InAtlas->Desc.NumMips) ? 1u : 0u;
+			if (bUseQuadPassing)
+			{
+				Parameters->OutAtlasTextureMip1 = Parameters->bWriteMIP1 ? GraphBuilder.CreateUAV(FRDGTextureUAVDesc(InAtlas, DstMip+1)) : DummyTextureUAV;
+			}
 			Parameters->RenderTargets[0] = FRenderTargetBinding(InAtlas, ERenderTargetLoadAction::ELoad, DstMip);
 
 			TShaderMapRef<FRectLightAtlasVS> VertexShader(ShaderMap);
