remove check which can trigger from nanite skinned component for wpo raytracing dynamic geometry pso precaching

#Synced-CL 51448921

diff --git a/Engine/Source/Runtime/Renderer/Private/RayTracing/RayTracingDynamicGeometryUpdateManager.cpp b/Engine/Source/Runtime/Renderer/Private/RayTracing/RayTracingDynamicGeometryUpdateManager.cpp
--- a/Engine/Source/Runtime/Renderer/Private/RayTracing/RayTracingDynamicGeometryUpdateManager.cpp
+++ b/Engine/Source/Runtime/Renderer/Private/RayTracing/RayTracingDynamicGeometryUpdateManager.cpp
@@ -249,8 +249,7 @@ void FRayTracingDynamicGeometryPSOCollector::CollectPSOInitializers(
 	{
 		return;
 	}
-
-	check(PreCacheParams.bDefaultMaterial || VertexFactoryData.VertexFactoryType->SupportsRayTracingDynamicGeometry());
+		
 	if (!VertexFactoryData.VertexFactoryType->SupportsRayTracingDynamicGeometry())
 	{
 		return;
