Virtual Texture: Change page table update to move barriers to outside of the update render passes.
RDG was moving the per mip transition barriers to their last usage passes. That would end up serializing GPU work on platforms without split barriers.
Page table updates are also broken into begin/update/end passes so that all page tables can process the begin first, then update, then end. Barriers are issued in the begin and end.
Also  the render pass loop order was flipped to iterate over textures then mips then channels. This may help reduce context rolls.
#rnx
#Synced-CL 51491952

diff --git a/Engine/Source/Runtime/Renderer/Private/VT/VirtualTextureSpace.cpp b/Engine/Source/Runtime/Renderer/Private/VT/VirtualTextureSpace.cpp
--- a/Engine/Source/Runtime/Renderer/Private/VT/VirtualTextureSpace.cpp
+++ b/Engine/Source/Runtime/Renderer/Private/VT/VirtualTextureSpace.cpp
@@ -323,13 +323,8 @@ void FVirtualTextureSpace::AllocateTextures(FRDGBuilder& GraphBuilder)
 }
 
 
-void FVirtualTextureSpace::ApplyUpdates(FVirtualTextureSystem* System, FRDGBuilder& GraphBuilder, FRDGExternalAccessQueue& ExternalAccessQueue)
+bool FVirtualTextureSpace::BeginUpdates(FVirtualTextureSystem* System, FRDGBuilder& GraphBuilder, FUpdateContext& OutContext)
 {
-	ON_SCOPE_EXIT
-	{
-		FinalizeTextures(GraphBuilder, ExternalAccessQueue);
-	};
-
 	static TArray<FPageTableUpdate> ExpandedUpdates[VIRTUALTEXTURE_SPACE_MAXLAYERS][16];
 
 	if (bNeedToAllocatePageTable)
@@ -340,12 +335,9 @@ void FVirtualTextureSpace::ApplyUpdates(FVirtualTextureSystem* System, FRDGBuild
 		// This could cause a glitch if we sample a VT on the first frame it is allocated. That's usually not the case (we usually sample some time after loading).
 		// But it can be the case for Adaptive Virtual Texture which does a lot of dynamic page table allocation during the texture life.
 		// However Adaptive Virtual Texture is OK because it always sets bPrivateSpace which gives fixed allocation of the actual page table texture.
-		return;
+		return false;
 	}
 
-	// Multi-GPU support
-	RDG_GPU_MASK_SCOPE(GraphBuilder, FRHIGPUMask::All());
-
 	for (uint32 LayerIndex = 0u; LayerIndex < Description.NumPageTableLayers; ++LayerIndex)
 	{
 		FTexturePageMap& PageMap = PhysicalPageMap[LayerIndex];
@@ -371,46 +363,79 @@ void FVirtualTextureSpace::ApplyUpdates(FVirtualTextureSystem* System, FRDGBuild
 	}
 	bForceEntireUpdate = false;
 
-	// TODO Expand 3D updates for slices of volume texture
-
 	uint32 TotalNumUpdates = 0;
+	uint32 TextureUpdateMask = 0;
 	for (uint32 LayerIndex = 0u; LayerIndex < Description.NumPageTableLayers; ++LayerIndex)
 	{
+		const uint32 TextureIndex = LayerIndex / LayersPerPageTableTexture;
 		for (uint32 Mip = 0; Mip < CachedNumPageTableLevels; Mip++)
 		{
-			TotalNumUpdates += ExpandedUpdates[LayerIndex][Mip].Num();
+			const int32 NumUpdates = ExpandedUpdates[LayerIndex][Mip].Num();
+			TotalNumUpdates += NumUpdates;
+			TextureUpdateMask |= NumUpdates > 0 ? (1 << TextureIndex) : 0;
 		}
 	}
 
 	if (TotalNumUpdates == 0u)
 	{
-		return;
+		return false;
 	}
 
 	void* StagingMemory = GraphBuilder.Alloc(TotalNumUpdates * sizeof(FPageTableUpdate));
 	uint8* WritePtr = (uint8*)StagingMemory;
-	for (uint32 LayerIndex = 0u; LayerIndex < Description.NumPageTableLayers; ++LayerIndex)
+
+	for (uint32 TextureIndex = 0; TextureIndex < GetNumPageTableTextures(); ++TextureIndex)
 	{
+		const uint32 MinLayerIndex = TextureIndex * LayersPerPageTableTexture;
+		const uint32 MaxLayerIndex = FMath::Min(MinLayerIndex + LayersPerPageTableTexture, (uint32)Description.NumPageTableLayers);
+
 		for (uint32 Mip = 0; Mip < CachedNumPageTableLevels; Mip++)
 		{
-			const uint32 NumUpdates = ExpandedUpdates[LayerIndex][Mip].Num();
-			if (NumUpdates)
+			for (uint32 LayerIndex = MinLayerIndex; LayerIndex < MaxLayerIndex; ++LayerIndex)
 			{
-				size_t UploadSize = NumUpdates * sizeof(FPageTableUpdate);
-				FMemory::Memcpy(WritePtr, ExpandedUpdates[LayerIndex][Mip].GetData(), UploadSize);
-				WritePtr += UploadSize;
+				const uint32 NumUpdates = ExpandedUpdates[LayerIndex][Mip].Num();
+				if (NumUpdates)
+				{
+					size_t UploadSize = NumUpdates * sizeof(FPageTableUpdate);
+					FMemory::Memcpy(WritePtr, ExpandedUpdates[LayerIndex][Mip].GetData(), UploadSize);
+					ExpandedUpdates[LayerIndex][Mip].Reset();
+					WritePtr += UploadSize;
+				}
+				OutContext.NumUpdates[LayerIndex][Mip] = NumUpdates;
 			}
 		}
 	}
 
 	FRDGBufferRef UpdateBuffer = GraphBuilder.CreateBuffer(FRDGBufferDesc::CreateBufferDesc(sizeof(FPageTableUpdate), TotalNumUpdates), TEXT("VirtualTexture_PageTableUpdateBuffer"));
 	GraphBuilder.QueueBufferUpload(UpdateBuffer, StagingMemory, TotalNumUpdates * sizeof(FPageTableUpdate), ERDGInitialDataFlags::NoCopy);
-	
-	FRDGBufferSRVRef UpdateBufferSRV = GraphBuilder.CreateSRV(FRDGBufferSRVDesc(UpdateBuffer, PF_R16G16B16A16_UINT));
 
-	// Draw
-	RDG_EVENT_SCOPE(GraphBuilder, "PageTableUpdate");
+	OutContext.UpdateBufferSRV = GraphBuilder.CreateSRV(FRDGBufferSRVDesc(UpdateBuffer, PF_R16G16B16A16_UINT));
 
+	FRDGTextureAccessArray BatchBarrierTextures;
+	for (uint32 TextureIndex = 0; TextureIndex < GetNumPageTableTextures(); ++TextureIndex)
+	{
+		if ((TextureUpdateMask & (1 << TextureIndex)) != 0)
+		{
+			// ForceImmediateFirstBarrier so that RTV transitions for the page table textures aren't hoisted
+			// into RenderFinalize() or earlier where they will be incorrect for virtual texture sampling.
+			FTextureEntry& PageTableEntry = PageTable[TextureIndex];
+			FRDGTextureRef PageTableTexture = GraphBuilder.RegisterExternalTexture(PageTableEntry.RenderTarget, ERDGTextureFlags::ForceImmediateFirstBarrier);
+			GraphBuilder.UseInternalAccessMode(PageTableTexture);
+
+			OutContext.Textures[TextureIndex] = PageTableTexture;
+			BatchBarrierTextures.Add(FRDGTextureAccess(PageTableTexture, ERHIAccess::RTV));
+		}
+	}
+
+	// Add barrier pass to move all SRV->RTV transition barriers before all the render passes.
+	AddBarrierPass(GraphBuilder, ERDGPassFlags::Raster | ERDGPassFlags::SkipRenderPass | ERDGPassFlags::NeverMerge, BatchBarrierTextures);
+
+	OutContext.Space = this;
+	return true;
+}
+
+void FVirtualTextureSpace::ApplyUpdates(FRDGBuilder& GraphBuilder, FUpdateContext const& Context)
+{
 	auto ShaderMap = GetGlobalShaderMap(GetFeatureLevel());
 	TShaderRef<FPageTableUpdateVS> VertexShader;
 	if (Description.PageTableFormat == EVTPageTableFormat::UInt16)
@@ -424,26 +449,15 @@ void FVirtualTextureSpace::ApplyUpdates(FVirtualTextureSystem* System, FRDGBuild
 	check(VertexShader.IsValid());
 
 	uint32 FirstUpdate = 0;
-	for (uint32 LayerIndex = 0u; LayerIndex < Description.NumPageTableLayers; ++LayerIndex)
-	{
-		const uint32 TextureIndex = LayerIndex / LayersPerPageTableTexture;
-		const uint32 LayerInTexture = LayerIndex % LayersPerPageTableTexture;
-		FTextureEntry& PageTableEntry = PageTable[TextureIndex];
-
-		// ForceImmediateFirstBarrier so that RTV transitions for the page table textures aren't hoisted
-		// into RenderFinalize() or earlier where they will be incorrect for virtual texture sampling.
-		FRDGTextureRef PageTableTexture = GraphBuilder.RegisterExternalTexture(PageTableEntry.RenderTarget, ERDGTextureFlags::ForceImmediateFirstBarrier);
-		GraphBuilder.UseInternalAccessMode(PageTableTexture);
 
-		// Use color write mask to update the proper page table entry for this layer
-		FRHIBlendState* BlendStateRHI = nullptr;
-		switch (LayerInTexture)
+	// Loop over mips for each texture.
+	// Inner loop is over texture channels setting the channel mask state.
+	for (uint32 TextureIndex = 0; TextureIndex < GetNumPageTableTextures(); ++TextureIndex)
+	{
+		FRDGTextureRef PageTableTexture = Context.Textures[TextureIndex];
+		if (PageTableTexture == nullptr)
 		{
-		case 0u: BlendStateRHI = TStaticBlendState<CW_RED>::GetRHI(); break;
-		case 1u: BlendStateRHI = TStaticBlendState<CW_GREEN>::GetRHI(); break;
-		case 2u: BlendStateRHI = TStaticBlendState<CW_BLUE>::GetRHI(); break;
-		case 3u: BlendStateRHI = TStaticBlendState<CW_ALPHA>::GetRHI(); break;
-		default: check(false); break;
+			continue;
 		}
 
 		TShaderRef<FGlobalShader> PixelShader;
@@ -459,22 +473,44 @@ void FVirtualTextureSpace::ApplyUpdates(FVirtualTextureSystem* System, FRDGBuild
 		}
 		check(PixelShader.IsValid());
 
+		const uint32 MinLayerIndex = TextureIndex * LayersPerPageTableTexture;
+		const uint32 MaxLayerIndex = FMath::Min(MinLayerIndex + LayersPerPageTableTexture, (uint32)Description.NumPageTableLayers);
+		
 		uint32 MipWidth = CachedPageTableWidth;
 		uint32 MipHeight = CachedPageTableHeight;
+	
 		for (uint32 Mip = 0; Mip < CachedNumPageTableLevels; Mip++)
 		{
-			const uint32 NumUpdates = ExpandedUpdates[LayerIndex][Mip].Num();
-			if (NumUpdates)
+			for (uint32 LayerIndex = MinLayerIndex; LayerIndex < MaxLayerIndex; ++LayerIndex)
 			{
+				const uint32 NumUpdates = Context.NumUpdates[LayerIndex][Mip];
+				if (NumUpdates == 0)
+				{
+					continue;
+				}
+
+				const uint32 LayerInTexture = LayerIndex - MinLayerIndex;
+
+				// Use color write mask to update the proper page table entry for this layer
+				FRHIBlendState* BlendStateRHI = nullptr;
+				switch (LayerInTexture)
+				{
+				case 0u: BlendStateRHI = TStaticBlendState<CW_RED>::GetRHI(); break;
+				case 1u: BlendStateRHI = TStaticBlendState<CW_GREEN>::GetRHI(); break;
+				case 2u: BlendStateRHI = TStaticBlendState<CW_BLUE>::GetRHI(); break;
+				case 3u: BlendStateRHI = TStaticBlendState<CW_ALPHA>::GetRHI(); break;
+				default: check(false); break;
+				}
+
 				auto* PassParameters = GraphBuilder.AllocParameters<FPageTableUpdateParameters>();
-				PassParameters->VS.UpdateBuffer = UpdateBufferSRV;
+				PassParameters->VS.UpdateBuffer = Context.UpdateBufferSRV;
 				PassParameters->VS.PageTableSize = FUintVector2(CachedPageTableWidth, CachedPageTableHeight);
 				PassParameters->VS.FirstUpdate = FirstUpdate;
 				PassParameters->VS.NumUpdates = NumUpdates;
 				PassParameters->RenderTargets[0] = FRenderTargetBinding(PageTableTexture, ERenderTargetLoadAction::ELoad, Mip);
 
 				GraphBuilder.AddPass(
-					RDG_EVENT_NAME("PageTableUpdate (Mip: %d)", Mip),
+					RDG_EVENT_NAME("PageTableUpdate (Id: %d, Tex %d, Mip: %d)", ID, TextureIndex, Mip),
 					PassParameters,
 					ERDGPassFlags::Raster,
 					[PassParameters, VertexShader, PixelShader, BlendStateRHI, NumUpdates, MipWidth, MipHeight](FRDGAsyncTask, FRHICommandList& RHICmdList)
@@ -504,31 +540,37 @@ void FVirtualTextureSpace::ApplyUpdates(FVirtualTextureSystem* System, FRDGBuild
 					RHICmdList.DrawIndexedPrimitive(GQuadIndexBuffer.IndexBufferRHI, 0, 0, 32, 0, 2 * QuadsPerInstance, FMath::DivideAndRoundUp(NumUpdates, QuadsPerInstance));
 				});
 
-				ExpandedUpdates[LayerIndex][Mip].Reset();
+				FirstUpdate += NumUpdates;
 			}
 
-			FirstUpdate += NumUpdates;
 			MipWidth = FMath::Max(MipWidth / 2u, 1u);
 			MipHeight = FMath::Max(MipHeight / 2u, 1u);
 		}
-
-		PageTableEntry.RenderTarget = GraphBuilder.ConvertToExternalTexture(PageTableTexture);
 	}
 }
 
-void FVirtualTextureSpace::FinalizeTextures(FRDGBuilder& GraphBuilder, FRDGExternalAccessQueue& ExternalAccessQueue)
+void FVirtualTextureSpace::EndUpdates(FRDGBuilder& GraphBuilder, FUpdateContext const& Context, FRDGExternalAccessQueue& ExternalAccessQueue)
 {
-	for (uint32 LayerIndex = 0u; LayerIndex < Description.NumPageTableLayers; ++LayerIndex)
+	// Add barrier pass to keep all RTV->SRV transition after all the render passes.
+	// Marking the access as RTV here stops the RDG barrier optimizer from moving the barriers earlier where they serialize the render passes.
+	FRDGTextureAccessArray BatchBarrierTextures;
+	for (FRDGTextureRef Texture : Context.Textures)
 	{
-		const uint32 TextureIndex = LayerIndex / LayersPerPageTableTexture;
-		FTextureEntry& PageTableEntry = PageTable[TextureIndex];
-		if (PageTableEntry.RenderTarget)
+		if (Texture != nullptr)
 		{
-			// It's only necessary to enable external access mode on textures modified by RDG this frame.
-			if (FRDGTexture* Texture = GraphBuilder.FindExternalTexture(PageTableEntry.RenderTarget))
-			{
-				ExternalAccessQueue.Add(Texture, ERHIAccess::SRVMask, ERHIPipeline::All);
-			}
+			BatchBarrierTextures.Add(FRDGTextureAccess(Texture, ERHIAccess::RTV));
+		}
+	}
+	AddBarrierPass(GraphBuilder, ERDGPassFlags::Raster | ERDGPassFlags::SkipRenderPass | ERDGPassFlags::NeverMerge, BatchBarrierTextures);
+
+	// Return textures to be externally managed.
+	for (uint32 TextureIndex = 0; TextureIndex < TextureCapacity; ++TextureIndex)
+	{
+		if (FRDGTextureRef Texture = Context.Textures[TextureIndex])
+		{
+			FTextureEntry& PageTableEntry = PageTable[TextureIndex];
+			PageTableEntry.RenderTarget = GraphBuilder.ConvertToExternalTexture(Texture);
+			ExternalAccessQueue.Add(Texture, ERHIAccess::SRVMask, ERHIPipeline::All);
 		}
 	}
 }
diff --git a/Engine/Source/Runtime/Renderer/Private/VT/VirtualTextureSpace.h b/Engine/Source/Runtime/Renderer/Private/VT/VirtualTextureSpace.h
--- a/Engine/Source/Runtime/Renderer/Private/VT/VirtualTextureSpace.h
+++ b/Engine/Source/Runtime/Renderer/Private/VT/VirtualTextureSpace.h
@@ -48,6 +48,7 @@ class FVirtualTextureSpace final : public FRenderResource
 {
 public:
 	static const uint32 LayersPerPageTableTexture = IAllocatedVirtualTexture::LayersPerPageTableTexture;
+	static const uint32 TextureCapacity = (VIRTUALTEXTURE_SPACE_MAXLAYERS + LayersPerPageTableTexture - 1u) / LayersPerPageTableTexture;
 
 	FVirtualTextureSpace(FVirtualTextureSystem* InSystem, uint8 InID, const FVTSpaceDescription& InDesc, uint32 InSizeNeeded);
 	virtual ~FVirtualTextureSpace();
@@ -100,11 +101,22 @@ public:
 		return PageTableIndirection.RenderTarget;
 	}
 
-	void				QueueUpdate( uint8 Layer, uint8 vLogSize, uint32 vAddress, uint8 vLevel, const FPhysicalTileLocation& pTileLocation);
-	void				AllocateTextures(FRDGBuilder& GraphBuilder);
-	void				FinalizeTextures(FRDGBuilder& GraphBuilder, FRDGExternalAccessQueue& ExternalAccessQueue);
-	void				ApplyUpdates(FVirtualTextureSystem* System, FRDGBuilder& GraphBuilder, FRDGExternalAccessQueue& ExternalAccessQueue);
-	void				QueueUpdateEntirePageTable();
+	void QueueUpdate(uint8 Layer, uint8 vLogSize, uint32 vAddress, uint8 vLevel, const FPhysicalTileLocation& pTileLocation);
+	void AllocateTextures(FRDGBuilder& GraphBuilder);
+	void QueueUpdateEntirePageTable();
+
+	/** Update context. For efficient generation of GPU barriers we want to call each update phase (Begin/Apply/End) on all spaces before moving to the next phase. */
+	struct FUpdateContext
+	{
+		FVirtualTextureSpace* Space = nullptr;
+		FRDGBufferSRVRef UpdateBufferSRV = nullptr;
+		FRDGTextureRef Textures[TextureCapacity] = { nullptr };
+		int32 NumUpdates[VIRTUALTEXTURE_SPACE_MAXLAYERS][16] = {{ 0 }};
+	};
+	
+	bool BeginUpdates(FVirtualTextureSystem* System, FRDGBuilder& GraphBuilder, FUpdateContext& OutContext);
+	void ApplyUpdates(FRDGBuilder& GraphBuilder, FUpdateContext const& Context);
+	void EndUpdates(FRDGBuilder& GraphBuilder, FUpdateContext const& Context, FRDGExternalAccessQueue& ExternalAccessQueue);
 
 	void DumpToConsole(bool verbose);
 
@@ -115,7 +127,6 @@ public:
 private:
 	FUintPoint GetRequiredPageTableAllocationSize() const;
 
-	static const uint32 TextureCapacity = (VIRTUALTEXTURE_SPACE_MAXLAYERS + LayersPerPageTableTexture - 1u) / LayersPerPageTableTexture;
 
 	struct FTextureEntry
 	{
diff --git a/Engine/Source/Runtime/Renderer/Private/VT/VirtualTextureSystem.cpp b/Engine/Source/Runtime/Renderer/Private/VT/VirtualTextureSystem.cpp
--- a/Engine/Source/Runtime/Renderer/Private/VT/VirtualTextureSystem.cpp
+++ b/Engine/Source/Runtime/Renderer/Private/VT/VirtualTextureSystem.cpp
@@ -2547,14 +2547,31 @@ void FVirtualTextureSystem::FinalizeRequests(FRDGBuilder& GraphBuilder, ISceneRe
 		SCOPE_CYCLE_COUNTER(STAT_PageTableUpdates);
 		RDG_EVENT_SCOPE_STAT(GraphBuilder, VirtualTexture, "VirtualTexturePageTableUpdates");
 		TRACE_CPUPROFILER_EVENT_SCOPE(FVirtualTextureSystem::ApplyUpdates);
+		RDG_GPU_MASK_SCOPE(GraphBuilder, FRHIGPUMask::All());
 
-		for (int32 ID = 0; ID < Spaces.Num(); ID++)
+		TArray<FVirtualTextureSpace::FUpdateContext, TInlineAllocator<16>> UpdateContexts;
+		for (TUniquePtr<FVirtualTextureSpace>& Space : Spaces)
 		{
-			if (Spaces[ID])
+			FVirtualTextureSpace::FUpdateContext Context;
+			if (Space.IsValid() && Space->BeginUpdates(this, GraphBuilder, Context))
 			{
-				Spaces[ID]->ApplyUpdates(this, GraphBuilder, ExternalAccessQueue);
+				UpdateContexts.Add(Context);
+			}
+		}
+		for (FVirtualTextureSpace::FUpdateContext const& Context : UpdateContexts)
+		{
+			Context.Space->ApplyUpdates(GraphBuilder, Context);
+		}
+		for (FVirtualTextureSpace::FUpdateContext const& Context : UpdateContexts)
+		{
+			Context.Space->EndUpdates(GraphBuilder, Context, ExternalAccessQueue);
+		}
 
-				if (FAdaptiveVirtualTexture* AdaptiveVT = GetAdaptiveVirtualTexture(ID))
+		for (TUniquePtr<FVirtualTextureSpace>& Space : Spaces)
+		{
+			if (Space.IsValid())
+			{
+				if (FAdaptiveVirtualTexture* AdaptiveVT = GetAdaptiveVirtualTexture(Space->GetID()))
 				{
 					AdaptiveVT->ApplyPageTableUpdates(this, GraphBuilder, ExternalAccessQueue);
 				}
