[PostProcessing] Fix resource transition validation error on BloomDirtMask and convert various usages of fallback textures to RDG.

#rnx

#Synced-CL 51468985

diff --git a/Engine/Source/Runtime/RenderCore/Private/GlobalRenderResources.cpp b/Engine/Source/Runtime/RenderCore/Private/GlobalRenderResources.cpp
--- a/Engine/Source/Runtime/RenderCore/Private/GlobalRenderResources.cpp
+++ b/Engine/Source/Runtime/RenderCore/Private/GlobalRenderResources.cpp
@@ -52,7 +52,7 @@ public:
 	{
 		// BGRA typed UAV is unsupported per D3D spec, use RGBA here.
 		const FRHITextureCreateDesc Desc =
-			FRHITextureCreateDesc::Create2D(TEXT("ColoredTexture"), 1, 1, PF_R8G8B8A8)
+			FRHITextureCreateDesc::Create2D(GetTextureName(), 1, 1, PF_R8G8B8A8)
 			.SetFlags(ETextureCreateFlags::ShaderResource)
 			.SetClassName(TEXT("FColoredTexture"))
 			.SetInitActionInitializer();
@@ -75,6 +75,7 @@ public:
 
 	virtual uint32 GetSizeX() const override { return 1; }
 	virtual uint32 GetSizeY() const override { return 1; }
+	virtual const TCHAR* GetTextureName() const { return TEXT("ColoredTexture"); }
 };
 
 class FEmptyVertexBuffer : public FVertexBufferWithSRV
@@ -228,6 +229,16 @@ public:
 	virtual uint32 GetSizeY() const override { return 1; }
 };
 
+class FWhiteTexture : public FColoredTexture<255, 255, 255, 255>
+{
+	virtual const TCHAR* GetTextureName() const override { return TEXT("FWhiteTexture"); }
+};
+
+class FBlackTexture : public FColoredTexture<0, 0, 0, 0>
+{
+	virtual const TCHAR* GetTextureName() const override { return TEXT("FBlackTexture"); }
+};
+
 class FBlackTextureWithSRV : public FColoredTexture<0, 0, 0, 255>
 {
 public:
@@ -242,11 +253,13 @@ public:
 		FRHITextureReference::DefaultTexture.SafeRelease();
 		FColoredTexture::ReleaseRHI();
 	}
+
+	virtual const TCHAR* GetTextureName() const override { return TEXT("FBlackTextureWithSRV"); }
 };
 
-FTextureWithSRV* GWhiteTextureWithSRV = new TGlobalResource<FColoredTexture<255, 255, 255, 255>, FRenderResource::EInitPhase::Pre>;
+FTextureWithSRV* GWhiteTextureWithSRV = new TGlobalResource<FWhiteTexture, FRenderResource::EInitPhase::Pre>;
 FTextureWithSRV* GBlackTextureWithSRV = new TGlobalResource<FBlackTextureWithSRV, FRenderResource::EInitPhase::Pre>();
-FTextureWithSRV* GTransparentBlackTextureWithSRV = new TGlobalResource<FColoredTexture<0, 0, 0, 0>, FRenderResource::EInitPhase::Pre>;
+FTextureWithSRV* GTransparentBlackTextureWithSRV = new TGlobalResource<FBlackTexture, FRenderResource::EInitPhase::Pre>;
 FTexture* GDummyTransitionTexture = new TGlobalResource<FDummyTransitionTexture, FRenderResource::EInitPhase::Pre>;
 
 FTexture* GWhiteTexture = GWhiteTextureWithSRV;
diff --git a/Engine/Source/Runtime/Renderer/Private/PostProcess/PostProcessEyeAdaptation.cpp b/Engine/Source/Runtime/Renderer/Private/PostProcess/PostProcessEyeAdaptation.cpp
--- a/Engine/Source/Runtime/Renderer/Private/PostProcess/PostProcessEyeAdaptation.cpp
+++ b/Engine/Source/Runtime/Renderer/Private/PostProcess/PostProcessEyeAdaptation.cpp
@@ -533,7 +533,22 @@ void FinishUpdateExposureCompensationCurveLUT(FRHICommandListImmediate& RHICmdLi
 	}
 }
 
-FEyeAdaptationParameters GetEyeAdaptationParameters(const FViewInfo& View)
+static bool IsCurveLUTUsed(const FViewInfo& View)
+{
+	const bool IsDebugViewMode = IsAutoExposureDebugMode(View);
+
+	const bool bUseCurveLUT =
+		View.FinalPostProcessSettings.AutoExposureBiasCurve != nullptr
+		&& View.ViewState != nullptr
+		&& View.ViewState->GetExposureCompensationCurveLUT() != nullptr
+		&& CVarEyeAdaptationExposureCompensationCurveLUT.GetValueOnRenderThread()
+		&& !IsDebugViewMode
+		&& !View.Family->ExposureSettings.bFixed;
+
+	return bUseCurveLUT;
+}
+
+FEyeAdaptationScalarParameters GetEyeAdaptationScalarParameters(const FViewInfo& View)
 {
 	const bool bExtendedLuminanceRange = IsExtendLuminanceRangeEnabled();
 
@@ -558,29 +573,9 @@ FEyeAdaptationParameters GetEyeAdaptationParameters(const FViewInfo& View)
 	float MinWhitePointLuminance = 1.0f;
 	float MaxWhitePointLuminance = 1.0f;
 
-	const bool bUseCurveLUT = 
-		Settings.AutoExposureBiasCurve != nullptr
-		&& View.ViewState != nullptr
-		&& View.ViewState->GetExposureCompensationCurveLUT() != nullptr
-		&& CVarEyeAdaptationExposureCompensationCurveLUT.GetValueOnRenderThread();
-
-	FTextureRHIRef ExposureCompensationCurveLUT = nullptr;
-
-	if (bUseCurveLUT)
-	{
-		ExposureCompensationCurveLUT = View.ViewState->GetExposureCompensationCurveLUT();
-	}
-	else
-	{
-		ExposureCompensationCurveLUT = GWhiteTexture->TextureRHI;
-	}
-
 	// Get the exposure compensation from the post process volume settings (everything except the curve)
 	float ExposureCompensationSettings = GetAutoExposureCompensationFromSettings(View);
 
-	// Get the exposure compensation from the curve
-	float ExposureCompensationCurve = bUseCurveLUT ? 1.0f : GetAutoExposureCompensationFromCurve(View);
-
 	const float BlackHistogramBucketInfluence = CVarEyeAdaptationBlackHistogramBucketInfluence.GetValueOnRenderThread();
 
 	const float kMiddleGrey = 0.18f;
@@ -590,18 +585,21 @@ FEyeAdaptationParameters GetEyeAdaptationParameters(const FViewInfo& View)
 
 	const bool IsDebugViewMode = IsAutoExposureDebugMode(View);
 
+	const bool bUseCurveLUT = IsCurveLUTUsed(View);
+
+	// Get the exposure compensation from the curve
+	float ExposureCompensationCurve = bUseCurveLUT ? 1.0f : GetAutoExposureCompensationFromCurve(View);
+
 	if (IsDebugViewMode)
 	{
 		ExposureCompensationSettings = 1.0f;
 		ExposureCompensationCurve = 1.0f;
-		ExposureCompensationCurveLUT = GWhiteTexture->TextureRHI;
 	}
 	// Fixed exposure override in effect.
 	else if (View.Family->ExposureSettings.bFixed)
 	{
 		ExposureCompensationSettings = 1.0f;
 		ExposureCompensationCurve = 1.0f;
-		ExposureCompensationCurveLUT = GWhiteTexture->TextureRHI;
 
 		// ignores bExtendedLuminanceRange
 		MinWhitePointLuminance = MaxWhitePointLuminance = CalculateFixedAutoExposure(View);
@@ -652,19 +650,6 @@ FEyeAdaptationParameters GetEyeAdaptationParameters(const FViewInfo& View)
 	// mode does the calculation in pre-exposure space, which is why we need to multiply by View.PreExposure.
 	const float LuminanceMin = (AutoExposureMethod == AEM_Basic) ? 0.0001f : FMath::Exp2(HistogramLogMin);
 
-	FTextureRHIRef MeterMask = nullptr;
-
-	if (Settings.AutoExposureMeterMask &&
-		Settings.AutoExposureMeterMask->GetResource() &&
-		Settings.AutoExposureMeterMask->GetResource()->TextureRHI)
-	{
-		MeterMask = Settings.AutoExposureMeterMask->GetResource()->TextureRHI;
-	}
-	else
-	{
-		MeterMask = GWhiteTexture->TextureRHI;
-	}
-
 	// The distance at which we switch from linear to exponential. I.e. at StartDistance=1.5, when linear is 1.5 f-stops away from hitting the 
 	// target, we switch to exponential.
 	const float StartDistance = CVarEyeAdaptationExponentialTransitionDistance.GetValueOnRenderThread();
@@ -695,7 +680,25 @@ FEyeAdaptationParameters GetEyeAdaptationParameters(const FViewInfo& View)
 	// if we don't have a valid range (AutoExposureMinBrightness == AutoExposureMaxBrightness) then force it like Manual as well.
 	const float ForceTarget = (View.bCameraCut || View.bPrevTransformsReset || AutoExposureMethod == EAutoExposureMethod::AEM_Manual || !bValidRange || !bValidSpeeds) ? 1.0f : 0.0f;
 
-	FEyeAdaptationParameters Parameters;
+	FVector3f LuminanceWeights;
+	const int32 LuminanceMethod = CVarAutoExposureLuminanceMethod.GetValueOnRenderThread();
+	if (LuminanceMethod == 1)
+	{
+		// NTSC (deprecated legacy weights in Common.ush)
+		LuminanceWeights = FVector3f(0.3f, 0.59f, 0.11f);
+	}
+	else if (LuminanceMethod == 2)
+	{
+		// Working color space (sRGB/Rec709 default)
+		LuminanceWeights = FVector3f(UE::Color::FColorSpace::GetWorking().GetLuminanceFactors());
+	}
+	else
+	{
+		// default (uniform weights)
+		LuminanceWeights = FVector3f(1.0f / 3.0f);
+	}
+
+	FEyeAdaptationScalarParameters Parameters;
 	Parameters.ExposureLowPercent = ExposureLowPercent;
 	Parameters.ExposureHighPercent = ExposureHighPercent;
 	Parameters.MinAverageLuminance = MinAverageLuminance;
@@ -721,34 +724,50 @@ FEyeAdaptationParameters GetEyeAdaptationParameters(const FViewInfo& View)
 	Parameters.ForceTarget = ForceTarget;
 	Parameters.VisualizeDebugType = CVarEyeAdaptationVisualizeDebugType.GetValueOnRenderThread();
 	Parameters.EV100ToExposureCompensationCurveLUTScaleBias = GetExposureCompensationCurveLUTScaleBias();
-	Parameters.ExposureCompensationCurveLUT = ExposureCompensationCurveLUT;
-	Parameters.ExposureCompensationCurveSampler = TStaticSamplerState<SF_Bilinear, AM_Clamp, AM_Clamp, AM_Clamp>::GetRHI();
-	Parameters.MeterMaskTexture = MeterMask;
-	Parameters.MeterMaskSampler = TStaticSamplerState<SF_Bilinear, AM_Clamp, AM_Clamp, AM_Clamp>::GetRHI();
+	Parameters.LuminanceWeights = LuminanceWeights;
 
-	const int32 LuminanceMethod = CVarAutoExposureLuminanceMethod.GetValueOnRenderThread();
-	if (LuminanceMethod == 1)
-	{
-		// NTSC (deprecated legacy weights in Common.ush)
-		Parameters.LuminanceWeights = FVector3f(0.3f, 0.59f, 0.11f);
-	}
-	else if (LuminanceMethod == 2)
+	return Parameters;
+}
+
+FEyeAdaptationParameters GetEyeAdaptationParameters(FRDGBuilder& GraphBuilder, const FViewInfo& View)
+{
+	FEyeAdaptationScalarParameters ScalarParameters = GetEyeAdaptationScalarParameters(View);
+
+	const FPostProcessSettings& Settings = View.FinalPostProcessSettings;
+
+	const FEngineShowFlags& EngineShowFlags = View.Family->EngineShowFlags;
+
+	const bool IsDebugViewMode = IsAutoExposureDebugMode(View);
+
+	const bool bUseCurveLUT = IsCurveLUTUsed(View);
+
+	FRDGTextureRef ExposureCompensationCurveLUT = GSystemTextures.GetWhiteDummy(GraphBuilder);
+	if (bUseCurveLUT)
 	{
-		// Working color space (sRGB/Rec709 default)
-		Parameters.LuminanceWeights = FVector3f(UE::Color::FColorSpace::GetWorking().GetLuminanceFactors());
+		ExposureCompensationCurveLUT = RegisterExternalTexture(GraphBuilder, View.ViewState->GetExposureCompensationCurveLUT(), TEXT("ExposureCompensationCurveLUT"));
 	}
-	else
+
+	FRDGTextureRef MeterMask = GSystemTextures.GetWhiteDummy(GraphBuilder);
+	if (Settings.AutoExposureMeterMask &&
+		Settings.AutoExposureMeterMask->GetResource() &&
+		Settings.AutoExposureMeterMask->GetResource()->TextureRHI)
 	{
-		// default (uniform weights)
-		Parameters.LuminanceWeights = FVector3f(1.0f / 3.0f);
+		MeterMask = RegisterExternalTexture(GraphBuilder, Settings.AutoExposureMeterMask->GetResource()->TextureRHI, TEXT("MeterMask"));
 	}
 
+	FEyeAdaptationParameters Parameters;
+	Parameters.Scalars = ScalarParameters;
+	Parameters.ExposureCompensationCurveLUT = ExposureCompensationCurveLUT;
+	Parameters.ExposureCompensationCurveSampler = TStaticSamplerState<SF_Bilinear, AM_Clamp, AM_Clamp, AM_Clamp>::GetRHI();
+	Parameters.MeterMaskTexture = MeterMask;
+	Parameters.MeterMaskSampler = TStaticSamplerState<SF_Bilinear, AM_Clamp, AM_Clamp, AM_Clamp>::GetRHI();
+
 	return Parameters;
 }
 
 float GetEyeAdaptationFixedExposure(const FViewInfo& View)
 {
-	const FEyeAdaptationParameters Parameters = GetEyeAdaptationParameters(View);
+	const FEyeAdaptationScalarParameters Parameters = GetEyeAdaptationScalarParameters(View);
 
 	const float Exposure = (Parameters.MinAverageLuminance + Parameters.MaxAverageLuminance) * 0.5f;
 
@@ -975,7 +994,7 @@ FRDGTextureRef AddSetupExposureIlluminancePass(
 			PassParameters->ColorTexture = SceneTextures.Color.Resolve;
 			PassParameters->RWIlluminanceTexture = GraphBuilder.CreateUAV(OutputTexture);
 			PassParameters->Illuminance = GetScreenPassTextureViewportParameters(OutputViewport);
-			PassParameters->EyeAdaptation = GetEyeAdaptationParameters(View);
+			PassParameters->EyeAdaptation = GetEyeAdaptationParameters(GraphBuilder, View);
 			PassParameters->IlluminanceRectToColorRect = FScreenTransform::ChangeRectFromTo(OutputViewport.Rect, SceneViewport.Rect);
 
 			auto ComputeShader = View.ShaderMap->GetShader<FSetupExposureIlluminanceCS>();
@@ -1055,7 +1074,7 @@ FRDGTextureRef AddCalculateExposureIlluminancePass(
 			PassParameters->RWIlluminanceTexture = GraphBuilder.CreateUAV(ExposureIlluminanceSetup);
 			PassParameters->Illuminance = GetScreenPassTextureViewportParameters(OutputViewport);
 			PassParameters->IlluminanceRectToColorRect = FScreenTransform::ChangeRectFromTo(OutputViewport.Rect, SceneViewport.Rect);
-			PassParameters->EyeAdaptation = GetEyeAdaptationParameters(View);
+			PassParameters->EyeAdaptation = GetEyeAdaptationParameters(GraphBuilder, View);
 
 			PassParameters->PreIntegratedGF = GSystemTextures.PreintegratedGF->GetRHI();
 			PassParameters->PreIntegratedGFSampler = TStaticSamplerState<SF_Bilinear, AM_Clamp, AM_Clamp, AM_Clamp>::GetRHI();
diff --git a/Engine/Source/Runtime/Renderer/Private/PostProcess/PostProcessEyeAdaptation.h b/Engine/Source/Runtime/Renderer/Private/PostProcess/PostProcessEyeAdaptation.h
--- a/Engine/Source/Runtime/Renderer/Private/PostProcess/PostProcessEyeAdaptation.h
+++ b/Engine/Source/Runtime/Renderer/Private/PostProcess/PostProcessEyeAdaptation.h
@@ -65,7 +65,7 @@ struct FUpdateExposureCompensationCurveLUTTaskData
 void BeginUpdateExposureCompensationCurveLUT(TArrayView<FViewInfo> Views, FUpdateExposureCompensationCurveLUTTaskData* TaskData);
 void FinishUpdateExposureCompensationCurveLUT(FRHICommandListImmediate& RHICmdList, FUpdateExposureCompensationCurveLUTTaskData* TaskData);
 
-BEGIN_SHADER_PARAMETER_STRUCT(FEyeAdaptationParameters, )
+BEGIN_SHADER_PARAMETER_STRUCT(FEyeAdaptationScalarParameters, )
 	SHADER_PARAMETER(float, ExposureLowPercent)
 	SHADER_PARAMETER(float, ExposureHighPercent)
 	SHADER_PARAMETER(float, MinAverageLuminance)
@@ -92,13 +92,19 @@ BEGIN_SHADER_PARAMETER_STRUCT(FEyeAdaptationParameters, )
 	SHADER_PARAMETER(float, ForceTarget)
 	SHADER_PARAMETER(int, VisualizeDebugType)
 	SHADER_PARAMETER(FVector2f, EV100ToExposureCompensationCurveLUTScaleBias)
-	SHADER_PARAMETER_TEXTURE(Texture2D, ExposureCompensationCurveLUT)
+END_SHADER_PARAMETER_STRUCT()
+
+BEGIN_SHADER_PARAMETER_STRUCT(FEyeAdaptationParameters, )
+	SHADER_PARAMETER_STRUCT_INCLUDE(FEyeAdaptationScalarParameters, Scalars)
+
+	SHADER_PARAMETER_RDG_TEXTURE(Texture2D, ExposureCompensationCurveLUT)
 	SHADER_PARAMETER_SAMPLER(SamplerState, ExposureCompensationCurveSampler)
-	SHADER_PARAMETER_TEXTURE(Texture2D, MeterMaskTexture)
+	SHADER_PARAMETER_RDG_TEXTURE(Texture2D, MeterMaskTexture)
 	SHADER_PARAMETER_SAMPLER(SamplerState, MeterMaskSampler)
 END_SHADER_PARAMETER_STRUCT()
 
-FEyeAdaptationParameters GetEyeAdaptationParameters(const FViewInfo& ViewInfo);
+FEyeAdaptationScalarParameters GetEyeAdaptationScalarParameters(const FViewInfo& ViewInfo);
+FEyeAdaptationParameters GetEyeAdaptationParameters(FRDGBuilder& GraphBuilder, const FViewInfo& ViewInfo);
 
 // Computes the a fixed exposure to be used to replace the dynamic exposure when it's not supported (< SM5).
 float GetEyeAdaptationFixedExposure(const FViewInfo& View);
diff --git a/Engine/Source/Runtime/Renderer/Private/PostProcess/PostProcessLocalExposure.cpp b/Engine/Source/Runtime/Renderer/Private/PostProcess/PostProcessLocalExposure.cpp
--- a/Engine/Source/Runtime/Renderer/Private/PostProcess/PostProcessLocalExposure.cpp
+++ b/Engine/Source/Runtime/Renderer/Private/PostProcess/PostProcessLocalExposure.cpp
@@ -216,7 +216,7 @@ FLocalExposureParameters GetLocalExposureParameters(const FViewInfo& View, FIntP
 	if (AutoExposureMethod == EAutoExposureMethod::AEM_Manual)
 	{
 		// when using manual exposure cancel exposure compensation setting and curve from middle grey used by local exposure.
-		LocalExposureMiddleGreyExposureCompensation /= (EyeAdaptationParameters.ExposureCompensationSettings * EyeAdaptationParameters.ExposureCompensationCurve);
+		LocalExposureMiddleGreyExposureCompensation /= (EyeAdaptationParameters.Scalars.ExposureCompensationSettings * EyeAdaptationParameters.Scalars.ExposureCompensationCurve);
 	}
 
 	const FVector2f LocalExposureBilateralGridUVScale = GetLocalExposureBilateralGridUVScale(ViewRectSize);
diff --git a/Engine/Source/Runtime/Renderer/Private/PostProcess/PostProcessMobile.cpp b/Engine/Source/Runtime/Renderer/Private/PostProcess/PostProcessMobile.cpp
--- a/Engine/Source/Runtime/Renderer/Private/PostProcess/PostProcessMobile.cpp
+++ b/Engine/Source/Runtime/Renderer/Private/PostProcess/PostProcessMobile.cpp
@@ -1079,7 +1079,7 @@ public:
 		SHADER_PARAMETER_SAMPLER(SamplerState, BloomSetup_BloomSampler)
 		SHADER_PARAMETER_RDG_TEXTURE(Texture2D, BloomUpTexture)
 		SHADER_PARAMETER_SAMPLER(SamplerState, BloomUpSampler)
-		SHADER_PARAMETER_TEXTURE(Texture2D, BloomDirtMaskTexture)
+		SHADER_PARAMETER_RDG_TEXTURE(Texture2D, BloomDirtMaskTexture)
 		SHADER_PARAMETER_SAMPLER(SamplerState, BloomDirtMaskSampler)
 		RENDER_TARGET_BINDING_SLOTS()
 	END_SHADER_PARAMETER_STRUCT()
@@ -1170,11 +1170,11 @@ FScreenPassTexture AddMobileSunMergePass(FRDGBuilder& GraphBuilder, const FViewI
 	const float BloomGaussianIntensity = Settings.BloomIntensity * Settings.BloomGaussianIntensity;
 	FLinearColor BloomColor = Settings.Bloom1Tint * BloomGaussianIntensity * 0.5;
 
-	FRHITexture* BloomDirtMaskTexture = GBlackTexture->TextureRHI;
+	FRDGTextureRef BloomDirtMaskTexture = GSystemTextures.GetBlackDummy(GraphBuilder);
 
-	if (Settings.BloomDirtMask && Settings.BloomDirtMask->GetResource())
+	if (Settings.BloomDirtMask && Settings.BloomDirtMask->GetResource() && Settings.BloomDirtMask->GetResource()->TextureRHI)
 	{
-		BloomDirtMaskTexture = Settings.BloomDirtMask->GetResource()->TextureRHI;
+		BloomDirtMaskTexture = RegisterExternalTexture(GraphBuilder, Settings.BloomDirtMask->GetResource()->TextureRHI, TEXT("BloomDirtMask"));
 	}
 
 	FMobileSunMergePS::FParameters* PSShaderParameters = GraphBuilder.AllocParameters<FMobileSunMergePS::FParameters>();
diff --git a/Engine/Source/Runtime/Renderer/Private/PostProcess/PostProcessTonemap.cpp b/Engine/Source/Runtime/Renderer/Private/PostProcess/PostProcessTonemap.cpp
--- a/Engine/Source/Runtime/Renderer/Private/PostProcess/PostProcessTonemap.cpp
+++ b/Engine/Source/Runtime/Renderer/Private/PostProcess/PostProcessTonemap.cpp
@@ -341,7 +341,7 @@ BEGIN_SHADER_PARAMETER_STRUCT(FTonemapParameters, )
 
 	SHADER_PARAMETER_RDG_BUFFER_SRV(StructuredBuffer<float4>, EyeAdaptationBuffer)
 	SHADER_PARAMETER_RDG_TEXTURE(Texture3D, ColorGradingLUT)
-	SHADER_PARAMETER_TEXTURE(Texture2D, BloomDirtMaskTexture)
+	SHADER_PARAMETER_RDG_TEXTURE(Texture2D, BloomDirtMaskTexture)
 	SHADER_PARAMETER_SAMPLER(SamplerState, ColorSampler)
 	SHADER_PARAMETER_SAMPLER(SamplerState, ColorGradingLUTSampler)
 	SHADER_PARAMETER_SAMPLER(SamplerState, BloomDirtMaskSampler)
@@ -363,7 +363,7 @@ BEGIN_SHADER_PARAMETER_STRUCT(FTonemapParameters, )
 	SHADER_PARAMETER(FVector4f, VignetteParams1)
 	SHADER_PARAMETER(float, 	VignetteParams2)
 	SHADER_PARAMETER(uint32,	VignetteType)
-	SHADER_PARAMETER_TEXTURE(Texture2D, VignetteTexture)
+	SHADER_PARAMETER_RDG_TEXTURE(Texture2D, VignetteTexture)
 	SHADER_PARAMETER_SAMPLER(SamplerState, VignetteTextureSampler)
 END_SHADER_PARAMETER_STRUCT()
 
@@ -787,8 +787,8 @@ FScreenPassTexture AddTonemapPass(FRDGBuilder& GraphBuilder, const FViewInfo& Vi
 		CommonParameters.FilmGrain.FilmGrainSampler = TStaticSamplerState<SF_Bilinear, AM_Wrap, AM_Wrap>::GetRHI();
 
 		// Grain should be temporally stable when accumulating multiple samples per output frame, so we use OutputFrameIndex instead of FrameIndex.
-		int32 RandomSequenceLength = CVarFilmGrainSequenceLength.GetValueOnRenderThread();
-		int32 RandomSequenceIndex = (View.ViewState ? View.ViewState->GetOutputFrameIndex() : 0) % RandomSequenceLength;
+		const int32 RandomSequenceLength = FMath::Max(1, CVarFilmGrainSequenceLength.GetValueOnRenderThread());
+		const int32 RandomSequenceIndex = (View.ViewState ? View.ViewState->GetOutputFrameIndex() : 0) % RandomSequenceLength;
 
 		FVector2f RandomGrainTextureUVOffset;
 		RandomGrainTextureUVOffset.X = Halton(RandomSequenceIndex + 1, 2);
@@ -887,7 +887,7 @@ FScreenPassTexture AddTonemapPass(FRDGBuilder& GraphBuilder, const FViewInfo& Vi
 		1.0f / FMath::Max(0.001, PostProcessSettings.VignetteSize.X), 
 		1.0f / FMath::Max(0.001, PostProcessSettings.VignetteSize.Y));
 	CommonParameters.VignetteParams2 = PostProcessSettings.VignetteSoftness;
-	CommonParameters.VignetteTexture = GWhiteTexture->TextureRHI;
+	CommonParameters.VignetteTexture = GSystemTextures.GetWhiteDummy(GraphBuilder);
 	CommonParameters.VignetteTextureSampler = TStaticSamplerState<SF_Bilinear>::GetRHI();
 	if (PostProcessSettings.VignetteType == EVignetteType::Texture && PostProcessSettings.VignetteTexture)
 	{
@@ -895,7 +895,8 @@ FScreenPassTexture AddTonemapPass(FRDGBuilder& GraphBuilder, const FViewInfo& Vi
 
 		if (VignetteTextureResource && VignetteTextureResource->TextureRHI)
 		{
-			CommonParameters.VignetteTexture = VignetteTextureResource->TextureRHI;
+			FRDGTextureRef VignetteTextureRDG = RegisterExternalTexture(GraphBuilder, VignetteTextureResource->TextureRHI, TEXT("VignetteTexture"));
+			CommonParameters.VignetteTexture = VignetteTextureRDG;
 		}
 	}
 
@@ -974,16 +975,24 @@ FScreenPassTexture AddTonemapPass(FRDGBuilder& GraphBuilder, const FViewInfo& Vi
 		}
 
 		// TODO: PostProcessSettings.BloomDirtMask->GetResource() is not thread safe
-		if (bUseBloom && PostProcessSettings.BloomDirtMask && PostProcessSettings.BloomDirtMask->GetResource() && PostProcessSettings.BloomDirtMaskIntensity > 0 && PostProcessSettings.BloomIntensity > 0.0)
+		const bool bUseBloomDirtMask = 
+			bUseBloom 
+			&& PostProcessSettings.BloomDirtMask 
+			&& PostProcessSettings.BloomDirtMask->GetResource() 
+			&& PostProcessSettings.BloomDirtMask->GetResource()->TextureRHI
+			&& PostProcessSettings.BloomDirtMaskIntensity > 0 
+			&& PostProcessSettings.BloomIntensity > 0.0;
+		if (bUseBloomDirtMask)
 		{
 			CommonParameters.BloomDirtMaskTint = PostProcessSettings.BloomDirtMaskTint * PostProcessSettings.BloomDirtMaskIntensity / PostProcessSettings.BloomIntensity;
-			CommonParameters.BloomDirtMaskTexture = PostProcessSettings.BloomDirtMask->GetResource()->TextureRHI;
+			FRDGTextureRef BloomDirtMaskRDG = RegisterExternalTexture(GraphBuilder, PostProcessSettings.BloomDirtMask->GetResource()->TextureRHI, TEXT("BloomDirtMask"));
+			CommonParameters.BloomDirtMaskTexture = BloomDirtMaskRDG;
 			CommonParameters.BloomDirtMaskSampler = BilinearClampSampler;
 		}
 		else
 		{
 			CommonParameters.BloomDirtMaskTint = FLinearColor::Black;
-			CommonParameters.BloomDirtMaskTexture = GBlackTexture->TextureRHI;
+			CommonParameters.BloomDirtMaskTexture = GSystemTextures.GetBlackDummy(GraphBuilder);
 			CommonParameters.BloomDirtMaskSampler = BilinearClampSampler;
 		}
 	}
diff --git a/Engine/Source/Runtime/Renderer/Private/PostProcess/PostProcessing.cpp b/Engine/Source/Runtime/Renderer/Private/PostProcess/PostProcessing.cpp
--- a/Engine/Source/Runtime/Renderer/Private/PostProcess/PostProcessing.cpp
+++ b/Engine/Source/Runtime/Renderer/Private/PostProcess/PostProcessing.cpp
@@ -409,7 +409,7 @@ void AddPostProcessingPasses(
 	const FScreenPassTexture OriginalSceneColor = SceneColor;
 
 	// Default the new eye adaptation to the last one in case it's not generated this frame.
-	const FEyeAdaptationParameters EyeAdaptationParameters = GetEyeAdaptationParameters(View);
+	const FEyeAdaptationParameters EyeAdaptationParameters = GetEyeAdaptationParameters(GraphBuilder, View);
 	FRDGBufferRef LastEyeAdaptationBuffer = GetEyeAdaptationBuffer(GraphBuilder, View);
 	FRDGBufferRef EyeAdaptationBuffer = LastEyeAdaptationBuffer;
 
@@ -2280,7 +2280,7 @@ void AddDebugViewPostProcessingPasses(
 	PassSequence.SetEnabled(EPass::SecondaryUpscale, View.RequiresSecondaryUpscale() || View.Family->GetSecondarySpatialUpscalerInterface() != nullptr);
 	PassSequence.Finalize();
 
-	const FEyeAdaptationParameters EyeAdaptationParameters = GetEyeAdaptationParameters(View);
+	const FEyeAdaptationParameters EyeAdaptationParameters = GetEyeAdaptationParameters(GraphBuilder, View);
 
 	if (bTonemapBefore)
 	{
@@ -2612,7 +2612,7 @@ void AddMobilePostProcessingPasses(FRDGBuilder& GraphBuilder, FScene* Scene, con
 	FScreenPassTexture SceneDepthAux((*Inputs.SceneTextures)->SceneDepthAuxTexture, FinalOutputViewRect);
 
 	// Default the new eye adaptation to the last one in case it's not generated this frame.
-	const FEyeAdaptationParameters EyeAdaptationParameters = GetEyeAdaptationParameters(View);
+	const FEyeAdaptationParameters EyeAdaptationParameters = GetEyeAdaptationParameters(GraphBuilder, View);
 	FRDGBufferRef LastEyeAdaptationBuffer = GetEyeAdaptationBuffer(GraphBuilder, View);
 
 	FVisualizeBufferContext VisualizeBufferContext{};
