Change default r.Forward.LightBuffer.Mode to 1 (visible local + directional lights).
- overhead should be negligible since ForwardLightBuffer is still dense and only includes 1 or 2 extra lights on typical usage.

#Synced-CL 51462794

diff --git a/Engine/Source/Runtime/Renderer/Private/LightGridInjection.cpp b/Engine/Source/Runtime/Renderer/Private/LightGridInjection.cpp
--- a/Engine/Source/Runtime/Renderer/Private/LightGridInjection.cpp
+++ b/Engine/Source/Runtime/Renderer/Private/LightGridInjection.cpp
@@ -55,7 +55,7 @@ enum class ELightBufferMode
 	VisibleLightsStableIndices	= 2
 };
 
-static int32 GLightBufferMode = (int32)ELightBufferMode::VisibleLocalLights;
+static int32 GLightBufferMode = (int32)ELightBufferMode::VisibleLights;
 static FAutoConsoleVariableRef CVarLightBufferMode(
 	TEXT("r.Forward.LightBuffer.Mode"),
 	GLightBufferMode,
