Substrate - Toon - make the Max GGX_D LUT mapping exact for roughness [0, 1].
Also removes energy normalisation from toon as this is not needed and will make it cheaper.
Made phong path MaxSpecularExponent be 2048.

#tests PC editor with toon shading material
#rnx

#virtualized
#Synced-CL 51530297

#Synced-CL 51530297

diff --git a/Engine/Shaders/Private/ShadingCommon.ush b/Engine/Shaders/Private/ShadingCommon.ush
--- a/Engine/Shaders/Private/ShadingCommon.ush
+++ b/Engine/Shaders/Private/ShadingCommon.ush
@@ -45,9 +45,12 @@
 #define SSS_PROFILE_ID_INVALID  255
 #define SSS_PROFILE_ID_PERPIXEL 512
 
-// This is selected so that the maximum of the LUT is below 1.0. the maximum was measure as 2.94922.
-#define GGX_MAX_LUT_LOAD_SCALE	(2.95)
-#define GGX_MAX_LUT_STORE_SCALE	(1.0 / GGX_MAX_LUT_LOAD_SCALE)
+// This is selected so that the maximum of the LUT is <= 1.0. The maximum can be PI for D_GGX.
+#define GGX_MAX_LUT_LOAD_SCALE		(PI)
+#define GGX_MAX_LUT_STORE_SCALE		(1.0 / GGX_MAX_LUT_LOAD_SCALE)
+#define GGX_MAX_LUT_RESOLUTION		32
+#define GGX_MAX_LUT_LOAD_UV_SCALE	((float(GGX_MAX_LUT_RESOLUTION) - 1.0f) / float(GGX_MAX_LUT_RESOLUTION))
+#define GGX_MAX_LUT_LOAD_UV_BIAS	(0.5f / float(GGX_MAX_LUT_RESOLUTION))
 
 // for debugging and to visualize
 float3 GetShadingModelColor(uint ShadingModelID)
diff --git a/Engine/Shaders/Private/ShadingEnergyConservation.ush b/Engine/Shaders/Private/ShadingEnergyConservation.ush
--- a/Engine/Shaders/Private/ShadingEnergyConservation.ush
+++ b/Engine/Shaders/Private/ShadingEnergyConservation.ush
@@ -99,7 +99,7 @@ float GGXEnergyLookup(float Roughness, float NoV, float Eta)
 float GGXMaxEnergyLookup(float Roughness, float NoV)
 {
 	// Always evaluate since it is only used for toon shading as of today without any more options.
-	return GGX_MAX_LUT_LOAD_SCALE * View.ShadingEnergyGGXMaxSpecTexture.SampleLevel(SharedShadingEnergySampler, float2(NoV, Roughness), 0).r;
+	return GGX_MAX_LUT_LOAD_SCALE * View.ShadingEnergyGGXMaxSpecTexture.SampleLevel(SharedShadingEnergySampler, float2(NoV, Roughness * GGX_MAX_LUT_LOAD_UV_SCALE + GGX_MAX_LUT_LOAD_UV_BIAS), 0).r;
 }
 
 float2 ClothEnergyLookup(float Roughness, float NoV)
diff --git a/Engine/Shaders/Private/ShadingEnergyConservationTable.usf b/Engine/Shaders/Private/ShadingEnergyConservationTable.usf
--- a/Engine/Shaders/Private/ShadingEnergyConservationTable.usf
+++ b/Engine/Shaders/Private/ShadingEnergyConservationTable.usf
@@ -35,7 +35,15 @@ void BuildEnergyTableCS(uint3 DispatchThreadId : SV_DispatchThreadID)
 
 	// Map from thread ID to the exact roughness/viewing angle we want to compute
 	const float NoV = max(Index.x, 1e-8);
-	const float Roughness = Index.y;
+
+#if BUILD_ENERGY_TABLE == 4
+	// In this case we want exact uvs [0.5/height, (height-0.5)/height] to map to roughness [0, 1] to catch the extrems.
+	// This approach could potentially be applied to all the LUTs.
+	const float Roughness = float(DispatchThreadId.y) / float(EnergyTableResolution - 1);
+#else
+	const float Roughness =  Index.y;
+#endif
+
 #if BUILD_ENERGY_TABLE == 1
 	const float EtaIn = 1.0 + 2.0 * Index.z; // only handle Ior in [1,3] range
 	const float EtaOut = rcp(EtaIn);
@@ -156,6 +164,6 @@ void BuildEnergyTableCS(uint3 DispatchThreadId : SV_DispatchThreadID)
 #endif
 
 #if BUILD_ENERGY_TABLE == 4
-	Output3Texture2D[DispatchThreadId.xy] = GGX_MAX_LUT_STORE_SCALE * (1.0f / Result.x);
+	Output3Texture2D[DispatchThreadId.xy] = saturate(GGX_MAX_LUT_STORE_SCALE * (1.0f / Result.x));
 #endif
 }
diff --git a/Engine/Shaders/Private/Substrate/SubstrateEvaluation.ush b/Engine/Shaders/Private/Substrate/SubstrateEvaluation.ush
--- a/Engine/Shaders/Private/Substrate/SubstrateEvaluation.ush
+++ b/Engine/Shaders/Private/Substrate/SubstrateEvaluation.ush
@@ -1373,18 +1373,7 @@ FSubstrateEvaluateResult SubstrateEvaluateBSDFCommon(FSubstrateBSDFContext BSDFC
 				else
 				{
 
-					float D = 0.0f;
-					// Generalized microfacet specular
-					if (IntegrationType == INTEGRATION_PUNCTUAL_LIGHT)
-					{
-						D = Substrate_D_GGX(SafeRoughness, Alpha2Spec, BSDFContext.Context.NoH);
-					}
-					else
-					{
-						const float Energy = EnergyNormalization(Alpha2Spec, BSDFContext.Context.VoH, AreaLight);
-						D = Substrate_D_GGX(SafeRoughness, Alpha2Spec, BSDFContext.Context.NoH) * Energy;
-					}
-
+					float D = Substrate_D_GGX(SafeRoughness, Alpha2Spec, BSDFContext.Context.NoH);
 					//const float3 FresnelTerm = Substrate_F_GGX(F0, F90, BSDFContext.Context.VoH);
 					//float Vis = Substrate_Vis_GGX(SafeRoughness, Alpha2Spec, BSDFContext.Context.NoV, AreaLight.NoL);
 
@@ -1399,7 +1388,7 @@ FSubstrateEvaluateResult SubstrateEvaluateBSDFCommon(FSubstrateBSDFContext BSDFC
 				const float RoughnessExponent = saturate(1.0 - Roughness);
 				const float RoughnessExponent2 = RoughnessExponent * RoughnessExponent; // Arbitrary: this is so that the roughness control appears smooth and linear without big jumps due to 8bit encoding.
 
-				const float MaxSpecularExponent = 64.0f;
+				const float MaxSpecularExponent = 2048.0f;
 				const float SpecularExponentScale = MaxSpecularExponent - 1.0f;
 			#if 0
 				Specular = pow(saturate(dot(AreaLight.SpecularL, BSDFContext.R)), max(1.0f, 1.0 + RoughnessExponent2 * SpecularExponentScale));	// NoR based is also interesting visually.
diff --git a/Engine/Source/Runtime/Renderer/Private/ShadingEnergyConservation.cpp b/Engine/Source/Runtime/Renderer/Private/ShadingEnergyConservation.cpp
--- a/Engine/Source/Runtime/Renderer/Private/ShadingEnergyConservation.cpp
+++ b/Engine/Source/Runtime/Renderer/Private/ShadingEnergyConservation.cpp
@@ -418,6 +418,7 @@ void Init(FRDGBuilder& GraphBuilder, const FViewInfo& View)
 		// Change this to true in order to regenerate the energy tables, and manually copy the coefficients into ShadingEnergyConservationData.h
 		const bool bRuntimeGeneration = CVarShadingEnergyConservation_RuntimeGeneration.GetValueOnRenderThread() > 0;
 		const int Size2D = FMath::Clamp(CVarShadingEnergyConservation_TableResolution.GetValueOnRenderThread(), 16, 512);
+		const int Size2DMaxLUT = 32; // Must match GGX_MAX_LUT_RESOLUTION
 
 		EPixelFormat Format = PF_R8G8;
 		// for low roughness we would get banding with PF_R8G8 but for low spec it could be used, for now we don't do this optimization
@@ -532,14 +533,14 @@ void Init(FRDGBuilder& GraphBuilder, const FViewInfo& View)
 					TShaderMapRef<FBuildShadingEnergyConservationTableCS> ComputeShader(View.ShaderMap, PermutationVector);
 					FBuildShadingEnergyConservationTableCS::FParameters* PassParameters = GraphBuilder.AllocParameters<FBuildShadingEnergyConservationTableCS::FParameters>();
 					PassParameters->NumSamples = NumSamples;
-					PassParameters->EnergyTableResolution = Size2D;
+					PassParameters->EnergyTableResolution = Size2DMaxLUT;
 					PassParameters->Output3Texture2D = GraphBuilder.CreateUAV(FRDGTextureUAVDesc(GGXMaxSpecEnergyTexture, 0));
 					FComputeShaderUtils::AddPass(
 						GraphBuilder,
 						RDG_EVENT_NAME("ShadingEnergyConservation::BuildTable(GGXMaxSpec)"),
 						ComputeShader,
 						PassParameters,
-						FComputeShaderUtils::GetGroupCount(FIntPoint(Size2D, Size2D), FComputeShaderUtils::kGolden2DGroupSize));
+						FComputeShaderUtils::GetGroupCount(FIntPoint(Size2DMaxLUT, Size2DMaxLUT), FComputeShaderUtils::kGolden2DGroupSize));
 				}
 
 				Out.GGXSpecEnergyTexture  = GraphBuilder.ConvertToExternalTexture(GGXSpecEnergyTexture);
