[Lit Translucency] Small changes to TLV update.
- new cvar to support disabling typed UAV load in TLV injection logic (only intended to help debug an issue with flickering lighting on particles).
- renamed r.EnableAsyncComputeTranslucencyLightingVolumeClear to r.TranslucencyLightingVolume.AsyncComputeClear.
- also remove CLEAR_COMPUTE_SHADER define from ClearTranslucentLightingVolumeCS since we can achieve the same result using the built-in macro.

#Synced-CL 51545719

diff --git a/Engine/Shaders/Private/TranslucentLightInjectionShaders.usf b/Engine/Shaders/Private/TranslucentLightInjectionShaders.usf
--- a/Engine/Shaders/Private/TranslucentLightInjectionShaders.usf
+++ b/Engine/Shaders/Private/TranslucentLightInjectionShaders.usf
@@ -431,7 +431,8 @@ void InjectBatchMainCS(
 
 #endif // InjectBatchMainCS
 
-#if CLEAR_COMPUTE_SHADER
+#ifdef ClearTranslucentLightingVolumeCS
+
 RWTexture3D<float4> RWAmbient0;
 RWTexture3D<float4> RWDirectional0;
 RWTexture3D<float4> RWAmbient1;
@@ -450,4 +451,4 @@ void ClearTranslucentLightingVolumeCS(
 	RWAmbient1[Index.xyz] = float4(0.0, 0.0, 0.0, 0.0);
 	RWDirectional1[Index.xyz] = float4(0.0, 0.0, 0.0, 0.0);
 }
-#endif
\ No newline at end of file
+#endif // ClearTranslucentLightingVolumeCS
\ No newline at end of file
diff --git a/Engine/Source/Runtime/Renderer/Private/DeferredShadingRenderer.cpp b/Engine/Source/Runtime/Renderer/Private/DeferredShadingRenderer.cpp
--- a/Engine/Source/Runtime/Renderer/Private/DeferredShadingRenderer.cpp
+++ b/Engine/Source/Runtime/Renderer/Private/DeferredShadingRenderer.cpp
@@ -183,10 +183,9 @@ static TAutoConsoleVariable<int32> CVarIrisNormal(
 	TEXT(" 1: on"),
 	ECVF_ReadOnly);
 
-int32 GbEnableAsyncComputeTranslucencyLightingVolumeClear = 0; // @todo: disabled due to GPU crashes
-static FAutoConsoleVariableRef CVarEnableAsyncComputeTranslucencyLightingVolumeClear(
-	TEXT("r.EnableAsyncComputeTranslucencyLightingVolumeClear"),
-	GbEnableAsyncComputeTranslucencyLightingVolumeClear,
+static TAutoConsoleVariable<bool> CVarTranslucencyLightingVolumeAsyncComputeClear(
+	TEXT("r.TranslucencyLightingVolume.AsyncComputeClear"),
+	false, // @todo: disabled due to GPU crashes
 	TEXT("Whether to clear the translucency lighting volume using async compute.\n"),
 	ECVF_RenderThreadSafe | ECVF_Scalability
 );
@@ -382,6 +381,11 @@ DECLARE_DWORD_COUNTER_STAT(TEXT("BasePass Visible Shading Bins"), STAT_NaniteBas
 
 CSV_DEFINE_CATEGORY(LightCount, true);
 
+static bool ShouldClearTranslucencyLightingVolumeInAsyncCompute()
+{
+	return CVarTranslucencyLightingVolumeAsyncComputeClear.GetValueOnRenderThread() && GSupportsEfficientAsyncCompute;
+}
+
 /*-----------------------------------------------------------------------------
 	Global Illumination Plugin Function Delegates
 -----------------------------------------------------------------------------*/
@@ -2869,7 +2873,7 @@ void FDeferredShadingSceneRenderer::Render(FRDGBuilder& GraphBuilder, const FSce
 
 		FTranslucencyLightingVolumeTextures TranslucencyLightingVolumeTextures;
 
-		if (bRenderDeferredLighting && GbEnableAsyncComputeTranslucencyLightingVolumeClear && GSupportsEfficientAsyncCompute)
+		if (bRenderDeferredLighting && ShouldClearTranslucencyLightingVolumeInAsyncCompute())
 		{
 			TranslucencyLightingVolumeTextures.Init(GraphBuilder, Views, ERDGPassFlags::AsyncCompute);
 		}
@@ -3341,7 +3345,7 @@ void FDeferredShadingSceneRenderer::Render(FRDGBuilder& GraphBuilder, const FSce
 			RenderDFAOAsIndirectShadowing(GraphBuilder, SceneTextures, DynamicBentNormalAOTextures);
 
 			// Clear the translucent lighting volumes before we accumulate
-			if ((GbEnableAsyncComputeTranslucencyLightingVolumeClear && GSupportsEfficientAsyncCompute) == false)
+			if (!ShouldClearTranslucencyLightingVolumeInAsyncCompute())
 			{
 				TranslucencyLightingVolumeTextures.Init(GraphBuilder, Views, ERDGPassFlags::Compute);
 			}
diff --git a/Engine/Source/Runtime/Renderer/Private/LightRendering.cpp b/Engine/Source/Runtime/Renderer/Private/LightRendering.cpp
--- a/Engine/Source/Runtime/Renderer/Private/LightRendering.cpp
+++ b/Engine/Source/Runtime/Renderer/Private/LightRendering.cpp
@@ -1336,8 +1336,6 @@ void FSceneRenderer::GetLightNameForDrawEvent(const FLightSceneProxy* LightProxy
 #endif
 }
 
-extern int32 GbEnableAsyncComputeTranslucencyLightingVolumeClear;
-
 uint32 GetShadowQuality();
 
 static bool LightRequiresDenoiser(const FLightSceneInfo& LightSceneInfo)
diff --git a/Engine/Source/Runtime/Renderer/Private/TranslucentLighting.cpp b/Engine/Source/Runtime/Renderer/Private/TranslucentLighting.cpp
--- a/Engine/Source/Runtime/Renderer/Private/TranslucentLighting.cpp
+++ b/Engine/Source/Runtime/Renderer/Private/TranslucentLighting.cpp
@@ -162,9 +162,16 @@ static TAutoConsoleVariable<bool> CVarTranslucencyLightingVolumeMarkVoxels(
 static TAutoConsoleVariable<int32> CVarTranslucencyLightingVolumeBatch(
 	TEXT("r.TranslucencyLightingVolume.Batch"),
 	1,
-	TEXT("When enabled, batches supported lights into a single draw call for efficiency"),
+	TEXT("When enabled, batches supported lights into a single draw call for efficiency."),
 	ECVF_RenderThreadSafe | ECVF_Scalability);
 
+static TAutoConsoleVariable<bool> CVarTranslucencyLightingVolumeUseTypedUAVLoad(
+	TEXT("r.TranslucencyLightingVolume.UseTypedUAVLoad"),
+	true,
+	TEXT("Whether to use Typed UAV loads when updating the translucency lighting volume.\n")
+	TEXT("Recommended to be enabled since it reduces transient memory usage."),
+	ECVF_RenderThreadSafe);
+
 static TAutoConsoleVariable<int32> CVarTranslucencyLightingVolumeAccurateRectLights(
 	TEXT("r.TranslucencyLightingVolume.AccurateRectLights"),
 	1,
@@ -1206,7 +1213,6 @@ public:
 	static void ModifyCompilationEnvironment(const FGlobalShaderPermutationParameters& Parameters, FShaderCompilerEnvironment& OutEnvironment)
 	{
 		FGlobalShader::ModifyCompilationEnvironment(Parameters, OutEnvironment);
-		OutEnvironment.SetDefine(TEXT("CLEAR_COMPUTE_SHADER"), 1);
 		OutEnvironment.SetDefine(TEXT("CLEAR_BLOCK_SIZE"), CLEAR_BLOCK_SIZE);
 	}
 };
@@ -1724,7 +1730,7 @@ static void InjectTranslucencyLightingVolumeBatch(
 	FTranslucencyLightingVolumeTextures::GetTextureFormatAndCreationFlags(TranslucencyPixelFormat, TranslucencyTargetFlags);
 
 	const bool bUseLightFunctionAtlas = LightFunctionAtlas::IsEnabled(View, ELightFunctionAtlasSystem::DeferredLighting);
-	const bool bUseUAVTypedLoad = UE::PixelFormat::HasCapabilities(TranslucencyPixelFormat, EPixelFormatCapabilities::TypedUAVLoad);
+	const bool bUseUAVTypedLoad = CVarTranslucencyLightingVolumeUseTypedUAVLoad.GetValueOnRenderThread() && UE::PixelFormat::HasCapabilities(TranslucencyPixelFormat, EPixelFormatCapabilities::TypedUAVLoad);
 
 	auto* PassParameters = GraphBuilder.AllocParameters<FTranslucentLightingInjectBatchCS::FParameters>();
 	if (bSupportVirtualShadowMaps)
