Remove unused code

#rnx

#Synced-CL 51488618

diff --git a/Engine/Source/Runtime/Renderer/Private/VirtualShadowMaps/VirtualShadowMapCacheManager.cpp b/Engine/Source/Runtime/Renderer/Private/VirtualShadowMaps/VirtualShadowMapCacheManager.cpp
--- a/Engine/Source/Runtime/Renderer/Private/VirtualShadowMaps/VirtualShadowMapCacheManager.cpp
+++ b/Engine/Source/Runtime/Renderer/Private/VirtualShadowMaps/VirtualShadowMapCacheManager.cpp
@@ -584,75 +584,6 @@ void FVirtualShadowMapArrayCacheManager::FInvalidatingPrimitiveCollector::AddInv
 	}
 }
 
-
-
-FVirtualShadowMapFeedback::FVirtualShadowMapFeedback()
-{
-	for (int32 i = 0; i < MaxBuffers; ++i)
-	{
-		Buffers[i].Buffer = new FRHIGPUBufferReadback(TEXT("Shadow.Virtual.Readback"));
-		Buffers[i].Size = 0;
-	}
-}
-
-FVirtualShadowMapFeedback::~FVirtualShadowMapFeedback()
-{
-	for (int32 i = 0; i < MaxBuffers; ++i)
-	{
-		delete Buffers[i].Buffer;
-		Buffers[i].Buffer = nullptr;
-		Buffers[i].Size = 0;
-	}
-}
-
-void FVirtualShadowMapFeedback::SubmitFeedbackBuffer(
-	FRDGBuilder& GraphBuilder,
-	FRDGBufferRef FeedbackBuffer)
-{
-	// Source copy usage is required for readback
-	check((FeedbackBuffer->Desc.Usage & EBufferUsageFlags::SourceCopy) == EBufferUsageFlags::SourceCopy);
-
-	if (NumPending == MaxBuffers)
-	{
-		return;
-	}
-
-	FRHIGPUBufferReadback* ReadbackBuffer = Buffers[WriteIndex].Buffer;
-	Buffers[WriteIndex].Size = FeedbackBuffer->Desc.GetSize();
-
-	AddReadbackBufferPass(GraphBuilder, RDG_EVENT_NAME("Readback"), FeedbackBuffer,
-		[ReadbackBuffer, FeedbackBuffer](FRDGAsyncTask, FRHICommandList& RHICmdList)
-		{
-			ReadbackBuffer->EnqueueCopy(RHICmdList, FeedbackBuffer->GetRHI(), 0u);
-		});
-
-	WriteIndex = (WriteIndex + 1) % MaxBuffers;
-	NumPending = FMath::Min(NumPending + 1, MaxBuffers);
-}
-
-FVirtualShadowMapFeedback::FReadbackInfo FVirtualShadowMapFeedback::GetLatestReadbackBuffer()
-{
-	int32 LatestBufferIndex = -1;
-
-	// Find latest buffer that is ready
-	while (NumPending > 0)
-	{
-		uint32 Index = (WriteIndex + MaxBuffers - NumPending) % MaxBuffers;
-		if (Buffers[Index].Buffer->IsReady())
-		{
-			--NumPending;
-			LatestBufferIndex = Index;
-		}
-		else
-		{
-			break;
-		}
-	}
-
-	return LatestBufferIndex >= 0 ? Buffers[LatestBufferIndex] : FReadbackInfo();
-}
-
-
 IMPLEMENT_SCENE_EXTENSION(FVirtualShadowMapArrayCacheManager);
 
 bool FVirtualShadowMapArrayCacheManager::ShouldCreateExtension(FScene& Scene)
diff --git a/Engine/Source/Runtime/Renderer/Private/VirtualShadowMaps/VirtualShadowMapCacheManager.h b/Engine/Source/Runtime/Renderer/Private/VirtualShadowMaps/VirtualShadowMapCacheManager.h
--- a/Engine/Source/Runtime/Renderer/Private/VirtualShadowMaps/VirtualShadowMapCacheManager.h
+++ b/Engine/Source/Runtime/Renderer/Private/VirtualShadowMaps/VirtualShadowMapCacheManager.h
@@ -207,28 +207,6 @@ private:
 	float LightRadius = -1.0f;		// Negative means infinite
 };
 
-class FVirtualShadowMapFeedback
-{
-public:
-	FVirtualShadowMapFeedback();
-	~FVirtualShadowMapFeedback();
-
-	struct FReadbackInfo
-	{
-		FRHIGPUBufferReadback* Buffer = nullptr;
-		uint32 Size = 0;
-	};
-
-	void SubmitFeedbackBuffer(FRDGBuilder& GraphBuilder, FRDGBufferRef FeedbackBuffer);
-	FReadbackInfo GetLatestReadbackBuffer();
-
-private:
-	static const int32 MaxBuffers = 3;
-	int32 WriteIndex = 0;
-	int32 NumPending = 0;
-	FReadbackInfo Buffers[MaxBuffers];
-};
-
 // Persistent buffers that we ping pong frame by frame
 struct FVirtualShadowMapArrayFrameData
 {
