Substrate - toon - fixed diffuse not showing up when unified diffuse is enabled.

#tests Pc editor with toon material substrate enabled and r.Substrate.Experimental.ToonUnifiedDiffuse=1
#rnx

#Synced-CL 51544534

diff --git a/Engine/Source/Runtime/Renderer/Private/LightRendering.cpp b/Engine/Source/Runtime/Renderer/Private/LightRendering.cpp
--- a/Engine/Source/Runtime/Renderer/Private/LightRendering.cpp
+++ b/Engine/Source/Runtime/Renderer/Private/LightRendering.cpp
@@ -2747,7 +2747,7 @@ static FDeferredLightPS::FParameters GetDeferredLightPSParameters(
 
 	// PS - Render Targets
 	Out.RenderTargets[0] = FRenderTargetBinding(SceneColorTexture, ERenderTargetLoadAction::ELoad);
-	if (SubstrateTileType != ESubstrateTileType::ECount && SubstrateTileType >= ESubstrateTileType::EComplex && IsExperimentalToonUnifiedDiffuseEnabled())
+	if (Substrate::IsSubstrateEnabled() && SubstrateTileType >= ESubstrateTileType::EComplex && IsExperimentalToonUnifiedDiffuseEnabled())
 	{
 		const FSubstrateSceneData& SubstrateSceneData = *View.GetSceneTextures().SubstrateSceneData;
 		Out.RenderTargets[1] = FRenderTargetBinding(SubstrateSceneData.SubstrateSceneToonData, ERenderTargetLoadAction::ELoad);
@@ -2785,7 +2785,7 @@ uint32 SetupLightGraphicsPSOState(
 	EShaderPlatform ShaderPlatform)
 {
 	GraphicsPSOInit.PrimitiveType = PT_TriangleList;
-	if (SubstrateTileMaterialType != ESubstrateTileType::ECount && SubstrateTileMaterialType >= ESubstrateTileType::EComplex && IsExperimentalToonUnifiedDiffuseEnabled())
+	if (Substrate::IsSubstrateEnabled() && SubstrateTileMaterialType >= ESubstrateTileType::EComplex && IsExperimentalToonUnifiedDiffuseEnabled())
 	{
 		GraphicsPSOInit.BlendState = TStaticBlendState<
 			CW_RGBA, BO_Add, BF_One, BF_One, BO_Add, BF_One, BF_One,
