Substrate - Toon - unified diffuse is now tiled when needed and can Complex and Complex special tiles accounting for potential toon pixel.

#rnx
#tests PC editor dx12 toon material with unifed diffuse enabled. Tested engine test preflight compilation/shader cook 69afb76b643db7a11a5ca077

#Synced-CL 51544534

diff --git a/Engine/Shaders/Private/DeferredLightPixelShaders.usf b/Engine/Shaders/Private/DeferredLightPixelShaders.usf
--- a/Engine/Shaders/Private/DeferredLightPixelShaders.usf
+++ b/Engine/Shaders/Private/DeferredLightPixelShaders.usf
@@ -651,6 +651,8 @@ void DeferredLightPixelMain(
 Texture2D<float4> ToonDataTexture;
 
 void DeferredLightApplyToonDiffusePS(
+	in float2 ScreenUV : TEXCOORD0,
+	in float3 ScreenVector : TEXCOORD1,
 	in float4 SVPos : SV_POSITION,
 	out float4 OutColor : SV_Target0
 )
@@ -659,11 +661,11 @@ void DeferredLightApplyToonDiffusePS(
 	OutColor = 0;
 
 
-	FInputParams InputParams = (FInputParams)0;
-	InputParams.PixelPos = SVPos.xy;
-	InputParams.ScreenPosition = SvPositionToScreenPosition(SVPos);
-	InputParams.ScreenUV = InputParams.ScreenPosition.xy / InputParams.ScreenPosition.w * View.ScreenPositionScaleBias.xy + View.ScreenPositionScaleBias.wz;
-	InputParams.ScreenVector = 0;
+	FInputParams InputParams	= (FInputParams)0;
+	InputParams.PixelPos		= SVPos.xy;
+	InputParams.ScreenPosition	= SvPositionToScreenPosition(SVPos);
+	InputParams.ScreenUV		= ScreenUV;
+	InputParams.ScreenVector	= ScreenVector;
 
 #if SUBSTRATE_ENABLED==1
 
diff --git a/Engine/Source/Runtime/Renderer/Private/LightRendering.cpp b/Engine/Source/Runtime/Renderer/Private/LightRendering.cpp
--- a/Engine/Source/Runtime/Renderer/Private/LightRendering.cpp
+++ b/Engine/Source/Runtime/Renderer/Private/LightRendering.cpp
@@ -2504,32 +2504,100 @@ void FDeferredShadingSceneRenderer::RenderLights(
 			{
 				FViewInfo& View = Views[ViewIndex];
 
-				if (Substrate::GetSubstrateUsesTileType(View, ESubstrateTileType::EComplex) || Substrate::GetSubstrateUsesTileType(View, ESubstrateTileType::EComplexSpecial))
+				FDeferredLightApplyToonDiffusePS::FParameters* ToonPassParameters = GraphBuilder.AllocParameters<FDeferredLightApplyToonDiffusePS::FParameters>();
+				// PS - General parameters
+				ToonPassParameters->View = View.ViewUniformBuffer;
+				ToonPassParameters->SceneTextures = SceneTextures.UniformBuffer;
+				ToonPassParameters->Substrate = Substrate::BindSubstrateGlobalUniformParameters(View);
+
+				//	const FSubstrateSceneData& SubstrateSceneData = *View.GetSceneTextures().SubstrateSceneData;
+				ToonPassParameters->ToonDataTexture = SubstrateSceneData.SubstrateSceneToonData;
+
+				ToonPassParameters->RenderTargets[0] = FRenderTargetBinding(SceneTextures.Color.Target, ERenderTargetLoadAction::ELoad);
+				if (SceneTextures.Depth.Target)
 				{
-					FDeferredLightApplyToonDiffusePS::FParameters* ToonPassParameters = GraphBuilder.AllocParameters<FDeferredLightApplyToonDiffusePS::FParameters>();
-					// PS - General parameters
-					ToonPassParameters->View = View.ViewUniformBuffer;
-					ToonPassParameters->SceneTextures = SceneTextures.UniformBuffer;
-					ToonPassParameters->Substrate = Substrate::BindSubstrateGlobalUniformParameters(View);
+					ToonPassParameters->RenderTargets.DepthStencil = FDepthStencilBinding(SceneTextures.Depth.Resolve, ERenderTargetLoadAction::ELoad, ERenderTargetLoadAction::ELoad, FExclusiveDepthStencil::DepthRead_StencilRead);
+				}
+
+				FRHIBlendState* BlendState = TStaticBlendState<CW_RGBA, BO_Add, BF_One, BF_One, BO_Add, BF_One, BF_One>::GetRHI();
+
+				FDeferredLightVS::FPermutationDomain PermutationVectorVS;
+				PermutationVectorVS.Set<FDeferredLightVS::FRadialLight>(false);
+				TShaderMapRef<FDeferredLightVS> VertexShader(View.ShaderMap, PermutationVectorVS);
+
+				Substrate::FSubstrateTilePassVS::FPermutationDomain VSPermutationVector;
+				VSPermutationVector.Set< Substrate::FSubstrateTilePassVS::FEnableDebug >(false);
+				VSPermutationVector.Set< Substrate::FSubstrateTilePassVS::FEnableTexCoordScreenVector >(true);
+				VSPermutationVector.Set< Substrate::FSubstrateTilePassVS::FClassificationTileSize16 >(Substrate::UseTileSize16());
+				TShaderMapRef<Substrate::FSubstrateTilePassVS> TileVertexShader(View.ShaderMap, VSPermutationVector);
 
-					//	const FSubstrateSceneData& SubstrateSceneData = *View.GetSceneTextures().SubstrateSceneData;
-					ToonPassParameters->ToonDataTexture = SubstrateSceneData.SubstrateSceneToonData;
+				for (uint32 TileType = ESubstrateTileType::EComplex; TileType <= ESubstrateTileType::EComplexSpecial; ++TileType)
+				{
+					const ESubstrateTileType SubstrateTileType = ESubstrateTileType(TileType);
 
-					ToonPassParameters->RenderTargets[0] = FRenderTargetBinding(SceneTextures.Color.Target, ERenderTargetLoadAction::ELoad);
+					if (!Substrate::GetSubstrateUsesTileType(View, SubstrateTileType))
+					{
+						continue;
+					}
+					const bool bEnableSubstrateTiledPass = SubstrateTileType != ESubstrateTileType::ECount && Substrate::IsSubstrateEnabled() && Substrate::AreLightingPassesTiled();
 
 					FDeferredLightApplyToonDiffusePS::FPermutationDomain ToonPermutationVector;
-					const ESubstrateTileType TileType = ESubstrateTileType::EComplex;											// TODO render "complex" and "complex special" tiles
-					ToonPermutationVector.Set<Substrate::FSubstrateTileType>(TileType);
-					TShaderMapRef<FDeferredLightApplyToonDiffusePS> PixelShader(View.ShaderMap, ToonPermutationVector);
+					ToonPermutationVector.Set<Substrate::FSubstrateTileType>(SubstrateTileType);
+					TShaderMapRef<FDeferredLightApplyToonDiffusePS> ToonPixelShader(View.ShaderMap, ToonPermutationVector);
+
+					GraphBuilder.AddPass(
+						RDG_EVENT_NAME("Light::ApplyToonDiffuse"),
+						ToonPassParameters,
+						ERDGPassFlags::Raster,
+						[&View, ToonPixelShader, ToonPassParameters, SubstrateTileType, BlendState, VertexShader, TileVertexShader, bEnableSubstrateTiledPass](FRDGAsyncTask, FRHICommandList& RHICmdList)
+					{
+						FGraphicsPipelineStateInitializer GraphicsPSOInit;
+						RHICmdList.ApplyCachedRenderTargets(GraphicsPSOInit);
 
-					FRHIBlendState* BlendState = TStaticBlendState<CW_RGBA, BO_Add, BF_One, BF_One, BO_Add, BF_One, BF_One>::GetRHI();
+						RHICmdList.SetViewport(View.ViewRect.Min.X, View.ViewRect.Min.Y, 0.0f, View.ViewRect.Max.X, View.ViewRect.Max.Y, 1.0f);
+						RHICmdList.SetStreamSource(0, nullptr, 0);
 
-					FPixelShaderUtils::AddFullscreenPass<FDeferredLightApplyToonDiffusePS>(
-						GraphBuilder, View.ShaderMap, RDG_EVENT_NAME("Light::ApplyToonDiffuse"), PixelShader, ToonPassParameters,
-						View.ViewRect, BlendState);
-					// A full screen pass for now. SUBSTRATE_TODO use EComplex tiles.
-				}
-			}
+						uint32 StencilRef = 0;
+						Substrate::GetSubstrateTiledDepthStencilStateRead(SubstrateTileType, CF_Always, GraphicsPSOInit, StencilRef);
+						GraphicsPSOInit.PrimitiveType = PT_TriangleList;
+						GraphicsPSOInit.BlendState = BlendState;
+						GraphicsPSOInit.RasterizerState = TStaticRasterizerState<FM_Solid, CM_None>::GetRHI();
+						GraphicsPSOInit.bDepthBounds = false;
+						GraphicsPSOInit.BoundShaderState.VertexDeclarationRHI = GFilterVertexDeclaration.VertexDeclarationRHI;
+						GraphicsPSOInit.BoundShaderState.VertexShaderRHI = bEnableSubstrateTiledPass ? TileVertexShader.GetVertexShader() : VertexShader.GetVertexShader();
+						GraphicsPSOInit.BoundShaderState.PixelShaderRHI = ToonPixelShader.GetPixelShader();
+
+						SetGraphicsPipelineState(RHICmdList, GraphicsPSOInit, StencilRef);
+
+						SetShaderParameters(RHICmdList, ToonPixelShader, ToonPixelShader.GetPixelShader(), *ToonPassParameters);
+
+						if (bEnableSubstrateTiledPass)
+						{
+							check(Substrate::IsSubstrateEnabled());
+							Substrate::FSubstrateTilePassVS::FParameters VSParameters = Substrate::SetTileParameters(View, SubstrateTileType, GraphicsPSOInit.PrimitiveType);
+							SetShaderParameters(RHICmdList, TileVertexShader, TileVertexShader.GetVertexShader(), VSParameters);
+							// Render tiles for current complexity only, with stencil test.
+							RHICmdList.DrawPrimitiveIndirect(VSParameters.TileIndirectBuffer->GetIndirectRHICallBuffer(), Substrate::TileTypeDrawIndirectArgOffset(SubstrateTileType));
+						}
+						else
+						{
+							FDeferredLightVS::FParameters VSParameters2 = FDeferredLightVS::GetParameters(View);
+							SetShaderParameters(RHICmdList, VertexShader, VertexShader.GetVertexShader(), VSParameters2);
+							// Render as full screen rectangle only leveraging stencil test for better performance.
+							DrawRectangle(
+								RHICmdList,
+								0, 0,
+								View.ViewRect.Width(), View.ViewRect.Height(),
+								View.ViewRect.Min.X, View.ViewRect.Min.Y,
+								View.ViewRect.Width(), View.ViewRect.Height(),
+								View.ViewRect.Size(),
+								View.GetSceneTexturesConfig().Extent,
+								VertexShader,
+								EDRF_UseTriangleOptimization);
+						}
+					});
+				} // for each TileType
+			} // for each View
 		}
 	}
 
@@ -2689,6 +2757,14 @@ static TRDGUniformBufferRef<FDeferredLightUniformStruct> CreateDeferredLightUnif
 	return GraphBuilder.CreateUniformBuffer(DeferredLightStruct);
 }
 
+static bool MayRenderSubstrateToon(ESubstrateTileType SubstrateTileType)
+{
+	static const bool bSubstrateEnable = Substrate::IsSubstrateEnabled();						// Not runtime toggleable
+	static const bool bToonUnifiedDiffuseEnable = IsExperimentalToonUnifiedDiffuseEnabled();	// Not runtime toggleable
+	return bSubstrateEnable && bToonUnifiedDiffuseEnable &&
+		(SubstrateTileType == ESubstrateTileType::EComplex || SubstrateTileType == ESubstrateTileType::EComplexSpecial);	// Complex and ComplexSpecial tiles maybe contain pixels with toon shading.
+}
+
 static FDeferredLightPS::FParameters GetDeferredLightPSParameters(
 	FRDGBuilder& GraphBuilder,
 	const FScene* Scene,
@@ -2703,8 +2779,7 @@ static FDeferredLightPS::FParameters GetDeferredLightPSParameters(
 	bool bCloudShadow,
 	TRDGUniformBufferRef<FVirtualShadowMapUniformParameters> VirtualShadowMapUniformBuffer = nullptr,
 	FRDGTextureRef ShadowMaskBits = nullptr,
-	int32 VirtualShadowMapId = INDEX_NONE,
-	ESubstrateTileType SubstrateTileType = ESubstrateTileType::ECount)
+	int32 VirtualShadowMapId = INDEX_NONE)
 {
 	FDeferredLightPS::FParameters Out;
 
@@ -2747,15 +2822,7 @@ static FDeferredLightPS::FParameters GetDeferredLightPSParameters(
 
 	// PS - Render Targets
 	Out.RenderTargets[0] = FRenderTargetBinding(SceneColorTexture, ERenderTargetLoadAction::ELoad);
-	if (Substrate::IsSubstrateEnabled() && SubstrateTileType >= ESubstrateTileType::EComplex && IsExperimentalToonUnifiedDiffuseEnabled())
-	{
-		const FSubstrateSceneData& SubstrateSceneData = *View.GetSceneTextures().SubstrateSceneData;
-		Out.RenderTargets[1] = FRenderTargetBinding(SubstrateSceneData.SubstrateSceneToonData, ERenderTargetLoadAction::ELoad);
-
-		// As of today, OpaqueRoughRefraction cannot be used together with experimental toon features. This is reflected in the implementation of IsOpaqueRoughRefractionEnabled().
-		check(!Substrate::IsOpaqueRoughRefractionEnabled(View.GetShaderPlatform()));
-	}
-	else if (Substrate::IsOpaqueRoughRefractionEnabled(View.GetShaderPlatform()) && Substrate::UsesSubstrateMaterialBuffer(View.GetShaderPlatform()))
+	if (Substrate::IsOpaqueRoughRefractionEnabled(View.GetShaderPlatform()) && Substrate::UsesSubstrateMaterialBuffer(View.GetShaderPlatform()))
 	{
 		const FSubstrateSceneData& SubstrateSceneData = *View.GetSceneTextures().SubstrateSceneData;
 		Out.RenderTargets[1] = FRenderTargetBinding(SubstrateSceneData.SeparatedOpaqueRoughRefractionSceneColor, ERenderTargetLoadAction::ELoad);
@@ -2785,7 +2852,7 @@ uint32 SetupLightGraphicsPSOState(
 	EShaderPlatform ShaderPlatform)
 {
 	GraphicsPSOInit.PrimitiveType = PT_TriangleList;
-	if (Substrate::IsSubstrateEnabled() && SubstrateTileMaterialType >= ESubstrateTileType::EComplex && IsExperimentalToonUnifiedDiffuseEnabled())
+	if (MayRenderSubstrateToon(SubstrateTileMaterialType))
 	{
 		GraphicsPSOInit.BlendState = TStaticBlendState<
 			CW_RGBA, BO_Add, BF_One, BF_One, BO_Add, BF_One, BF_One,
@@ -3110,7 +3177,7 @@ static void RenderLight(
 		// PS - Generatl parameters
 		PassParameters->PS = GetDeferredLightPSParameters(
 			GraphBuilder, Scene, View, LightSceneInfo, SceneTextures.Color.Target, SceneTextures.Depth.Target, SceneTextures.UniformBuffer, View.HairStrandsViewData.UniformBuffer, 
-			ScreenShadowMaskTexture, LightingChannelsTexture, bCloudShadow, VirtualShadowMapUniformBuffer, ShadowMaskBits, VirtualShadowMapId, SubstrateTileType);
+			ScreenShadowMaskTexture, LightingChannelsTexture, bCloudShadow, VirtualShadowMapUniformBuffer, ShadowMaskBits, VirtualShadowMapId);
 		// VS - General parameters
 		if (bIsRadial)
 		{
@@ -3174,6 +3241,28 @@ static void RenderLight(
 					return (InTileType == SubstrateTileType || SubstrateTileType == ESubstrateTileType::ECount) && Substrate::GetSubstrateUsesTileType(View, InTileType);
 				};
 
+			FRenderLightParameters* PassParametersForToon = nullptr;
+			const auto SelectComplexPassParameters = [&PassParametersForToon, &PassParameters, &GraphBuilder, &View](ESubstrateTileType InTileType)
+				{
+					if (MayRenderSubstrateToon(InTileType))
+					{
+						if(PassParametersForToon == nullptr)
+						{
+							// Copy the parameters once for all the tile types that can render
+							PassParametersForToon = GraphBuilder.AllocParameters<FRenderLightParameters>();
+							*PassParametersForToon = *PassParameters;
+							// And setup the needed toon render target
+							const FSubstrateSceneData& SubstrateSceneData = *View.GetSceneTextures().SubstrateSceneData;
+							PassParametersForToon->PS.RenderTargets[1] = FRenderTargetBinding(SubstrateSceneData.SubstrateSceneToonData, ERenderTargetLoadAction::ELoad);
+							// As of today, OpaqueRoughRefraction cannot be used together with experimental toon features. This is reflected in the implementation of IsOpaqueRoughRefractionEnabled().
+							check(!Substrate::IsOpaqueRoughRefractionEnabled(View.GetShaderPlatform()));
+
+						}
+						return PassParametersForToon;
+					}
+					return PassParameters;
+				};
+
 			if (ShouldRenderTileType(ESubstrateTileType::EComplexSpecial))
 			{
 				const ESubstrateTileType TileType = ESubstrateTileType::EComplexSpecial;
@@ -3182,7 +3271,8 @@ static void RenderLight(
 				TShaderMapRef< FDeferredLightPS > PixelShader(View.ShaderMap, PermutationVectorRemap);
 				TShaderMapRef< FDeferredLightPS > UberPixelShader = GetUberDeferredLightPS(View.ShaderMap, View.GetShaderPlatform(), PermutationVectorRemap);
 
-				InternalRenderLight(GraphBuilder, Scene, View, LightSceneInfo, PixelShader, UberPixelShader, PassParameters, TileType, TEXT("Light::StandardDeferred(ComplexSpecial)"));
+				FRenderLightParameters* SelectedPassParameters = SelectComplexPassParameters(TileType);
+				InternalRenderLight(GraphBuilder, Scene, View, LightSceneInfo, PixelShader, UberPixelShader, SelectedPassParameters, TileType, TEXT("Light::StandardDeferred(ComplexSpecial)"));
 			}
 			if (ShouldRenderTileType(ESubstrateTileType::EComplex))
 			{
@@ -3192,7 +3282,8 @@ static void RenderLight(
 				TShaderMapRef< FDeferredLightPS > PixelShader(View.ShaderMap, PermutationVectorRemap);
 				TShaderMapRef< FDeferredLightPS > UberPixelShader = GetUberDeferredLightPS(View.ShaderMap, View.GetShaderPlatform(), PermutationVectorRemap);
 
-				InternalRenderLight(GraphBuilder, Scene, View, LightSceneInfo, PixelShader, UberPixelShader, PassParameters, TileType, TEXT("Light::StandardDeferred(Complex)"));
+				FRenderLightParameters* SelectedPassParameters = SelectComplexPassParameters(TileType);
+				InternalRenderLight(GraphBuilder, Scene, View, LightSceneInfo, PixelShader, UberPixelShader, SelectedPassParameters, TileType, TEXT("Light::StandardDeferred(Complex)"));
 			}
 			if (ShouldRenderTileType(ESubstrateTileType::ESingle))
 			{
