Fix instanced mesh LOD popping in mobile velocity pass

-Hardcoded bAllowStencilDither=false forced all instances to use global LODFactor instead of per-instance values. Now conditionally set bAllowDither=true on mobile
when PrePass is in DDM_AllOpaqueNoVelocity mode, to enable per-instance LOD transitions.

#tests consoles and Editor preview mode

diff --git a/Engine/Source/Runtime/Renderer/Private/VelocityRendering.cpp b/Engine/Source/Runtime/Renderer/Private/VelocityRendering.cpp
--- a/Engine/Source/Runtime/Renderer/Private/VelocityRendering.cpp
+++ b/Engine/Source/Runtime/Renderer/Private/VelocityRendering.cpp
@@ -1225,6 +1225,8 @@ bool FVelocityMeshProcessor::Process(
 	}
 
 	// When velocity is used as a depth pass we need to set a correct stencil state on mobile
+	bool bAllowDither = false;
+	
 	if (FeatureLevel == ERHIFeatureLevel::ES3_1 && EarlyZPassMode == DDM_AllOpaqueNoVelocity)
 	{
 		extern TOptional<uint32> GetOptionalDitheringFlag(
@@ -1252,10 +1254,12 @@ bool FVelocityMeshProcessor::Process(
 			bUsesDeferredShading, 
 			GetOptionalDitheringFlag(ViewIfDynamicMeshCommand, MeshBatch, StaticMeshId),
 			true);
+			
+		bAllowDither = true;
 	}
 
 	FMeshMaterialShaderElementData ShaderElementData;
-	ShaderElementData.InitializeMeshMaterialData(ViewIfDynamicMeshCommand, PrimitiveSceneProxy, MeshBatch, StaticMeshId, false);
+	ShaderElementData.InitializeMeshMaterialData(ViewIfDynamicMeshCommand, PrimitiveSceneProxy, MeshBatch, StaticMeshId, bAllowDither);
 
 	const FMeshDrawCommandSortKey SortKey = CalculateMeshStaticSortKey(VelocityPassShaders.VertexShader, VelocityPassShaders.PixelShader);
 
