[HWRT] Only check bVertexCountChanges condition for ray tracing geometries with a single segment.
- that code path assumes a single segment and caused issues for geometry collections with multiple segments since MaxVertices for some segments might be different than the overall NumVertices in the geometry.

#Synced-CL 51569112

diff --git a/Engine/Source/Runtime/Renderer/Private/RayTracing/RayTracingDynamicGeometryUpdateManager.cpp b/Engine/Source/Runtime/Renderer/Private/RayTracing/RayTracingDynamicGeometryUpdateManager.cpp
--- a/Engine/Source/Runtime/Renderer/Private/RayTracing/RayTracingDynamicGeometryUpdateManager.cpp
+++ b/Engine/Source/Runtime/Renderer/Private/RayTracing/RayTracingDynamicGeometryUpdateManager.cpp
@@ -450,7 +450,8 @@ void FRayTracingDynamicGeometryUpdateManager::AddDynamicGeometryToUpdate(
 	check(Geometry.IsInitialized());
 
 	// Primitive count or triangle count changed then setup the initializer and segment data again
-	bool bVertexCountChanges = Geometry.Initializer.Segments.Num() > 0 ? Geometry.Initializer.Segments[0].MaxVertices != NumVertices : false;
+	// TODO: Move this to higher level code since it only works with 1 segment. These "resizes" should be handled by the systems that request dynamic geometry updates.
+	const bool bVertexCountChanges = Geometry.Initializer.Segments.Num() == 1 ? Geometry.Initializer.Segments[0].MaxVertices != NumVertices : false;
 	if ((Geometry.Initializer.TotalPrimitiveCount != UpdateParams.NumTriangles || bVertexCountChanges) && UpdateParams.NumTriangles > 0)
 	{
 		checkf(Geometry.Initializer.Segments.Num() <= 1, TEXT("Dynamic ray tracing geometry '%s' has an unexpected number of segments."), *Geometry.Initializer.DebugName.ToString());
