Change default VSM coarse page marking to performance mode (ignore dynamic invalidations).

#Synced-CL 51587375

#Synced-CL 51587375

diff --git a/Engine/Source/Runtime/Renderer/Private/VirtualShadowMaps/VirtualShadowMapClipmap.cpp b/Engine/Source/Runtime/Renderer/Private/VirtualShadowMaps/VirtualShadowMapClipmap.cpp
--- a/Engine/Source/Runtime/Renderer/Private/VirtualShadowMaps/VirtualShadowMapClipmap.cpp
+++ b/Engine/Source/Runtime/Renderer/Private/VirtualShadowMaps/VirtualShadowMapClipmap.cpp
@@ -128,11 +128,11 @@ static TAutoConsoleVariable<bool> CVarOrthoVSMRayCastViewOrigin(
 
 TAutoConsoleVariable<int32> CVarMarkCoarsePagesDirectional(
 	TEXT("r.Shadow.Virtual.MarkCoarsePagesDirectional"),
-	1,
+	2,
 	TEXT("Marks coarse pages in local light virtual shadow maps so that low resolution data is available everywhere.\n")
 	TEXT(" 0 - Disabled\n")
 	TEXT(" 1 - Always mark the last clipmap level and levels indicated by First/LastCoarseLevel.\n")
-	TEXT(" 2 - Performance Mode: Suppress dynamic invalidations due to geometry changes (e.g., moving objects or WPO or animation).\n"),
+	TEXT(" 2 - Performance Mode: Suppress dynamic invalidations in coarse pages due to geometry changes (e.g., moving objects or WPO or animation).\n"),
 	ECVF_Scalability | ECVF_RenderThreadSafe
 );
 
