Build fix because of partial change of another CL - restore old behavior

#[robomerge][partner] megalights

#Synced-CL 51609817

diff --git a/Engine/Source/Runtime/Renderer/Private/PrimitiveSceneInfo.cpp b/Engine/Source/Runtime/Renderer/Private/PrimitiveSceneInfo.cpp
--- a/Engine/Source/Runtime/Renderer/Private/PrimitiveSceneInfo.cpp
+++ b/Engine/Source/Runtime/Renderer/Private/PrimitiveSceneInfo.cpp
@@ -1054,10 +1054,13 @@ void FPrimitiveSceneInfo::RemoveCachedNaniteMaterialBins()
 		NanitePassRasterBins.Reset();
 		NanitePassShadingBins.Reset();
 		NaniteMaterialSlots[MeshPassIndex].Reset();
-	}
-
+	} 
+	
 	// Primitive doesn't care about pending static mesh updates anymore
-	Scene->ClearPendingRequestStaticMeshUpdate(this);
+	if (auto MaterialsExtension = Scene->GetExtensionPtr<Nanite::FMaterialsSceneExtension>())
+	{
+		MaterialsExtension->ClearPendingRequestStaticMeshUpdate(this);
+	}
 
 	bNaniteRasterBinsRenderCustomDepth = false;
 }
