Moved core logic for building the primitive scene data inside the proxy to BuildPrimitiveUniformShaderParameters_RenderThread in FScene to reduce exposure of internal rendering data and improve performance.

#Synced-CL 43344078

diff --git a/Engine/Source/Runtime/Engine/Private/PrimitiveSceneProxy.cpp b/Engine/Source/Runtime/Engine/Private/PrimitiveSceneProxy.cpp
--- a/Engine/Source/Runtime/Engine/Private/PrimitiveSceneProxy.cpp
+++ b/Engine/Source/Runtime/Engine/Private/PrimitiveSceneProxy.cpp
@@ -736,22 +736,15 @@ void FPrimitiveSceneProxy::UpdateUniformBuffer(FRHICommandList& RHICmdList)
 
 void FPrimitiveSceneProxy::BuildUniformShaderParameters(FPrimitiveUniformShaderParametersBuilder &Builder) const
 {
-	bool bHasPrecomputedVolumetricLightmap;
-	FMatrix PreviousLocalToWorld;
-	int32 SingleCaptureIndex;
-	bool bOutputVelocity;
-
-	Scene->GetPrimitiveUniformShaderParameters_RenderThread(
-		PrimitiveSceneInfo,
-		bHasPrecomputedVolumetricLightmap,
-		PreviousLocalToWorld,
-		SingleCaptureIndex,
-		bOutputVelocity
-	);
-
-	bOutputVelocity |= AlwaysHasVelocity();
-	bOutputVelocity |= (IsNaniteMesh() && IsSkinnedMesh());
+	Builder = FPrimitiveUniformShaderParametersBuilder{}.Defaults();
+	// 1. Set up parameters that depend on the proxy.
+	BuildUniformShaderParametersInternal(Builder);
+	// 2. Set up those that depend on renderer internal representations.
+	Scene->BuildPrimitiveUniformShaderParameters_RenderThread(PrimitiveSceneInfo, Builder);
+}
 
+void FPrimitiveSceneProxy::BuildUniformShaderParametersInternal(FPrimitiveUniformShaderParametersBuilder& Builder) const
+{
 	FBoxSphereBounds PreSkinnedLocalBounds;
 	GetPreSkinnedLocalBounds(PreSkinnedLocalBounds);
 
@@ -764,28 +757,21 @@ void FPrimitiveSceneProxy::BuildUniformShaderParameters(FPrimitiveUniformShaderP
 #endif
 
 	// Update the uniform shader parameters.
-	Builder = FPrimitiveUniformShaderParametersBuilder{}
-		.Defaults()
-			.LocalToWorld(LocalToWorld)
-			.PreviousLocalToWorld(PreviousLocalToWorld)
-			.ActorWorldPosition(ActorPosition)
+	Builder.ActorWorldPosition(ActorPosition)
 			.WorldBounds(Bounds)
 			.LocalBounds(LocalBounds)
 			.PreSkinnedLocalBounds(PreSkinnedLocalBounds)
 			.ReceivesDecals(bReceivesDecals)
-			.OutputVelocity(bOutputVelocity)
 			.EvaluateWorldPositionOffset(EvaluateWorldPositionOffset() && AnyMaterialHasWorldPositionOffset())
 			.MaxWorldPositionOffsetExtent(GetMaxWorldPositionOffsetExtent())
 			.HasAlwaysEvaluateWPOMaterials(AnyMaterialAlwaysEvaluatesWorldPositionOffset())
 			.MinMaxMaterialDisplacement(GetMinMaxMaterialDisplacement())
 			.LightingChannelMask(GetLightingChannelMask())
 			.LightmapUVIndex(GetLightMapCoordinateIndex())
-			.SingleCaptureIndex(SingleCaptureIndex)
 			.CustomPrimitiveData(GetCustomPrimitiveData())
 			.HasRayTracingRepresentation(bHasRayTracingRepresentation)
 			.HasCapsuleRepresentation(HasDynamicIndirectShadowCasterRepresentation())
 			.UseSingleSampleShadowFromStationaryLights(UseSingleSampleShadowFromStationaryLights())
-			.UseVolumetricLightmap(bHasPrecomputedVolumetricLightmap)
 			.CastContactShadow(CastsContactShadow())
 			.CastHiddenShadow(CastsHiddenShadow())
 			.CastShadow(CastsDynamicShadow())
@@ -814,16 +800,6 @@ void FPrimitiveSceneProxy::BuildUniformShaderParameters(FPrimitiveUniformShaderP
 			.IsFirstPerson(IsFirstPerson())
 			.AnimationMinScreenSize(GetAnimationMinScreenSize());
 
-	if (PrimitiveSceneInfo != nullptr)
-	{
-		Builder.LightmapDataIndex(PrimitiveSceneInfo->GetLightmapDataOffset())
-			.InstanceSceneDataOffset(PrimitiveSceneInfo->GetInstanceSceneDataOffset())
-			.NumInstanceSceneDataEntries(PrimitiveSceneInfo->GetNumInstanceSceneDataEntries())
-			.InstancePayloadDataOffset(PrimitiveSceneInfo->GetInstancePayloadDataOffset())
-			.InstancePayloadDataStride(PrimitiveSceneInfo->GetInstancePayloadDataStride())
-			.PersistentPrimitiveIndex(PrimitiveSceneInfo->GetPersistentIndex().Index);
-	}
-
 	if (IsNaniteMesh())
 	{
 		const Nanite::FSceneProxyBase* NaniteProxy = static_cast<const Nanite::FSceneProxyBase*>(this);
diff --git a/Engine/Source/Runtime/Engine/Public/PrimitiveSceneProxy.h b/Engine/Source/Runtime/Engine/Public/PrimitiveSceneProxy.h
--- a/Engine/Source/Runtime/Engine/Public/PrimitiveSceneProxy.h
+++ b/Engine/Source/Runtime/Engine/Public/PrimitiveSceneProxy.h
@@ -1285,6 +1285,13 @@ protected:
 	 */
 	ENGINE_API void EnableGPUSceneSupportFlags();
 
+private:
+	/**
+	 * Internally used function that sets up the parameters that are not handled by the Scene BuildPrimitiveUniformShaderParameters_RenderThread implementations.
+     * In general, this should not be expanded on, as the data representation should be kept internal in the renderer.
+	 */
+	ENGINE_API void BuildUniformShaderParametersInternal(FPrimitiveUniformShaderParametersBuilder& Builder) const;
+
 public:
 	TSharedPtr<int32, ESPMode::ThreadSafe> SimpleStreamableAssetManagerIndex = MakeShared<int32, ESPMode::ThreadSafe>(INDEX_NONE);
 	
diff --git a/Engine/Source/Runtime/Engine/Public/PrimitiveUniformShaderParametersBuilder.h b/Engine/Source/Runtime/Engine/Public/PrimitiveUniformShaderParametersBuilder.h
--- a/Engine/Source/Runtime/Engine/Public/PrimitiveUniformShaderParametersBuilder.h
+++ b/Engine/Source/Runtime/Engine/Public/PrimitiveUniformShaderParametersBuilder.h
@@ -332,23 +332,24 @@ public:
 			// Also use double precision to calculate WorldToPreviousWorld to prevent precision issues at far distances
 
 			FMatrix LocalToRelativeWorld = FDFMatrix::MakeToRelativeWorldMatrixDouble(PositionHigh, AbsoluteLocalToWorld);
-			FMatrix PrevLocalToRelativeWorld = FDFMatrix::MakeClampedToRelativeWorldMatrixDouble(PositionHigh, AbsolutePreviousLocalToWorld);
 
 			FMatrix RelativeWorldToLocal = LocalToRelativeWorld.Inverse();
 
 			Parameters.LocalToRelativeWorld = FMatrix44f(LocalToRelativeWorld);
 			Parameters.RelativeWorldToLocal = FMatrix44f(RelativeWorldToLocal);
-			Parameters.WorldToPreviousWorld = FMatrix44f(RelativeWorldToLocal * PrevLocalToRelativeWorld);
 
 			if (bHasPreviousLocalToWorld)
 			{
+				FMatrix PrevLocalToRelativeWorld = FDFMatrix::MakeClampedToRelativeWorldMatrixDouble(PositionHigh, AbsolutePreviousLocalToWorld);
 				Parameters.PreviousLocalToRelativeWorld = FMatrix44f(PrevLocalToRelativeWorld);
 				Parameters.PreviousRelativeWorldToLocal = FMatrix44f(PrevLocalToRelativeWorld.Inverse());
+				Parameters.WorldToPreviousWorld = FMatrix44f(RelativeWorldToLocal * PrevLocalToRelativeWorld);
 			}
 			else
 			{
 				Parameters.PreviousLocalToRelativeWorld = Parameters.LocalToRelativeWorld;
 				Parameters.PreviousRelativeWorldToLocal = Parameters.RelativeWorldToLocal;
+				Parameters.WorldToPreviousWorld = FMatrix44f::Identity;
 			}
 		}
 
diff --git a/Engine/Source/Runtime/Engine/Public/SceneInterface.h b/Engine/Source/Runtime/Engine/Public/SceneInterface.h
--- a/Engine/Source/Runtime/Engine/Public/SceneInterface.h
+++ b/Engine/Source/Runtime/Engine/Public/SceneInterface.h
@@ -74,6 +74,7 @@ struct FExponentialHeightFogDynamicState;
 struct FPrimitiveSceneDesc;
 struct FInstancedStaticMeshSceneDesc;
 struct FLightSceneDesc;
+struct FPrimitiveUniformShaderParametersBuilder;
 
 enum EBasePassDrawListType
 {
@@ -339,9 +340,15 @@ public:
 	virtual void InvalidateLumenSurfaceCache_GameThread(UPrimitiveComponent* Component) {}
 
 	/** 
-	 * Retrieves primitive uniform shader parameters that are internal to the renderer.
+	 * Retrieves primitive uniform shader parameters that are internal to the renderer. 
+	 * NOTE: prefer BuildPrimitiveUniformShaderParameters_RenderThread since it is more efficient and general.
 	 */
-	virtual void GetPrimitiveUniformShaderParameters_RenderThread(const FPrimitiveSceneInfo* PrimitiveSceneInfo, bool& bHasPrecomputedVolumetricLightmap, FMatrix& PreviousLocalToWorld, int32& SingleCaptureIndex, bool& OutputVelocity) const {}
+	virtual void GetPrimitiveUniformShaderParameters_RenderThread(const FPrimitiveSceneInfo* PrimitiveSceneInfo, bool& bHasPrecomputedVolumetricLightmap, FMatrix& PreviousLocalToWorld, int32& SingleCaptureIndex, bool& OutputVelocity) const = 0;
+
+	/**
+	 * Sets up primitive uniform shader parameters that are internal to the renderer.
+	 */
+	virtual void BuildPrimitiveUniformShaderParameters_RenderThread(const FPrimitiveSceneInfo* PrimitiveSceneInfo, FPrimitiveUniformShaderParametersBuilder& Builder) const = 0;
 	 
 	/** 
 	 * Updates the transform of a light which has already been added to the scene. 
diff --git a/Engine/Source/Runtime/Renderer/Private/RendererScene.cpp b/Engine/Source/Runtime/Renderer/Private/RendererScene.cpp
--- a/Engine/Source/Runtime/Renderer/Private/RendererScene.cpp
+++ b/Engine/Source/Runtime/Renderer/Private/RendererScene.cpp
@@ -124,6 +124,8 @@
 #if WITH_EDITOR
 #include "Rendering/StaticLightingSystemInterface.h"
 #endif
+#include "PrimitiveUniformShaderParametersBuilder.h"
+#include "PrimitiveSceneShaderData.h"
 
 #define VALIDATE_PRIMITIVE_PACKED_INDEX 0
 
@@ -3400,6 +3402,9 @@ void FSceneVelocityData::StartFrame(FScene* Scene)
 void FScene::GetPrimitiveUniformShaderParameters_RenderThread(const FPrimitiveSceneInfo* PrimitiveSceneInfo, bool& bHasPrecomputedVolumetricLightmap, FMatrix& PreviousLocalToWorld, int32& SingleCaptureIndex, bool& bOutputVelocity) const 
 {
 	SCOPED_NAMED_EVENT(GetPrimitiveUniformShaderParameters_RenderThread, FColor::Yellow);
+
+	// NOTE: The logic here is duplicated in BuildPrimitiveUniformShaderParameters_RenderThread since that allows reordering the logic to avoid transform copies & needlessly setting Builder.PreviousLocalToWorld - which avoids 
+	//       a matrix inverse.
 	const FMatrix LocalToWorld = PrimitiveSceneInfo->Proxy->GetLocalToWorld();
 	PreviousLocalToWorld = LocalToWorld;
 	bOutputVelocity = false;
@@ -3416,6 +3421,39 @@ void FScene::GetPrimitiveUniformShaderParameters_RenderThread(const FPrimitiveSc
 	SingleCaptureIndex = PrimitiveSceneInfo->CachedReflectionCaptureProxy ? PrimitiveSceneInfo->CachedReflectionCaptureProxy->SortedCaptureIndex : 0;
 }
 
+void FScene::BuildPrimitiveUniformShaderParameters_RenderThread(const FPrimitiveSceneInfo* PrimitiveSceneInfo, FPrimitiveUniformShaderParametersBuilder& Builder) const
+{
+	const FPrimitiveSceneProxy* Proxy = PrimitiveSceneInfo->Proxy;
+
+	const FMatrix LocalToWorld = Proxy->GetLocalToWorld();
+
+	bool bOutputVelocity = Proxy->AlwaysHasVelocity() || (Proxy->IsNaniteMesh() && Proxy->IsSkinnedMesh());
+	const bool bHasPrecomputedVolumetricLightmap = VolumetricLightmapSceneData.HasData();
+
+	// Get index if proxy exists, otherwise fall back to index 0 which will contain the default black cubemap
+	const int32 SingleCaptureIndex = PrimitiveSceneInfo->CachedReflectionCaptureProxy ? PrimitiveSceneInfo->CachedReflectionCaptureProxy->SortedCaptureIndex : 0;
+
+	Builder.LocalToWorld(LocalToWorld)
+			.SingleCaptureIndex(SingleCaptureIndex)
+			.UseVolumetricLightmap(bHasPrecomputedVolumetricLightmap);
+
+	// Only set up previous transform if we actually have one.
+	if (const FMatrix* RESTRICT PreviousLocalToWorld = VelocityData.GetComponentPreviousLocalToWorld(PrimitiveSceneInfo->PrimitiveComponentId))
+	{
+		Builder.PreviousLocalToWorld(*PreviousLocalToWorld);
+		bOutputVelocity = bOutputVelocity || !LocalToWorld.Equals(*PreviousLocalToWorld, 0.0001f);
+	}
+
+	Builder.OutputVelocity(bOutputVelocity)
+			.LightmapDataIndex(PrimitiveSceneInfo->GetLightmapDataOffset())
+			.InstanceSceneDataOffset(PrimitiveSceneInfo->GetInstanceSceneDataOffset())
+			.NumInstanceSceneDataEntries(PrimitiveSceneInfo->GetNumInstanceSceneDataEntries())
+			.InstancePayloadDataOffset(PrimitiveSceneInfo->GetInstancePayloadDataOffset())
+			.InstancePayloadDataStride(PrimitiveSceneInfo->GetInstancePayloadDataStride())
+			.PersistentPrimitiveIndex(PrimitiveSceneInfo->GetPersistentIndex().Index);
+}
+
+
 bool DoesPlatformNeedLocalLightPrimitiveInteraction(EShaderPlatform ShaderPlatform)
 {
 	if (CVarVisibilityLocalLightPrimitiveInteraction.GetValueOnRenderThread() == 0)
@@ -6807,6 +6845,16 @@ public:
 	virtual void UpdateDecalFadeOutTime(UDecalComponent* Decal) override {};
 	virtual void UpdateDecalFadeInTime(UDecalComponent* Decal) override {};
 	virtual void BatchUpdateDecals(TArray<FDeferredDecalUpdateParams>&& UpdateParams) override {}
+	virtual void GetPrimitiveUniformShaderParameters_RenderThread(const FPrimitiveSceneInfo* PrimitiveSceneInfo, bool& bHasPrecomputedVolumetricLightmap, FMatrix& PreviousLocalToWorld, int32& SingleCaptureIndex, bool& OutputVelocity) const override 
+	{
+		PreviousLocalToWorld = FMatrix::Identity; 
+		SingleCaptureIndex = 0;
+		bHasPrecomputedVolumetricLightmap = OutputVelocity = false;
+	}
+	virtual void BuildPrimitiveUniformShaderParameters_RenderThread(const FPrimitiveSceneInfo* PrimitiveSceneInfo, FPrimitiveUniformShaderParametersBuilder& Builder) const override 
+	{
+		Builder.Defaults();
+	};
 
 	/** Updates the transform of a light which has already been added to the scene. */
 	virtual void UpdateLightTransform(ULightComponent* Light) override {}
diff --git a/Engine/Source/Runtime/Renderer/Private/ScenePrivate.h b/Engine/Source/Runtime/Renderer/Private/ScenePrivate.h
--- a/Engine/Source/Runtime/Renderer/Private/ScenePrivate.h
+++ b/Engine/Source/Runtime/Renderer/Private/ScenePrivate.h
@@ -2588,6 +2588,24 @@ public:
 		return false;
 	}
 
+	/** 
+	 * Looks up the PreviousLocalToWorld state for the given component.  
+	 * Returns null if none is found (the primitive has never been moved). 
+	 */
+	const FMatrix* RESTRICT GetComponentPreviousLocalToWorld(FPrimitiveComponentId PrimitiveComponentId) const
+	{
+		const FComponentVelocityData* VelocityData = ComponentData.Find(PrimitiveComponentId);
+
+		if (VelocityData && VelocityData->PrimitiveSceneInfo)
+		{
+			check(VelocityData->bPreviousLocalToWorldValid);
+			VelocityData->LastFrameUsed = InternalFrameIndex;
+			return &VelocityData->PreviousLocalToWorld;
+		}
+
+		return nullptr;
+	}
+
 	/** 
 	 * Updates a primitives current LocalToWorld state.
 	 */
@@ -3504,6 +3522,7 @@ public:
 	virtual void InvalidatePathTracedOutput(PathTracing::EInvalidateReason InvalidateReason = PathTracing::EInvalidateReason::Uncategorized) override;
 	virtual void InvalidateLumenSurfaceCache_GameThread(UPrimitiveComponent* Component) override;
 	virtual void GetPrimitiveUniformShaderParameters_RenderThread(const FPrimitiveSceneInfo* PrimitiveSceneInfo, bool& bHasPrecomputedVolumetricLightmap, FMatrix& PreviousLocalToWorld, int32& SingleCaptureIndex, bool& bOutputVelocity) const override;
+	virtual void BuildPrimitiveUniformShaderParameters_RenderThread(const FPrimitiveSceneInfo* PrimitiveSceneInfo, FPrimitiveUniformShaderParametersBuilder& Builder) const override;
 	virtual void UpdateLightTransform(ULightComponent* Light) override;
 	virtual void UpdateLightColorAndBrightness(ULightComponent* Light) override;
 	virtual void UpdateLightProxy(ULightComponent* Light, TFunction<void(FLightSceneProxy*)>&& Func) override;
