Fix race condition when in hair code.

For supporting high resolution screenshot, the hair system was reading GIsHighResScreenshot/GScreenshotResolutionX/GScreenshotResolutionY for computing the correct final pixel size so that hair coverage computation respect the final resolution. However the globals are initialize on the gamethread, but their value is read on the render thread causing a minor race condition. To avoid this, this CL introduce a final resolution scale, avoiding reading global variables.

#Synced-CL 51662140

diff --git a/Engine/Source/Editor/UnrealEd/Private/EditorViewportClient.cpp b/Engine/Source/Editor/UnrealEd/Private/EditorViewportClient.cpp
--- a/Engine/Source/Editor/UnrealEd/Private/EditorViewportClient.cpp
+++ b/Engine/Source/Editor/UnrealEd/Private/EditorViewportClient.cpp
@@ -4905,6 +4905,12 @@ void FEditorViewportClient::Draw(FViewport* InViewport, FCanvas* Canvas)
 		}
 	}
 
+	// Setup final resolution multiplier used by system relying on the final pixel size. Only used by the high resolution screenshot
+	if (GIsHighResScreenshot)
+	{
+		ViewFamily.FinalResolutionScale = GetHighResScreenshotConfig().ResolutionMultiplier;
+	}
+
 	FSceneView* View = nullptr;
 
 	// Stereo rendering
diff --git a/Engine/Source/Runtime/Engine/Private/GameViewportClient.cpp b/Engine/Source/Runtime/Engine/Private/GameViewportClient.cpp
--- a/Engine/Source/Runtime/Engine/Private/GameViewportClient.cpp
+++ b/Engine/Source/Runtime/Engine/Private/GameViewportClient.cpp
@@ -1422,6 +1422,12 @@ void UGameViewportClient::Draw(FViewport* InViewport, FCanvas* SceneCanvas)
 
 	ViewFamily.DebugDPIScale = GetDPIScale();
 
+	// Setup final resolution multiplier used by system relying on the final pixel size. Only used by the high resolution screenshot
+	if (GIsHighResScreenshot)
+	{
+		ViewFamily.FinalResolutionScale = GetHighResScreenshotConfig().ResolutionMultiplier;
+	}
+
 #if WITH_EDITOR
 	if (GIsEditor && ViewFamily.SupportsScreenPercentage())
 	{
diff --git a/Engine/Source/Runtime/Engine/Public/SceneView.h b/Engine/Source/Runtime/Engine/Public/SceneView.h
--- a/Engine/Source/Runtime/Engine/Public/SceneView.h
+++ b/Engine/Source/Runtime/Engine/Public/SceneView.h
@@ -2491,6 +2491,11 @@ public:
 	/** DPI scale to be used for debugging font. */
 	float DebugDPIScale = 1.0f;
 
+	/** Scale between the current rendering resolution and the final output resolution.
+	 *  Used e.g. by high-res tiled screenshot capture where individual tiles are rendered
+	 *  at a smaller resolution and composited into a larger final image */
+	float FinalResolutionScale = 1.0f;
+
 	/** Editor setting to allow designers to override the automatic expose. 0:Automatic, following indices: -4 .. +4 */
 	FExposureSettings ExposureSettings;
 
diff --git a/Engine/Source/Runtime/Renderer/Private/HairStrands/HairStrandsForwardRaster.cpp b/Engine/Source/Runtime/Renderer/Private/HairStrands/HairStrandsForwardRaster.cpp
--- a/Engine/Source/Runtime/Renderer/Private/HairStrands/HairStrandsForwardRaster.cpp
+++ b/Engine/Source/Runtime/Renderer/Private/HairStrands/HairStrandsForwardRaster.cpp
@@ -21,9 +21,9 @@ extern int32 GHairVisibilityComputeRaster_NumRasterizers;
 extern int32 GHairVisibilityComputeRaster_NumRasterizersNaive;
 extern int32 GHairVisibilityComputeRaster_Debug;
 
-void SetUpViewHairRenderInfo(const FViewInfo& ViewInfo, FVector4f& OutHairRenderInfo, uint32& OutHairRenderInfoBits, uint32& OutHairComponents);
-void SetUpViewHairRenderInfo(const FViewInfo& ViewInfo, bool bEnableMSAA, FVector4f& OutHairRenderInfo, uint32& OutHairRenderInfoBits, uint32& OutHairComponents);
-FMinHairRadiusAtDepth1 ComputeMinStrandRadiusAtDepth1(const FIntPoint& Resolution, const float FOV, const uint32 SampleCount, const float OverrideStrandHairRasterizationScale, const float OrthoWidth);
+void SetUpViewHairRenderInfo(const FViewInfo& ViewInfo, FVector4f& OutHairRenderInfo, uint32& OutHairRenderInfoBits, uint32& OutHairComponents, float FinalResolutionScale);
+void SetUpViewHairRenderInfo(const FViewInfo& ViewInfo, bool bEnableMSAA, FVector4f& OutHairRenderInfo, uint32& OutHairRenderInfoBits, uint32& OutHairComponents, float FinalResolutionScale);
+FMinHairRadiusAtDepth1 ComputeMinStrandRadiusAtDepth1(const FIntPoint& Resolution, const float FOV, const uint32 SampleCount, const float OverrideStrandHairRasterizationScale, const float OrthoWidth, const float FinalResolutionScale);
 
 inline bool IsHairStrandsForwardRasterSupported(EShaderPlatform In) 
 { 
@@ -370,7 +370,7 @@ FRasterForwardCullingOutput AddHairStrandsForwardCullingPass(
 	// Create and set the uniform buffer
 	const bool bEnableMSAA = false;
 	TUniformBufferRef<FViewUniformShaderParameters> ViewUniformShaderParameters;
-	SetUpViewHairRenderInfo(ViewInfo, bEnableMSAA, ViewInfo.CachedViewUniformShaderParameters->HairRenderInfo, ViewInfo.CachedViewUniformShaderParameters->HairRenderInfoBits, ViewInfo.CachedViewUniformShaderParameters->HairComponents);
+	SetUpViewHairRenderInfo(ViewInfo, bEnableMSAA, ViewInfo.CachedViewUniformShaderParameters->HairRenderInfo, ViewInfo.CachedViewUniformShaderParameters->HairRenderInfoBits, ViewInfo.CachedViewUniformShaderParameters->HairComponents, ViewInfo.Family->FinalResolutionScale);
 	ViewUniformShaderParameters = TUniformBufferRef<FViewUniformShaderParameters>::CreateUniformBufferImmediate(*ViewInfo.CachedViewUniformShaderParameters, UniformBuffer_SingleFrame);
 
 	TShaderMapRef<FHairCullSegmentCS> CullShader(ViewInfo.ShaderMap);
@@ -558,7 +558,7 @@ void AddHairStrandsForwardRasterPass(
 	TUniformBufferRef<FViewUniformShaderParameters> ViewUniformShaderParameters;
 	{
 		const bool bEnableMSAA = false;
-		SetUpViewHairRenderInfo(ViewInfo, bEnableMSAA, ViewInfo.CachedViewUniformShaderParameters->HairRenderInfo, ViewInfo.CachedViewUniformShaderParameters->HairRenderInfoBits, ViewInfo.CachedViewUniformShaderParameters->HairComponents);
+		SetUpViewHairRenderInfo(ViewInfo, bEnableMSAA, ViewInfo.CachedViewUniformShaderParameters->HairRenderInfo, ViewInfo.CachedViewUniformShaderParameters->HairRenderInfoBits, ViewInfo.CachedViewUniformShaderParameters->HairComponents, ViewInfo.Family->FinalResolutionScale);
 		ViewUniformShaderParameters = TUniformBufferRef<FViewUniformShaderParameters>::CreateUniformBufferImmediate(*ViewInfo.CachedViewUniformShaderParameters, UniformBuffer_SingleFrame);
 	}
 
@@ -580,7 +580,7 @@ void AddHairStrandsForwardRasterPass(
 		Common.ControlPoints = InData.ControlPointsSRV;
 		Common.ControlPointCount = InData.ControlPointCount;
 
-		Common.RadiusAtDepth1 = ComputeMinStrandRadiusAtDepth1(FIntPoint(ViewInfo.UnconstrainedViewRect.Width(), ViewInfo.UnconstrainedViewRect.Height()), ViewInfo.FOV, 1 /*SampleCount*/, 0 /*ScaleOverride*/, ViewInfo.ViewMatrices.GetOrthoDimensions().X).Primary;
+		Common.RadiusAtDepth1 = ComputeMinStrandRadiusAtDepth1(FIntPoint(ViewInfo.UnconstrainedViewRect.Width(), ViewInfo.UnconstrainedViewRect.Height()), ViewInfo.FOV, 1 /*SampleCount*/, 0 /*ScaleOverride*/, ViewInfo.ViewMatrices.GetOrthoDimensions().X, ViewInfo.Family->FinalResolutionScale).Primary;
 	}
 
 	// Fill in tile depth
diff --git a/Engine/Source/Runtime/Renderer/Private/HairStrands/HairStrandsUtils.cpp b/Engine/Source/Runtime/Renderer/Private/HairStrands/HairStrandsUtils.cpp
--- a/Engine/Source/Runtime/Renderer/Private/HairStrands/HairStrandsUtils.cpp
+++ b/Engine/Source/Runtime/Renderer/Private/HairStrands/HairStrandsUtils.cpp
@@ -125,14 +125,12 @@ FMinHairRadiusAtDepth1 ComputeMinStrandRadiusAtDepth1(
 	const float FOV,
 	const uint32 SampleCount,
 	const float OverrideStrandHairRasterizationScale,
-	const float OrthoWidth)
+	const float OrthoWidth,
+	const float FinalResolutionScale)
 {
-	FIntPoint Resolution = InResolution;
-	if (GIsHighResScreenshot)
-	{
-		Resolution.X = GScreenshotResolutionX;
-		Resolution.Y = GScreenshotResolutionY;
-	}
+	FIntPoint Resolution;
+	Resolution.X = FMath::CeilToInt32(InResolution.X * FinalResolutionScale);
+	Resolution.Y = FMath::CeilToInt32(InResolution.Y * FinalResolutionScale);
 
 	auto InternalMinRadiusAtDepth1 = [Resolution, FOV, SampleCount, OrthoWidth](float RasterizationScale)
 	{
diff --git a/Engine/Source/Runtime/Renderer/Private/HairStrands/HairStrandsUtils.h b/Engine/Source/Runtime/Renderer/Private/HairStrands/HairStrandsUtils.h
--- a/Engine/Source/Runtime/Renderer/Private/HairStrands/HairStrandsUtils.h
+++ b/Engine/Source/Runtime/Renderer/Private/HairStrands/HairStrandsUtils.h
@@ -19,7 +19,8 @@ FMinHairRadiusAtDepth1 ComputeMinStrandRadiusAtDepth1(
 	const float FOV,
 	const uint32 SampleCount,
 	const float OverrideStrandHairRasterizationScale,
-	const float OrthoWidth = 0.0f);
+	const float OrthoWidth = 0.0f,
+	const float FinalResolutionScale = 1.0f);
 
 FIntRect ComputeProjectedScreenRect(const FBox& B, const FViewInfo& View);
 
diff --git a/Engine/Source/Runtime/Renderer/Private/HairStrands/HairStrandsVisibility.cpp b/Engine/Source/Runtime/Renderer/Private/HairStrands/HairStrandsVisibility.cpp
--- a/Engine/Source/Runtime/Renderer/Private/HairStrands/HairStrandsVisibility.cpp
+++ b/Engine/Source/Runtime/Renderer/Private/HairStrands/HairStrandsVisibility.cpp
@@ -303,7 +303,7 @@ static uint32 GetTotalSampleCountForAllocation(FIntPoint Resolution, EShaderPlat
 	return Resolution.X * Resolution.Y * GetMeanSamplePerPixel(InPlatform);
 }
 
-void SetUpViewHairRenderInfo(const FViewInfo& ViewInfo, bool bEnableMSAA, FVector4f& OutHairRenderInfo, uint32& OutHairRenderInfoBits, uint32& OutHairComponents)
+void SetUpViewHairRenderInfo(const FViewInfo& ViewInfo, bool bEnableMSAA, FVector4f& OutHairRenderInfo, uint32& OutHairRenderInfoBits, uint32& OutHairComponents, float FinalResolutionScale)
 {
 	FVector2f PixelVelocity(1.f / (ViewInfo.ViewRect.Width() * 2), 1.f / (ViewInfo.ViewRect.Height() * 2));
 	const float VelocityMagnitudeScale = FMath::Clamp(CVarHairVelocityMagnitudeScale.GetValueOnAnyThread(), 0, 512) * FMath::Min(PixelVelocity.X, PixelVelocity.Y);
@@ -314,16 +314,16 @@ void SetUpViewHairRenderInfo(const FViewInfo& ViewInfo, bool bEnableMSAA, FVecto
 
 	const FIntRect EffectiveViewRect = GHairResolutionScaleAware > 0 ? ViewInfo.ViewRect : ViewInfo.UnconstrainedViewRect;
 	const FMinHairRadiusAtDepth1 MinHairRadius = ComputeMinStrandRadiusAtDepth1(
-		FIntPoint(EffectiveViewRect.Width(), EffectiveViewRect.Height()), ViewInfo.FOV, HairVisibilitySampleCount, RasterizationScaleOverride, ViewInfo.ViewMatrices.GetOrthoDimensions().X);
+		FIntPoint(EffectiveViewRect.Width(), EffectiveViewRect.Height()), ViewInfo.FOV, HairVisibilitySampleCount, RasterizationScaleOverride, ViewInfo.ViewMatrices.GetOrthoDimensions().X, FinalResolutionScale);
 
 	OutHairRenderInfo = PackHairRenderInfo(MinHairRadius.Primary, MinHairRadius.Stable, MinHairRadius.Velocity, VelocityMagnitudeScale);
 	OutHairRenderInfoBits = PackHairRenderInfoBits(!ViewInfo.IsPerspectiveProjection(), false);
 	OutHairComponents = ToBitfield(GetHairComponents());
 }
 
-void SetUpViewHairRenderInfo(const FViewInfo& ViewInfo, FVector4f& OutHairRenderInfo, uint32& OutHairRenderInfoBits, uint32& OutHairComponents)
+void SetUpViewHairRenderInfo(const FViewInfo& ViewInfo, FVector4f& OutHairRenderInfo, uint32& OutHairRenderInfoBits, uint32& OutHairComponents, float FinalResolutionScale)
 {
-	SetUpViewHairRenderInfo(ViewInfo, IsMsaaEnabled(ViewInfo.GetShaderPlatform()), OutHairRenderInfo, OutHairRenderInfoBits, OutHairComponents);
+	SetUpViewHairRenderInfo(ViewInfo, IsMsaaEnabled(ViewInfo.GetShaderPlatform()), OutHairRenderInfo, OutHairRenderInfoBits, OutHairComponents, FinalResolutionScale);
 }
 
 static bool IsCompatibleWithHairVisibility(const FMeshMaterialShaderPermutationParameters& Parameters)
@@ -1063,7 +1063,7 @@ static FMaterialPassOutput AddHairMaterialPass(
 
 		{
 			const bool bEnableMSAA = false;
-			SetUpViewHairRenderInfo(*ViewInfo, bEnableMSAA, ViewInfo->CachedViewUniformShaderParameters->HairRenderInfo, ViewInfo->CachedViewUniformShaderParameters->HairRenderInfoBits, ViewInfo->CachedViewUniformShaderParameters->HairComponents);
+			SetUpViewHairRenderInfo(*ViewInfo, bEnableMSAA, ViewInfo->CachedViewUniformShaderParameters->HairRenderInfo, ViewInfo->CachedViewUniformShaderParameters->HairRenderInfoBits, ViewInfo->CachedViewUniformShaderParameters->HairComponents, ViewInfo->Family->FinalResolutionScale);
 			PassParameters->View = TUniformBufferRef<FViewUniformShaderParameters>::CreateUniformBufferImmediate(*ViewInfo->CachedViewUniformShaderParameters, UniformBuffer_SingleFrame);
 		}
 
@@ -2444,7 +2444,7 @@ static void AddHairVisibilityCommonPass(
 	if (RenderMode == HairVisibilityRenderMode_Transmittance || RenderMode == HairVisibilityRenderMode_TransmittanceAndHairCount || RenderMode == HairVisibilityRenderMode_PPLL)
 	{
 		const bool bEnableMSAA = false;
-		SetUpViewHairRenderInfo(*ViewInfo, bEnableMSAA, ViewInfo->CachedViewUniformShaderParameters->HairRenderInfo, ViewInfo->CachedViewUniformShaderParameters->HairRenderInfoBits, ViewInfo->CachedViewUniformShaderParameters->HairComponents);
+		SetUpViewHairRenderInfo(*ViewInfo, bEnableMSAA, ViewInfo->CachedViewUniformShaderParameters->HairRenderInfo, ViewInfo->CachedViewUniformShaderParameters->HairRenderInfoBits, ViewInfo->CachedViewUniformShaderParameters->HairComponents, ViewInfo->Family->FinalResolutionScale);
 
 		// Create and set the uniform buffer
 		PassParameters->View = TUniformBufferRef<FViewUniformShaderParameters>::CreateUniformBufferImmediate(*ViewInfo->CachedViewUniformShaderParameters, UniformBuffer_SingleFrame);
@@ -3689,7 +3689,7 @@ static FRasterComputeOutput AddVisibilityComputeRasterPass(
 	// Create and set the uniform buffer
 	const bool bEnableMSAA = false;
 	TUniformBufferRef<FViewUniformShaderParameters> ViewUniformShaderParameters;
-	SetUpViewHairRenderInfo(ViewInfo, bEnableMSAA, ViewInfo.CachedViewUniformShaderParameters->HairRenderInfo, ViewInfo.CachedViewUniformShaderParameters->HairRenderInfoBits, ViewInfo.CachedViewUniformShaderParameters->HairComponents);
+	SetUpViewHairRenderInfo(ViewInfo, bEnableMSAA, ViewInfo.CachedViewUniformShaderParameters->HairRenderInfo, ViewInfo.CachedViewUniformShaderParameters->HairRenderInfoBits, ViewInfo.CachedViewUniformShaderParameters->HairComponents, ViewInfo.Family->FinalResolutionScale);
 	ViewUniformShaderParameters = TUniformBufferRef<FViewUniformShaderParameters>::CreateUniformBufferImmediate(*ViewInfo.CachedViewUniformShaderParameters, UniformBuffer_SingleFrame);
 
 	const bool bClassification = GHairStrandsHWSWClassifaction != 0 && SamplePerPixelCount <= 1;
@@ -3868,7 +3868,7 @@ static FRasterComputeOutput AddVisibilityComputeRasterPass(
 				Parameters->VS.ViewDir = FVector3f(ViewDirection.X, ViewDirection.Y, ViewDirection.Z);
 				Parameters->VS.HairMaterialId = PrimitiveInfo.MaterialId;
 				Parameters->VS.CameraOrigin = FVector3f(CameraOrigin.X, CameraOrigin.Y, CameraOrigin.Z);
-				Parameters->VS.RadiusAtDepth1 = ComputeMinStrandRadiusAtDepth1(Resolution, ViewInfo.FOV, 1, -1.0f, ViewInfo.ViewMatrices.GetOrthoDimensions().X).Primary;
+				Parameters->VS.RadiusAtDepth1 = ComputeMinStrandRadiusAtDepth1(Resolution, ViewInfo.FOV, 1, -1.0f, ViewInfo.ViewMatrices.GetOrthoDimensions().X, ViewInfo.Family->FinalResolutionScale).Primary;
 				Parameters->VS.PrimIDsBufferSize = MaxNumPrimIDs;
 				Parameters->VS.HairInstance = GetHairStrandsInstanceParameters(GraphBuilder, ViewInfo, HairGroupPublicData, bCullingEnable, bForceRegister);
 				Parameters->VS.ViewUniformBuffer = ViewUniformShaderParameters;
diff --git a/Engine/Source/Runtime/Renderer/Private/HairStrands/HairStrandsVisibility.h b/Engine/Source/Runtime/Renderer/Private/HairStrands/HairStrandsVisibility.h
--- a/Engine/Source/Runtime/Renderer/Private/HairStrands/HairStrandsVisibility.h
+++ b/Engine/Source/Runtime/Renderer/Private/HairStrands/HairStrandsVisibility.h
@@ -28,5 +28,5 @@ void RenderHairStrandsVisibilityBuffer(
 	FInstanceCullingManager& InstanceCullingManager);
 
 EPixelFormat GetHairLightSampleFormat();
-void SetUpViewHairRenderInfo(const FViewInfo& ViewInfo, FVector4f& OutHairRenderInfo, uint32& OutHairRenderInfoBits, uint32& OutHairComponents);
+void SetUpViewHairRenderInfo(const FViewInfo& ViewInfo, FVector4f& OutHairRenderInfo, uint32& OutHairRenderInfoBits, uint32& OutHairComponents, float FinalResolutionScale = 1.0f);
 
diff --git a/Engine/Source/Runtime/Renderer/Private/HairStrands/HairStrandsVoxelization.cpp b/Engine/Source/Runtime/Renderer/Private/HairStrands/HairStrandsVoxelization.cpp
--- a/Engine/Source/Runtime/Renderer/Private/HairStrands/HairStrandsVoxelization.cpp
+++ b/Engine/Source/Runtime/Renderer/Private/HairStrands/HairStrandsVoxelization.cpp
@@ -871,7 +871,7 @@ static FHairStrandsVoxelResources AllocateVirtualVoxelResources(
 	Out.Parameters.Common.NodeDescCount							= MacroGroupDatas.Num();
 	Out.Parameters.Common.Raytracing_ShadowOcclusionThreshold	= FMath::Max(0.f, GHairVirtualVoxelRaytracing_ShadowOcclusionThreshold);
 	Out.Parameters.Common.Raytracing_SkyOcclusionThreshold		= FMath::Max(0.f, GHairVirtualVoxelRaytracing_SkyOcclusionThreshold);
-	Out.Parameters.Common.HairCoveragePixelRadiusAtDepth1		= ComputeMinStrandRadiusAtDepth1(FIntPoint(View.ViewRect.Width(), View.ViewRect.Height()), View.FOV, 1/*SampleCount*/, 1/*RasterizationScale*/, View.ViewMatrices.GetOrthoDimensions().X).Primary;
+	Out.Parameters.Common.HairCoveragePixelRadiusAtDepth1		= ComputeMinStrandRadiusAtDepth1(FIntPoint(View.ViewRect.Width(), View.ViewRect.Height()), View.FOV, 1/*SampleCount*/, 1/*RasterizationScale*/, View.ViewMatrices.GetOrthoDimensions().X, View.Family->FinalResolutionScale).Primary;
 	Out.Parameters.Common.TranslatedWorldOffset					= FVector3f(View.ViewMatrices.GetPreViewTranslation());
 	Out.Parameters.Common.TranslatedWorldOffsetStereoCorrection = FVector3f(PreViewStereoCorrection);
 
@@ -981,7 +981,7 @@ static FHairStrandsVoxelResources AllocateDummyVirtualVoxelResources(
 	Out.Parameters.Common.Raytracing_SkyOcclusionThreshold		= 1;
 
 	Out.Parameters.Common.TranslatedWorldOffset = FVector3f(View.ViewMatrices.GetPreViewTranslation());
-	Out.Parameters.Common.HairCoveragePixelRadiusAtDepth1	= ComputeMinStrandRadiusAtDepth1(FIntPoint(View.ViewRect.Width(), View.ViewRect.Height()), View.FOV, 1/*SampleCount*/, 1/*RasterizationScale*/, View.ViewMatrices.GetOrthoDimensions().X).Primary;
+	Out.Parameters.Common.HairCoveragePixelRadiusAtDepth1	= ComputeMinStrandRadiusAtDepth1(FIntPoint(View.ViewRect.Width(), View.ViewRect.Height()), View.FOV, 1/*SampleCount*/, 1/*RasterizationScale*/, View.ViewMatrices.GetOrthoDimensions().X, View.Family->FinalResolutionScale).Primary;
 	Out.Parameters.Common.PageIndexCount = 0;
 
 	FRDGBufferRef Dummy2BytesBuffer = GSystemTextures.GetDefaultBuffer(GraphBuilder, 2, 0u);
