Deprecate SSGI (Screen Space Global Illumination)

SSGI is superseded by [REDACTED_PERSON]. This change:
- Renames the editor label from "Beta" to "Deprecated"
- Emits a one-time LogRenderer warning when SSGI is active
- Marks CVars as deprecated (except for the Quality cvar as it's being used in our default scalability settings and we don't want to trigger warnings for projects not using SSGI)

#Synced-CL 51621159

diff --git a/Engine/Source/Runtime/Engine/Classes/Engine/EngineTypes.h b/Engine/Source/Runtime/Engine/Classes/Engine/EngineTypes.h
--- a/Engine/Source/Runtime/Engine/Classes/Engine/EngineTypes.h
+++ b/Engine/Source/Runtime/Engine/Classes/Engine/EngineTypes.h
@@ -460,7 +460,7 @@ namespace EDynamicGlobalIlluminationMethod
 		Lumen,
 
 		/** Standalone Screen Space Global Illumination.  Low cost, but limited by screen space information. */
-		ScreenSpace UMETA(DisplayName="Screen Space (Beta)"),
+		ScreenSpace UMETA(DisplayName="Screen Space (Deprecated)"),
 
 		/** Use a plugin for Global Illumination */
 		Plugin UMETA(DisplayName="Plugin"),
diff --git a/Engine/Source/Runtime/Renderer/Private/IndirectLightRendering.cpp b/Engine/Source/Runtime/Renderer/Private/IndirectLightRendering.cpp
--- a/Engine/Source/Runtime/Renderer/Private/IndirectLightRendering.cpp
+++ b/Engine/Source/Runtime/Renderer/Private/IndirectLightRendering.cpp
@@ -536,6 +536,13 @@ void FDeferredShadingSceneRenderer::CommitIndirectLightingState()
 			{
 				DiffuseIndirectMethod = EDiffuseIndirectMethod::SSGI;
 				DiffuseIndirectDenoiser = IScreenSpaceDenoiser::GetDenoiserMode(CVarDiffuseIndirectDenoiser);
+
+				static bool bHasWarnedSSGIDeprecated = false;
+				if (!bHasWarnedSSGIDeprecated)
+				{
+					bHasWarnedSSGIDeprecated = true;
+					UE_LOG(LogRenderer, Warning, TEXT("Screen Space Global Illumination (SSGI) is deprecated (from 5.8) and will be removed in a future release. Please migrate to Lumen Global Illumination."));
+				}
 			}
 			else if (ShouldRenderPluginGlobalIllumination(View))
 			{
diff --git a/Engine/Source/Runtime/Renderer/Private/ScreenSpaceRayTracing.cpp b/Engine/Source/Runtime/Renderer/Private/ScreenSpaceRayTracing.cpp
--- a/Engine/Source/Runtime/Renderer/Private/ScreenSpaceRayTracing.cpp
+++ b/Engine/Source/Runtime/Renderer/Private/ScreenSpaceRayTracing.cpp
@@ -53,33 +53,38 @@ static TAutoConsoleVariable<int32> CVarSSRCompute(
 	ECVF_RenderThreadSafe);
 
 static TAutoConsoleVariable<int32> CVarSSGILeakFreeReprojection(
+	FCVarDeprecationArgs{ TEXT("5.8") },
 	TEXT("r.SSGI.LeakFreeReprojection"), 1,
-	TEXT("Whether use a more expensive but leak free reprojection of previous frame's scene color.\n"),
+	TEXT("[DEPRECATED] Whether use a more expensive but leak free reprojection of previous frame's scene color.\n"),
 	ECVF_RenderThreadSafe);
 
 static TAutoConsoleVariable<int32> CVarSSGIQuality(
 	TEXT("r.SSGI.Quality"), 4,
-	TEXT("Quality setting to control number of ray shot with SSGI, between 1 and 4 (defaults to 4).\n"),
+	TEXT("[DEPRECATED] Quality setting to control number of ray shot with SSGI, between 1 and 4 (defaults to 4).\n"),
 	ECVF_Scalability | ECVF_RenderThreadSafe);
 
 static TAutoConsoleVariable<float> CVarSSGIMinimumLuminance(
+	FCVarDeprecationArgs{ TEXT("5.8") },
 	TEXT("r.SSGI.MinimumLuminance"), 0.5f,
-	TEXT(""),
+	TEXT("[DEPRECATED] "),
 	ECVF_Scalability | ECVF_RenderThreadSafe);
 
 static TAutoConsoleVariable<int32> CVarSSGIRejectUncertainRays(
+	FCVarDeprecationArgs{ TEXT("5.8") },
 	TEXT("r.SSGI.RejectUncertainRays"), 1,
-	TEXT("Rejects the screen space ray if it was uncertain due to going behind screen geometry."),
+	TEXT("[DEPRECATED] Rejects the screen space ray if it was uncertain due to going behind screen geometry."),
 	ECVF_Scalability | ECVF_RenderThreadSafe);
 
 static TAutoConsoleVariable<int32> CVarSSGITerminateCertainRay(
+	FCVarDeprecationArgs{ TEXT("5.8") },
 	TEXT("r.SSGI.TerminateCertainRay"), 1,
-	TEXT("Optimisations that if the screen space ray is certain and didn't find any geometry, don't fallback on otehr tracing technic."),
+	TEXT("[DEPRECATED] Optimisations that if the screen space ray is certain and didn't find any geometry, don't fallback on otehr tracing technic."),
 	ECVF_Scalability | ECVF_RenderThreadSafe);
 
 static TAutoConsoleVariable<float> CVarSSGISkyDistance(
+	FCVarDeprecationArgs{ TEXT("5.8") },
 	TEXT("r.SSGI.SkyDistance"), 10000000,
-	TEXT("Distance of the sky in KM."),
+	TEXT("[DEPRECATED] Distance of the sky in KM."),
 	ECVF_Scalability | ECVF_RenderThreadSafe);
 
 
