ISM Culling - Include PackedData when testing for a matching ISM batch

diff --git a/Engine/Source/Runtime/Renderer/Private/InstanceCulling/InstanceCullingContext.cpp b/Engine/Source/Runtime/Renderer/Private/InstanceCulling/InstanceCullingContext.cpp
--- a/Engine/Source/Runtime/Renderer/Private/InstanceCulling/InstanceCullingContext.cpp
+++ b/Engine/Source/Runtime/Renderer/Private/InstanceCulling/InstanceCullingContext.cpp
@@ -1496,6 +1496,7 @@ void FInstanceCullingContext::SetupDrawCommands(
 
 	int32 CurrentStateBucketId = -1;
 	EMeshDrawCommandCullingPayloadFlags CurrentCullingPayloadFlags = EMeshDrawCommandCullingPayloadFlags::Default;
+	FMeshDrawCommandCullingPayload CurrentCullingPayload;
 	MaxInstances = 1;
 	// Only used to supply stats
 	uint32 CurrentAutoInstanceCount = 1;
@@ -1536,7 +1537,11 @@ void FInstanceCullingContext::SetupDrawCommands(
 		const bool bChangedFromDirectToIndirect = !bPrevUseIndirectDraw && bUseIndirectDraw;
 		bPrevUseIndirectDraw = bUseIndirectDraw;
 
-		if (bCompactIdenticalCommands && CurrentStateBucketId != -1 && VisibleMeshDrawCommand.StateBucketId == CurrentStateBucketId && VisibleMeshDrawCommand.CullingPayloadFlags == CurrentCullingPayloadFlags && !bChangedFromDirectToIndirect)
+		bool bMeshDrawCommandMatch = VisibleMeshDrawCommand.StateBucketId == CurrentStateBucketId
+									&& VisibleMeshDrawCommand.CullingPayloadFlags == CurrentCullingPayloadFlags
+									&& VisibleMeshDrawCommand.CullingPayload.PackedData == CurrentCullingPayload.PackedData;
+
+		if (bCompactIdenticalCommands && CurrentStateBucketId != -1 && bMeshDrawCommandMatch && !bChangedFromDirectToIndirect)
 		{
 			// Drop since previous covers for this
 
@@ -1599,6 +1604,7 @@ void FInstanceCullingContext::SetupDrawCommands(
 			// Record the last bucket ID (may be -1)
 			CurrentStateBucketId = VisibleMeshDrawCommand.StateBucketId;
 			CurrentCullingPayloadFlags = VisibleMeshDrawCommand.CullingPayloadFlags;
+			CurrentCullingPayload = VisibleMeshDrawCommand.CullingPayload;
 
 			// If we have dropped any we need to move up to maintain 1:1
 			if (DrawCommandIndex > NumDrawCommandsOut)
