Fix multithreading memory ordering issue in the MeshProcessor

#Synced-CL 51700160

diff --git a/Engine/Source/Runtime/Renderer/Private/MeshPassProcessor.cpp b/Engine/Source/Runtime/Renderer/Private/MeshPassProcessor.cpp
--- a/Engine/Source/Runtime/Renderer/Private/MeshPassProcessor.cpp
+++ b/Engine/Source/Runtime/Renderer/Private/MeshPassProcessor.cpp
@@ -857,7 +857,8 @@ bool FMinimalBoundShaderStateInput::AllowSkipUnloadedShaders() const
 
 void FGraphicsMinimalPipelineStateInitializer::SetupBoundShaderState(FRHIVertexDeclaration* VertexDeclaration, const FMeshProcessorShaders& Shaders)
 {
-	BoundShaderState = FMinimalBoundShaderStateInput();
+	BoundShaderState.~FMinimalBoundShaderStateInput();
+	new (&BoundShaderState) FMinimalBoundShaderStateInput();
 	BoundShaderState.VertexDeclarationRHI = VertexDeclaration;
 
 	checkf(Shaders.VertexShader.IsValid(), TEXT("Can't render without a vertex shader"));
@@ -2229,8 +2230,7 @@ void SetupGBufferRenderTargetInfo(const FSceneTexturesConfig& SceneTexturesConfi
 
 FGraphicsMinimalPipelineStateInitializer PSOCollectorStats::GetShadersOnlyInitializer(const FGraphicsMinimalPipelineStateInitializer& Initializer)
 {
-	FGraphicsMinimalPipelineStateInitializer ShadersOnlyInitializer;
-	ShadersOnlyInitializer.BoundShaderState = Initializer.BoundShaderState;
+	FGraphicsMinimalPipelineStateInitializer ShadersOnlyInitializer(Initializer);
 	ShadersOnlyInitializer.BoundShaderState.VertexDeclarationRHI = nullptr;
 	ShadersOnlyInitializer.ComputeStatePrecachePSOHash();
 
diff --git a/Engine/Source/Runtime/Renderer/Public/MeshPassProcessor.h b/Engine/Source/Runtime/Renderer/Public/MeshPassProcessor.h
--- a/Engine/Source/Runtime/Renderer/Public/MeshPassProcessor.h
+++ b/Engine/Source/Runtime/Renderer/Public/MeshPassProcessor.h
@@ -202,6 +202,35 @@ struct FMinimalBoundShaderStateInput
 {
 	inline FMinimalBoundShaderStateInput() {}
 
+	FMinimalBoundShaderStateInput(const FMinimalBoundShaderStateInput & Other)
+		: VertexDeclarationRHI(Other.VertexDeclarationRHI)
+		, CachedVertexShader(Other.CachedVertexShader.load(std::memory_order_relaxed))
+		, CachedPixelShader(Other.CachedPixelShader)
+		, bAllShaderAreLoaded(Other.bAllShaderAreLoaded)
+#if PLATFORM_SUPPORTS_GEOMETRY_SHADERS
+		, CachedGeometryShader(Other.CachedGeometryShader)
+#endif
+#if PLATFORM_SUPPORTS_MESH_SHADERS
+		, CachedMeshShader(Other.CachedMeshShader)
+#endif
+		, VertexShaderResource(Other.VertexShaderResource)
+		, PixelShaderResource(Other.PixelShaderResource)
+#if PLATFORM_SUPPORTS_GEOMETRY_SHADERS
+		, GeometryShaderResource(Other.GeometryShaderResource)
+#endif
+#if PLATFORM_SUPPORTS_MESH_SHADERS
+		, MeshShaderResource(Other.MeshShaderResource)
+#endif
+		, VertexShaderIndex(Other.VertexShaderIndex)
+		, PixelShaderIndex(Other.PixelShaderIndex)
+#if PLATFORM_SUPPORTS_GEOMETRY_SHADERS
+		, GeometryShaderIndex(Other.GeometryShaderIndex)
+#endif
+#if PLATFORM_SUPPORTS_MESH_SHADERS
+		, MeshShaderIndex(Other.MeshShaderIndex)
+#endif
+	{}
+
 	FBoundShaderStateInput AsBoundShaderState() const
 	{
 		bool bLocalAllShaderAreLoaded = true;
@@ -215,7 +244,8 @@ struct FMinimalBoundShaderStateInput
 			return Shader;
 		};
 
-		if (!CachedVertexShader || !bAllShaderAreLoaded)
+		FRHIVertexShader* LocalCachedVertexShader = CachedVertexShader.load(std::memory_order_acquire);   
+		if (!LocalCachedVertexShader || !bAllShaderAreLoaded)
 		{
 			CachedPixelShader = PixelShaderResource ? static_cast<FRHIPixelShader*>(GetShaderResource(PixelShaderResource->GetShader(PixelShaderIndex, !bCanSkipShader))) : nullptr;
 #if PLATFORM_SUPPORTS_GEOMETRY_SHADERS
@@ -224,9 +254,10 @@ struct FMinimalBoundShaderStateInput
 #if PLATFORM_SUPPORTS_MESH_SHADERS
 			CachedMeshShader = MeshShaderResource ? static_cast<FRHIMeshShader*>(MeshShaderResource->GetShader(MeshShaderIndex)) : nullptr;
 #endif
-			CachedVertexShader = VertexShaderResource ? static_cast<FRHIVertexShader*>(GetShaderResource(VertexShaderResource->GetShader(VertexShaderIndex, !bCanSkipShader))) : nullptr;
+			LocalCachedVertexShader = VertexShaderResource ? static_cast<FRHIVertexShader*>(GetShaderResource(VertexShaderResource->GetShader(VertexShaderIndex, !bCanSkipShader))) : nullptr;
 
 			bAllShaderAreLoaded = bLocalAllShaderAreLoaded;
+			CachedVertexShader.store(LocalCachedVertexShader, std::memory_order_release);
 		}
 
 #if PLATFORM_SUPPORTS_MESH_SHADERS
@@ -238,7 +269,7 @@ struct FMinimalBoundShaderStateInput
 #endif
 		{
 			return FBoundShaderStateInput(VertexDeclarationRHI
-				, CachedVertexShader
+				, LocalCachedVertexShader
 				, CachedPixelShader
 #if PLATFORM_SUPPORTS_GEOMETRY_SHADERS
 				, CachedGeometryShader
@@ -253,7 +284,8 @@ struct FMinimalBoundShaderStateInput
 
 	void ForceShaderReload() const 
 	{
-		if (!CachedVertexShader || !bAllShaderAreLoaded)
+		FRHIVertexShader* LocalCachedVertexShader = CachedVertexShader.load(std::memory_order_acquire);
+		if (!LocalCachedVertexShader || !bAllShaderAreLoaded)
 		{
 			CachedPixelShader = PixelShaderResource ? static_cast<FRHIPixelShader*>(PixelShaderResource->GetShader(PixelShaderIndex)) : nullptr;
 #if PLATFORM_SUPPORTS_GEOMETRY_SHADERS
@@ -262,9 +294,11 @@ struct FMinimalBoundShaderStateInput
 #if PLATFORM_SUPPORTS_MESH_SHADERS
 			CachedMeshShader = MeshShaderResource ? static_cast<FRHIMeshShader*>(MeshShaderResource->GetShader(MeshShaderIndex)) : nullptr;
 #endif
-			CachedVertexShader = VertexShaderResource ? static_cast<FRHIVertexShader*>(VertexShaderResource->GetShader(VertexShaderIndex)) : nullptr;
+			LocalCachedVertexShader = VertexShaderResource ? static_cast<FRHIVertexShader*>(VertexShaderResource->GetShader(VertexShaderIndex)) : nullptr;
 
 			bAllShaderAreLoaded = true;
+
+			CachedVertexShader.store(LocalCachedVertexShader, std::memory_order_release);
 		}
 	}
 
@@ -299,7 +333,7 @@ struct FMinimalBoundShaderStateInput
 	}
 
 	FRHIVertexDeclaration* VertexDeclarationRHI = nullptr;
-	mutable FRHIVertexShader* CachedVertexShader = nullptr;
+	mutable std::atomic<FRHIVertexShader*> CachedVertexShader = nullptr;
 	mutable FRHIPixelShader* CachedPixelShader = nullptr;
 	mutable bool bAllShaderAreLoaded = true;
 #if PLATFORM_SUPPORTS_GEOMETRY_SHADERS
