Fix crash when using at the same time ShaderPrint and RDG clobber resources.

The ShaderPrint's ShaderPrintEntryBuffer buffer didn't have any UAV allocated (its content was copied), which prevent the clobber system to infer the UAV type in order to clear the buffer. Added a clear pass on that buffer, fixed the issue.

#Synced-CL 51755560

diff --git a/Engine/Source/Runtime/Renderer/Private/ShaderPrint.cpp b/Engine/Source/Runtime/Renderer/Private/ShaderPrint.cpp
--- a/Engine/Source/Runtime/Renderer/Private/ShaderPrint.cpp
+++ b/Engine/Source/Runtime/Renderer/Private/ShaderPrint.cpp
@@ -1200,6 +1200,7 @@ namespace ShaderPrint
 		auto CreateTypedBufferFromStructured = [](FRDGBuilder& GraphBuilder, FRDGBufferRef Input, const TCHAR* ResultBufferName) -> FRDGBufferRef
 		{
 			FRDGBufferRef Result = GraphBuilder.CreateBuffer(FRDGBufferDesc::CreateBufferDesc(Input->Desc.BytesPerElement, Input->Desc.NumElements), ResultBufferName);
+			AddClearUAVPass(GraphBuilder, GraphBuilder.CreateUAV(Result, PF_R32_UINT), 0u);
 			AddCopyBufferPass(GraphBuilder, Result, Input);
 			return Result;
 		};
