Runtime reflection capture crash fixes:  fix assert when running shadow only timeslice rendering, caused by creating decal visibility task data, but Finish never being called due to decal rendering being skipped later.  Fix for issue where USceneCaptureComponent::GetViewState can occasionally be initially reached with a non-zero ViewIndex due to timeslicing, leading to an assert in ShareOrigin where the first element attempts to point to itself.  Need to check the number of elements in the array, not the index being allocated towards.

#rnx

#Synced-CL 51746839

diff --git a/Engine/Source/Runtime/Engine/Private/Components/SceneCaptureComponent.cpp b/Engine/Source/Runtime/Engine/Private/Components/SceneCaptureComponent.cpp
--- a/Engine/Source/Runtime/Engine/Private/Components/SceneCaptureComponent.cpp
+++ b/Engine/Source/Runtime/Engine/Private/Components/SceneCaptureComponent.cpp
@@ -411,7 +411,7 @@ FSceneViewStateInterface* USceneCaptureComponent::GetViewState(int32 ViewIndex)
 		ViewStates.Add(new FSceneViewStateReference());
 
 		// Cube map view states can share an origin, saving memory and performance
-		if ((ViewIndex > 0) && IsCube())
+		if ((ViewStates.Num() > 1) && IsCube())
 		{
 			ViewStates.Last().ShareOrigin(&ViewStates[0]);
 		}
diff --git a/Engine/Source/Runtime/Renderer/Private/MobileShadingRenderer.cpp b/Engine/Source/Runtime/Renderer/Private/MobileShadingRenderer.cpp
--- a/Engine/Source/Runtime/Renderer/Private/MobileShadingRenderer.cpp
+++ b/Engine/Source/Runtime/Renderer/Private/MobileShadingRenderer.cpp
@@ -783,7 +783,10 @@ void FMobileSceneRenderer::InitViews(
 		UpdateDirectionalLightUniformBuffers(GraphBuilder, Views[0]);
 	}
 
-	DecalVisibilityTaskData = FDecalVisibilityTaskData::Launch(GraphBuilder, *Scene, AllViews);
+	if (!bShadowOnlyTimeslice)
+	{
+		DecalVisibilityTaskData = FDecalVisibilityTaskData::Launch(GraphBuilder, *Scene, AllViews);
+	}
 }
 
 static void BeginOcclusionScope(FRDGBuilder& GraphBuilder, TArray<FViewInfo>& Views)
