Path Tracing: Avoid crash when enabling MeshEdges flag

The scene capture should have path tracing disabled for the wireframe render to complete succesefully.

#Synced-CL 51768242

diff --git a/Engine/Source/Runtime/Renderer/Private/MeshEdgesRendering.cpp b/Engine/Source/Runtime/Renderer/Private/MeshEdgesRendering.cpp
--- a/Engine/Source/Runtime/Renderer/Private/MeshEdgesRendering.cpp
+++ b/Engine/Source/Runtime/Renderer/Private/MeshEdgesRendering.cpp
@@ -328,6 +328,7 @@ void RenderMeshEdges(FSceneViewFamily& ViewFamily)
 		WireframeShowFlags.SetLensDistortion(false);
 		//WireframeShowFlags.SetScreenPercentage(false);
 		//WireframeShowFlags.SetTranslucency(false);
+        WireframeShowFlags.SetPathTracing(false);
 	}
 	
 	FSceneViewFamilyContext CaptureViewFamily(
