Added detailed Insights trace to FRuntimeVirtualTextureFinalizer::RenderFinalize in order to better assess RVT perf (useful in case of large page invalidations in particular)

Notes :
* The priority field is stored as a string, since enums are not currently supported in Insights
* The NumTilesPerLevel field is also stored as a string, since Insights doesn't yet have an automatic way to display int32[] fields
In practice, there's a very limited amount of tags per frame (usually, just 1) so the cost is minimal

#Synced-CL 51768384

diff --git a/Engine/Source/Runtime/Renderer/Private/VT/RuntimeVirtualTextureProducer.cpp b/Engine/Source/Runtime/Renderer/Private/VT/RuntimeVirtualTextureProducer.cpp
--- a/Engine/Source/Runtime/Renderer/Private/VT/RuntimeVirtualTextureProducer.cpp
+++ b/Engine/Source/Runtime/Renderer/Private/VT/RuntimeVirtualTextureProducer.cpp
@@ -6,6 +6,7 @@
 #include "VT/RuntimeVirtualTextureRender.h"
 #include "VT/RuntimeVirtualTextureSceneProxy.h"
 #include "RenderGraphBuilder.h"
+#include "UObject/Class.h"
 
 FRuntimeVirtualTextureFinalizer::FRuntimeVirtualTextureFinalizer(
 	FVTProducerDescription const& InDesc, 
@@ -37,6 +38,29 @@ void FRuntimeVirtualTextureFinalizer::AddTile(FTileEntry& Tile)
 	Tiles.Add(Tile);
 }
 
+UE_TRACE_EVENT_BEGIN(Cpu, RuntimeVirtualTextureFinalizerRenderFinalize, NoSync)
+UE_TRACE_EVENT_FIELD(UE::Trace::WideString, Name)
+UE_TRACE_EVENT_FIELD(UE::Trace::WideString, Priority) // Note : these strings are actually static (display names of EVTProducerPriority), so we could avoid using a string here, but Insights doesn't have support for FNames or enums yet
+UE_TRACE_EVENT_FIELD(int32, NumTiles)
+UE_TRACE_EVENT_FIELD(int32, NumBatches)
+UE_TRACE_EVENT_FIELD(UE::Trace::WideString, NumTilesPerLevel) // Note : this could be a int32[] but Insights doesn't support displaying that in a generic way so we have to store it as a string for now
+UE_TRACE_EVENT_END()
+
+namespace RuntimeVirtualTexture::Private
+{
+
+TStaticArray<FString, static_cast<int32>(EVTProducerPriority::Count)> GetProducerPriorityDisplayStrings()
+{
+	TStaticArray<FString, static_cast<int32>(EVTProducerPriority::Count)> EnumStrings;
+	for (uint8 EnumValue = 0; EnumValue < static_cast<uint8>(EVTProducerPriority::Count); ++EnumValue)
+	{
+		EnumStrings[EnumValue] = UEnum::GetDisplayValueAsText(static_cast<EVTProducerPriority>(EnumValue)).ToString();
+	}
+	return EnumStrings;
+}
+
+} // namespace RuntimeVirtualTexture::Private
+
 void FRuntimeVirtualTextureFinalizer::RenderFinalize(FRDGBuilder& GraphBuilder, ISceneRenderer* SceneRenderer)
 {
 	RDG_GPU_MASK_SCOPE(GraphBuilder, FRHIGPUMask::All());
@@ -53,6 +77,11 @@ void FRuntimeVirtualTextureFinalizer::RenderFinalize(FRDGBuilder& GraphBuilder,
 	RenderPageBatchDesc.FixedColor = FLinearColor::Transparent;
 	RenderPageBatchDesc.CustomMaterialData = CustomMaterialData;
 	
+#if CPUPROFILERTRACE_ENABLED
+	TArray<int32, TInlineAllocator<16>> NumTilesPerLevel;
+	NumTilesPerLevel.AddDefaulted(16);
+#endif // CPUPROFILERTRACE_ENABLED
+
 	for (int LayerIndex = 0; LayerIndex < RuntimeVirtualTexture::MaxTextureLayers; ++LayerIndex)
 	{
 		RenderPageBatchDesc.Targets[LayerIndex].PooledRenderTarget = Tiles[0].Targets[LayerIndex].PooledRenderTarget;
@@ -98,6 +127,14 @@ void FRuntimeVirtualTextureFinalizer::RenderFinalize(FRDGBuilder& GraphBuilder,
 		RenderPageDesc.vLevel = Entry.vLevel;
 		RenderPageDesc.UVRange = UVRange;
 
+#if CPUPROFILERTRACE_ENABLED
+		if (Entry.vLevel >= NumTilesPerLevel.Num())
+		{
+			NumTilesPerLevel.AddDefaulted(Entry.vLevel - NumTilesPerLevel.Num() + 1);
+		}
+		++NumTilesPerLevel[Entry.vLevel];
+#endif // CPUPROFILERTRACE_ENABLED
+
 		const int32 TileSize = Desc.TileSize + 2 * Desc.TileBorderSize;
 		for (int LayerIndex = 0; LayerIndex < RuntimeVirtualTexture::MaxTextureLayers; ++LayerIndex)
 		{
@@ -112,6 +149,46 @@ void FRuntimeVirtualTextureFinalizer::RenderFinalize(FRDGBuilder& GraphBuilder,
 		Batches.Add(RuntimeVirtualTexture::InitPageBatch(GraphBuilder, RenderPageBatchDesc));
 	}
 
+#if CPUPROFILERTRACE_ENABLED
+	auto GetNumTilesPerLevelToString = [&]() -> FString
+		{
+			int32 MinMipLevel = NumTilesPerLevel.IndexOfByPredicate([](int32 InNumTiles) { return (InNumTiles > 0); });
+			check(MinMipLevel != INDEX_NONE);
+			int32 MaxMipLevel = NumTilesPerLevel.FindLastByPredicate([](int32 InNumTiles) { return (InNumTiles > 0); });
+			check(MaxMipLevel != INDEX_NONE);
+			// Trim the last zeroes to reduce the size of the trace string
+			if (MaxMipLevel + 1 < NumTilesPerLevel.Num())
+			{
+				NumTilesPerLevel.RemoveAt(MaxMipLevel + 1, NumTilesPerLevel.Num() - MaxMipLevel - 1);
+			}
+			// Trim the first zeroes to reduce the size of the trace string even more
+			if (MinMipLevel > 0)
+			{
+				NumTilesPerLevel.RemoveAt(0, MinMipLevel);
+			}
+
+			FString NumTilesPerLevelString = FString::Printf(TEXT("Mips[%i-%i] "), MinMipLevel, MaxMipLevel);
+
+			for (int32 NumTiles : NumTilesPerLevel)
+			{
+				NumTilesPerLevelString += FString::Printf(TEXT("%i,"), NumTiles);
+			}
+			// Remove the trailing comma
+			NumTilesPerLevelString.LeftChopInline(1);
+			return NumTilesPerLevelString;
+		};
+
+	// static store for fast retrieval : 
+	static TStaticArray<FString, static_cast<int32>(EVTProducerPriority::Count)> EnumStrings = RuntimeVirtualTexture::Private::GetProducerPriorityDisplayStrings();
+
+	UE_TRACE_LOG_SCOPED_T(Cpu, RuntimeVirtualTextureFinalizerRenderFinalize, CpuChannel)
+		<< RuntimeVirtualTextureFinalizerRenderFinalize.Name(*Desc.Name.ToString())
+		<< RuntimeVirtualTextureFinalizerRenderFinalize.Priority(*EnumStrings[static_cast<int32>(Desc.Priority)])
+		<< RuntimeVirtualTextureFinalizerRenderFinalize.NumTiles(Tiles.Num())
+		<< RuntimeVirtualTextureFinalizerRenderFinalize.NumBatches(Batches.Num())
+		<< RuntimeVirtualTextureFinalizerRenderFinalize.NumTilesPerLevel(*GetNumTilesPerLevelToString());
+#endif // CPUPROFILERTRACE_ENABLED
+
 	for (RuntimeVirtualTexture::FBatchRenderContext const* Batch : Batches)
 	{
 		RuntimeVirtualTexture::RenderPageBatch(GraphBuilder, *Batch);
