Deprecate unused includes from public rendering headers.

#Synced-CL 51779376

diff --git a/Engine/Source/Runtime/RenderCore/Public/ProfilingDebugging/RealtimeGPUProfiler.h b/Engine/Source/Runtime/RenderCore/Public/ProfilingDebugging/RealtimeGPUProfiler.h
--- a/Engine/Source/Runtime/RenderCore/Public/ProfilingDebugging/RealtimeGPUProfiler.h
+++ b/Engine/Source/Runtime/RenderCore/Public/ProfilingDebugging/RealtimeGPUProfiler.h
@@ -10,16 +10,22 @@
 #include "CoreTypes.h"
 #include "GpuProfilerTrace.h"
 #include "HAL/CriticalSection.h"
-#include "MultiGPU.h"
-#include "ProfilingDebugging/CsvProfiler.h"
-#include "ProfilingDebugging/CsvProfilerConfig.h"
+
 #include "RHI.h"
+#include "RHIDefinitions.h"
 #include "RHICommandList.h"
 #include "RHIBreadcrumbs.h"
 #include "RenderingThread.h"
+#include "Misc/Build.h"
+
+#if UE_ENABLE_INCLUDE_ORDER_DEPRECATED_IN_5_8
+#include "MultiGPU.h"
+#include "ProfilingDebugging/CsvProfiler.h"
+#include "ProfilingDebugging/CsvProfilerConfig.h"
 #include "Stats/Stats.h"
 #include "UObject/NameTypes.h"
 #include <tuple>
+#endif // UE_ENABLE_INCLUDE_ORDER_DEPRECATED_IN_5_8
 
 #define WANTS_DRAW_MESH_EVENTS (WITH_PROFILEGPU && WITH_RHI_BREADCRUMBS)
 
diff --git a/Engine/Source/Runtime/RenderCore/Public/RenderGraphBuilder.h b/Engine/Source/Runtime/RenderCore/Public/RenderGraphBuilder.h
--- a/Engine/Source/Runtime/RenderCore/Public/RenderGraphBuilder.h
+++ b/Engine/Source/Runtime/RenderCore/Public/RenderGraphBuilder.h
@@ -13,7 +13,6 @@
 #include "Misc/AssertionMacros.h"
 #include "MultiGPU.h"
 #include "PixelFormat.h"
-#include "ProfilingDebugging/CsvProfiler.h"
 #include "RHIBreadcrumbs.h"
 #include "RHIDefinitions.h"
 #include "RenderGraphAllocator.h"
@@ -26,12 +25,16 @@
 #include "RenderGraphValidation.h"
 #include "RendererInterface.h"
 #include "ShaderParameterMacros.h"
-#include "Stats/Stats.h"
 #include "Templates/RefCounting.h"
 #include "Templates/UnrealTemplate.h"
 #include "Tasks/Pipe.h"
 #include "Experimental/Containers/RobinHoodHashTable.h"
 
+#if UE_ENABLE_INCLUDE_ORDER_DEPRECATED_IN_5_8
+#include "ProfilingDebugging/CsvProfiler.h"
+#include "Stats/Stats.h"
+#endif //UE_ENABLE_INCLUDE_ORDER_DEPRECATED_IN_5_8
+
 enum class ERenderTargetTexture : uint8;
 struct FParallelPassSet;
 struct FRHIRenderPassInfo;
diff --git a/Engine/Source/Runtime/RenderCore/Public/RenderGraphEvent.h b/Engine/Source/Runtime/RenderCore/Public/RenderGraphEvent.h
--- a/Engine/Source/Runtime/RenderCore/Public/RenderGraphEvent.h
+++ b/Engine/Source/Runtime/RenderCore/Public/RenderGraphEvent.h
@@ -5,23 +5,29 @@
 #include "Containers/Array.h"
 #include "Containers/StaticArray.h"
 #include "Containers/UnrealString.h"
-#include "DynamicRenderScaling.h"
 #include "HAL/Platform.h"
 #include "HAL/PlatformCrt.h"
 #include "Misc/AssertionMacros.h"
 #include "MultiGPU.h"
 #include "ProfilingDebugging/CsvProfiler.h"
+#include "GpuProfilerTrace.h"
+#include "RHIFwd.h"
+#include "RHICommandList.h"
+#include "RenderGraphAllocator.h"
+#include "RenderGraphDefinitions.h"
+#include "Templates/UnrealTemplate.h"
+#include "Misc/TVariant.h"
+#include "UObject/NameTypes.h"
+
+#if UE_ENABLE_INCLUDE_ORDER_DEPRECATED_IN_5_8
+#include "DynamicRenderScaling.h"
 #include "ProfilingDebugging/CsvProfilerConfig.h"
 #include "ProfilingDebugging/CpuProfilerTrace.h"
 #include "ProfilingDebugging/RealtimeGPUProfiler.h"
 #include "RHI.h"
-#include "RHICommandList.h"
-#include "RenderGraphAllocator.h"
-#include "RenderGraphDefinitions.h"
 #include "RendererInterface.h"
 #include "Stats/Stats.h"
-#include "Templates/UnrealTemplate.h"
-#include "UObject/NameTypes.h"
+#endif //UE_ENABLE_INCLUDE_ORDER_DEPRECATED_IN_5_8
 
 //////////////////////////////////////////////////////////////////////////
 //
@@ -31,6 +37,11 @@
 
 class FRDGScopeState;
 
+namespace DynamicRenderScaling
+{
+	class FBudget;
+}
+
 /** Stores a GPU event name for the render graph. Draw events can be compiled out entirely from
  *  a release build for performance.
  */
diff --git a/Engine/Source/Runtime/RenderCore/Public/RenderGraphResources.h b/Engine/Source/Runtime/RenderCore/Public/RenderGraphResources.h
--- a/Engine/Source/Runtime/RenderCore/Public/RenderGraphResources.h
+++ b/Engine/Source/Runtime/RenderCore/Public/RenderGraphResources.h
@@ -16,13 +16,17 @@
 #include "RenderGraphDefinitions.h"
 #include "RenderGraphParameter.h"
 #include "RenderGraphTextureSubresource.h"
-#include "RendererInterface.h"
+#include "PooledRenderTarget.h"
 #include "ShaderParameterMacros.h"
 #include "Templates/Function.h"
 #include "Templates/RefCounting.h"
 #include "Templates/TypeHash.h"
 #include "Templates/UnrealTemplate.h"
 
+#if UE_ENABLE_INCLUDE_ORDER_DEPRECATED_IN_5_8
+#include "RendererInterface.h"
+#endif //UE_ENABLE_INCLUDE_ORDER_DEPRECATED_IN_5_8
+
 class FRDGBarrierBatchBegin;
 class FRDGBarrierValidation;
 class FRDGBuffer;
diff --git a/Engine/Source/Runtime/RenderCore/Public/RendererInterface.h b/Engine/Source/Runtime/RenderCore/Public/RendererInterface.h
--- a/Engine/Source/Runtime/RenderCore/Public/RendererInterface.h
+++ b/Engine/Source/Runtime/RenderCore/Public/RendererInterface.h
@@ -10,17 +10,24 @@
 #include "Templates/RefCounting.h"
 #include "Misc/MemStack.h"
 #include "Modules/ModuleInterface.h"
+#include "Misc/EnumClassFlags.h"
+#include "RenderGraphFwd.h"
+#include "RHIResources.h"
+#include "RHIFeatureLevel.h"
+#include "VirtualTexturingFwd.h"
+#include "PathTracingOutputInvalidateReason.h"
+
+#include "SceneRenderingAllocator.h"
+#include "PooledRenderTarget.h"
+
+#if UE_ENABLE_INCLUDE_ORDER_DEPRECATED_IN_5_8
 #include "RHI.h"
 #include "RenderResource.h"
 #include "RenderUtils.h"
-#include "Misc/EnumClassFlags.h"
 #include "UniformBuffer.h"
 #include "VirtualTexturing.h"
 #include "RenderGraphDefinitions.h"
-#include "PathTracingOutputInvalidateReason.h"
-
-#include "SceneRenderingAllocator.h"
-#include "PooledRenderTarget.h"
+#endif //UE_ENABLE_INCLUDE_ORDER_DEPRECATED_IN_5_8
 
 class FCanvas;
 class FCanvasRenderContext;
diff --git a/Engine/Source/Runtime/RenderCore/Public/UnifiedBuffer.h b/Engine/Source/Runtime/RenderCore/Public/UnifiedBuffer.h
--- a/Engine/Source/Runtime/RenderCore/Public/UnifiedBuffer.h
+++ b/Engine/Source/Runtime/RenderCore/Public/UnifiedBuffer.h
@@ -15,6 +15,7 @@
 #include "RenderGraphDefinitions.h"
 #include "RenderGraphResources.h"
 #include "Templates/RefCounting.h"
+#include "SceneRenderingAllocator.h"
 
 class FRDGBuilder;
 class FRHICommandList;
diff --git a/Engine/Source/Runtime/Renderer/Private/SceneTextures.cpp b/Engine/Source/Runtime/Renderer/Private/SceneTextures.cpp
--- a/Engine/Source/Runtime/Renderer/Private/SceneTextures.cpp
+++ b/Engine/Source/Runtime/Renderer/Private/SceneTextures.cpp
@@ -23,6 +23,7 @@
 #include "Substrate/Substrate.h"
 #include "VisualizeTexture.h"
 #include "PostProcess/PostProcessSubsurface.h"
+#include "DBufferTextures.h"
 
 static TAutoConsoleVariable<int32> CVarSceneTargetsResizeMethod(
 	TEXT("r.SceneRenderTargetResizeMethod"),
@@ -999,6 +1000,11 @@ TArray<FRDGTextureRef> FSceneTextures::EnumerateSceneTextures() const
 	return Results;
 }
 
+FRHITexture* FSceneTextureExtracts::GetDepthTexture() const
+{
+	return Depth ? Depth->GetRHI() : nullptr;
+}
+
 void FSceneTextureExtracts::QueueExtractions(FRDGBuilder& GraphBuilder, const FSceneTextures& SceneTextures)
 {
 	// Free up the memory for reuse during the RDG execution phase.
diff --git a/Engine/Source/Runtime/Renderer/Private/ScreenPass.cpp b/Engine/Source/Runtime/Renderer/Private/ScreenPass.cpp
--- a/Engine/Source/Runtime/Renderer/Private/ScreenPass.cpp
+++ b/Engine/Source/Runtime/Renderer/Private/ScreenPass.cpp
@@ -7,9 +7,9 @@
 #include "PostProcess/SceneFilterRendering.h"
 #include "RendererModule.h"
 #include "RenderGraphUtils.h"
-#include "ScenePrivate.h"
 #include "SystemTextures.h"
 #include "UnrealClient.h"
+#include "SceneRendering.h"
 
 IMPLEMENT_GLOBAL_SHADER(FScreenPassVS, "/Engine/Private/ScreenPass.usf", "ScreenPassVS", SF_Vertex);
 
@@ -271,6 +271,22 @@ void SetScreenPassPipelineState(FRHICommandList& RHICmdList, const FScreenPassPi
 	SetGraphicsPipelineState(RHICmdList, GraphicsPSOInit, ScreenPassDraw.StencilRef);
 }
 
+FScreenPassViewInfo::FScreenPassViewInfo(const FSceneView& View)
+	: FeatureLevel(View.FeatureLevel)
+	, StereoViewIndex(View.StereoViewIndex)
+	, InstanceCount(View.GetStereoPassInstanceFactor())
+	, bHMDHiddenAreaMaskActive(View.bHMDHiddenAreaMaskActive)
+{
+}
+
+FScreenPassViewInfo::FScreenPassViewInfo(const FSceneView& View, const uint32 OverrideInstanceCount)
+	: FeatureLevel(View.FeatureLevel)
+	, StereoViewIndex(View.StereoViewIndex)
+	, InstanceCount(OverrideInstanceCount)
+	, bHMDHiddenAreaMaskActive(View.bHMDHiddenAreaMaskActive)
+{
+}
+
 void DrawScreenPass_PostSetup(
 	FRHICommandList& RHICmdList,
 	const FScreenPassViewInfo& ViewInfo,
@@ -516,4 +532,12 @@ void AddDownsampleDepthPass(
 		VertexShader, PixelShader,
 		DepthStencilState,
 		PassParameters);
-}
\ No newline at end of file
+}
+
+FCanvas& CreateCanvasForRenderTarget(FRDGBuilder& GraphBuilder, const FSceneView& View, const FScreenPassRenderTarget& RenderTarget)
+{
+	const FSceneViewFamily& ViewFamily = *View.Family;
+	FCanvas& Canvas = *FCanvas::Create(GraphBuilder, RenderTarget.Texture, nullptr, ViewFamily.Time, View.GetFeatureLevel(), ViewFamily.DebugDPIScale);
+	Canvas.SetRenderTargetRect(RenderTarget.ViewRect);
+	return Canvas;
+}
diff --git a/Engine/Source/Runtime/Renderer/Public/MeshPassProcessor.h b/Engine/Source/Runtime/Renderer/Public/MeshPassProcessor.h
--- a/Engine/Source/Runtime/Renderer/Public/MeshPassProcessor.h
+++ b/Engine/Source/Runtime/Renderer/Public/MeshPassProcessor.h
@@ -6,21 +6,43 @@
 
 #pragma once
 
-#include "DebugViewModeHelpers.h"
+#include "CoreMinimal.h"
+#include "Math/MathFwd.h"
+#include "Experimental/Containers/RobinHoodHashTable.h"
+#include "RHI.h"
+#include "RHIResources.h"
+#include "Shader.h"
+#include "PSOPrecache.h"
+#include "PSOPrecacheMaterial.h"
+#include "PipelineStateCache.h"
+#include "Misc/ScopeRWLock.h"
+#include "MeshDrawCommandStatsDefines.h"
+#include "VertexFactory.h"
 #include "MeshDrawShaderBindings.h"
-#include "MeshMaterialShader.h"
-#include "RendererInterface.h"
-#include "GBufferInfo.h"
-#include "SceneUtils.h"
+#include "Misc/AssertionMacros.h"
+#include "SceneRenderingAllocator.h"
 #include "MeshBatch.h"
-#include "MeshDrawCommandStatsDefines.h"
-#include "PSOPrecacheMaterial.h"
+#include "SceneUtils.h"
+#include "GBufferInfo.h"
 #include "Hash/CityHash.h"
-#include "Experimental/Containers/RobinHoodHashTable.h"
+#include "DebugViewModeHelpers.h"
+#include "ProfilingDebugging/RealtimeGPUProfiler.h"
 #include <atomic>
 
+#if UE_ENABLE_INCLUDE_ORDER_DEPRECATED_IN_5_8
+#include "MeshDrawCommandStatsDefines.h"
+#include "RendererInterface.h"
+#include "MeshMaterialShader.h"
+#endif //UE_ENABLE_INCLUDE_ORDER_DEPRECATED_IN_5_8
+
 #define MESH_DRAW_COMMAND_DEBUG_DATA ((!UE_BUILD_SHIPPING && !UE_BUILD_TEST) || VALIDATE_MESH_COMMAND_BINDINGS || WANTS_DRAW_MESH_EVENTS || WITH_DEBUG_VIEW_MODES)
 
+class FScene;
+class FPrimitiveSceneInfo;
+class FPrimitiveSceneProxy;
+class FMaterialRenderProxy;
+struct FMeshBatch;
+
 class FGPUScene;
 class FInstanceCullingDrawParams;
 
diff --git a/Engine/Source/Runtime/Renderer/Public/SceneRenderTargetParameters.h b/Engine/Source/Runtime/Renderer/Public/SceneRenderTargetParameters.h
--- a/Engine/Source/Runtime/Renderer/Public/SceneRenderTargetParameters.h
+++ b/Engine/Source/Runtime/Renderer/Public/SceneRenderTargetParameters.h
@@ -2,11 +2,25 @@
 
 #pragma once
 
-#include "RendererInterface.h"
+#include "Math/MathFwd.h"
+#include "Misc/EnumClassFlags.h"
+#include "RenderGraphFwd.h"
+#include "RenderResource.h"
+#include "RHIFeatureLevel.h"
+#include "ShaderParameterMacros.h"
 #include "SceneTexturesConfig.h"
 
-class FRDGBuilder;
+#if UE_ENABLE_INCLUDE_ORDER_DEPRECATED_IN_5_8
+#include "RendererInterface.h"
+#endif //UE_ENABLE_INCLUDE_ORDER_DEPRECATED_IN_5_8
+
 struct FSceneTextures;
+class FSceneTextureUniformParameters;
+class FSceneView;
+class FViewInfo;
+class FMobileSceneTextureUniformParameters;
+class FSceneTextureShaderParameters;
+struct IPooledRenderTarget;
 
 enum class ESceneTexture
 {
@@ -154,10 +168,7 @@ public:
 		return Parameters;
 	}
 
-	FRHITexture* GetDepthTexture() const
-	{
-		return Depth ? Depth->GetRHI() : nullptr;
-	}
+	FRHITexture* GetDepthTexture() const;
 
 	RENDERER_API void QueueExtractions(FRDGBuilder& GraphBuilder, const FSceneTextures& SceneTextures);
 
diff --git a/Engine/Source/Runtime/Renderer/Public/ScreenPass.h b/Engine/Source/Runtime/Renderer/Public/ScreenPass.h
--- a/Engine/Source/Runtime/Renderer/Public/ScreenPass.h
+++ b/Engine/Source/Runtime/Renderer/Public/ScreenPass.h
@@ -9,11 +9,16 @@
 #include "RenderGraphBuilder.h"
 #include "RenderGraphUtils.h"
 #include "RHIStaticStates.h"
+
+#if UE_ENABLE_INCLUDE_ORDER_DEPRECATED_IN_5_8
 #include "SceneView.h"
+#endif //UE_ENABLE_INCLUDE_ORDER_DEPRECATED_IN_5_8
 
 struct FScreenPassTexture;
 struct FScreenPassTextureSlice;
 
+class FSceneView;
+
 // Returns whether a HMD hidden area mask is being used for VR.
 RENDERER_API bool IsHMDHiddenAreaMaskActive();
 
@@ -446,21 +451,9 @@ struct FScreenPassViewInfo
 	{
 	}
 
-	FScreenPassViewInfo(const FSceneView& View)
-		: FeatureLevel(View.FeatureLevel)
-		, StereoViewIndex(View.StereoViewIndex)
-		, InstanceCount(View.GetStereoPassInstanceFactor())
-		, bHMDHiddenAreaMaskActive(View.bHMDHiddenAreaMaskActive)
-	{
-	}
+	RENDERER_API FScreenPassViewInfo(const FSceneView& View);
 
-	FScreenPassViewInfo(const FSceneView& View, const uint32 OverrideInstanceCount)
-		: FeatureLevel(View.FeatureLevel)
-		, StereoViewIndex(View.StereoViewIndex)
-		, InstanceCount(OverrideInstanceCount)
-		, bHMDHiddenAreaMaskActive(View.bHMDHiddenAreaMaskActive)
-	{
-	}
+	RENDERER_API FScreenPassViewInfo(const FSceneView& View, const uint32 OverrideInstanceCount);
 };
 
 RENDERER_API void DrawScreenPass_PostSetup(
@@ -725,6 +718,8 @@ inline void AddRenderTargetPass(
 	GraphBuilder.AddPass(MoveTemp(PassName), PassParameters, ERDGPassFlags::Raster, MoveTemp(Function));
 }
 
+RENDERER_API FCanvas& CreateCanvasForRenderTarget(FRDGBuilder& GraphBuilder, const FSceneView& View, const FScreenPassRenderTarget& RenderTarget);
+
 template <typename TFunction>
 inline void AddDrawCanvasPass(
 	FRDGBuilder& GraphBuilder,
@@ -735,9 +730,7 @@ inline void AddDrawCanvasPass(
 {
 	check(Output.IsValid());
 
-	const FSceneViewFamily& ViewFamily = *View.Family;
-	FCanvas& Canvas = *FCanvas::Create(GraphBuilder, Output.Texture, nullptr, ViewFamily.Time, View.GetFeatureLevel(), ViewFamily.DebugDPIScale);
-	Canvas.SetRenderTargetRect(Output.ViewRect);
+	FCanvas& Canvas = CreateCanvasForRenderTarget(GraphBuilder, View, Output);
 
 	Function(Canvas);
 
