Reflection Capture editing optimization:  When reflection captures change, we need to selectively update only primitives that are affected by the change.  Previously, a full primitive scene info and subsequent mesh draw command invalidation would occur for every static mesh, on any type of reflection capture edit.  This can drop the frame rate to low single digits depending on the complexity of scene loaded, and make dragging or slider adjusting values impossible.  In my moderately sized test scene (3000 primitives), cost went from ~60 ms to ~2 ms for transform changes, and ~5 ms for editing other properties (transform only affects newly intersected or no longer intersected primitives, while other properties affect all primitives in the capture's bounds due to the proxy being fully recreated).
* Deprecated unused CachedReflectionCaptureProxies, since including the three nearest reflection captures would cause more primitives to update, needlessly complicate the logic, and itself is a non-trivial performance cost (adding a couple ms in my test example).
* Fixed a bug where transforming a reflection capture didn't attempt to update its bounding volume in the scene (RegisteredReflectionCapturePositionAndRadius).
* Upgraded Guid derived from component path, used for deterministic sorting, from 32-bit to 64-bit, to reduce collision likelihood to negligible, since we are counting on the hash being unique for detecting sort order changes.

#rnx

#Synced-CL 51775596

diff --git a/Engine/Source/Runtime/Engine/Private/Components/ReflectionCaptureComponent.cpp b/Engine/Source/Runtime/Engine/Private/Components/ReflectionCaptureComponent.cpp
--- a/Engine/Source/Runtime/Engine/Private/Components/ReflectionCaptureComponent.cpp
+++ b/Engine/Source/Runtime/Engine/Private/Components/ReflectionCaptureComponent.cpp
@@ -1310,7 +1310,9 @@ FReflectionCaptureProxy::FReflectionCaptureProxy(const UReflectionCaptureCompone
 	SetTransform(InComponent->GetComponentTransform().ToMatrixWithScale());
 	InfluenceRadius = InComponent->GetInfluenceBoundingRadius();
 	Brightness = InComponent->Brightness;
-	Guid = GetTypeHash( Component->GetPathName() );
+
+	FString PathName = Component->GetPathName();
+	Guid = FXxHash64::HashBuffer(*PathName, PathName.NumBytesWithoutNull());
 
 	bUsingPreviewCaptureData = MapBuildData == NULL;
 }
diff --git a/Engine/Source/Runtime/Engine/Public/SceneProxies/ReflectionCaptureProxy.h b/Engine/Source/Runtime/Engine/Public/SceneProxies/ReflectionCaptureProxy.h
--- a/Engine/Source/Runtime/Engine/Public/SceneProxies/ReflectionCaptureProxy.h
+++ b/Engine/Source/Runtime/Engine/Public/SceneProxies/ReflectionCaptureProxy.h
@@ -2,6 +2,7 @@
 
 #pragma once
 
+#include "Hash/xxhash.h"
 #include "Math/DoubleFloat.h"
 #include "ShaderParameterMacros.h"
 
@@ -45,7 +46,7 @@ public:
 	FDFVector3 Position;
 	float InfluenceRadius;
 	float Brightness;
-	uint32 Guid;
+	FXxHash64 Guid;
 	FVector3f CaptureOffset;
 	int32 SortedCaptureIndex; // Index into ReflectionSceneData.SortedCaptures (and ReflectionCaptures uniform buffer).
 
diff --git a/Engine/Source/Runtime/Renderer/Private/PlanarReflectionRendering.cpp b/Engine/Source/Runtime/Renderer/Private/PlanarReflectionRendering.cpp
--- a/Engine/Source/Runtime/Renderer/Private/PlanarReflectionRendering.cpp
+++ b/Engine/Source/Runtime/Renderer/Private/PlanarReflectionRendering.cpp
@@ -663,7 +663,7 @@ void FScene::AddPlanarReflection(UPlanarReflectionComponent* Component)
 	ENQUEUE_RENDER_COMMAND(FAddPlanarReflectionCommand)(
 		[SceneProxy, Scene](FRHICommandListImmediate& RHICmdList)
 		{
-			Scene->ReflectionSceneData.bRegisteredReflectionCapturesHasChanged = true;
+			Scene->ReflectionSceneData.SetRegisteredReflectionCapturesHasChanged();
 			Scene->PlanarReflections.Add(SceneProxy);
 		});
 }
@@ -678,8 +678,14 @@ void FScene::RemovePlanarReflection(UPlanarReflectionComponent* Component)
 	ENQUEUE_RENDER_COMMAND(FRemovePlanarReflectionCommand)(
 		[SceneProxy, Scene](FRHICommandListImmediate& RHICmdList)
 		{
-			Scene->ReflectionSceneData.bRegisteredReflectionCapturesHasChanged = true;
+			Scene->ReflectionSceneData.SetRegisteredReflectionCapturesHasChanged();
 			Scene->PlanarReflections.Remove(SceneProxy);
+
+			// Need to clear out all planar reflection references on removal to avoid dangling pointers.
+			for (int32 PrimitiveIndex = 0; PrimitiveIndex < Scene->Primitives.Num(); ++PrimitiveIndex)
+			{
+				Scene->Primitives[PrimitiveIndex]->RemoveCachedPlanarReflection(SceneProxy);
+			}
 		});
 }
 
@@ -693,7 +699,7 @@ void FScene::UpdatePlanarReflectionTransform(UPlanarReflectionComponent* Compone
 	ENQUEUE_RENDER_COMMAND(FUpdatePlanarReflectionCommand)(
 		[SceneProxy, Transform, Scene](FRHICommandListImmediate& RHICmdList)
 		{
-			Scene->ReflectionSceneData.bRegisteredReflectionCapturesHasChanged = true;
+			Scene->ReflectionSceneData.SetRegisteredReflectionCapturesHasChanged();
 			SceneProxy->UpdateTransform(Transform);
 		});
 }
diff --git a/Engine/Source/Runtime/Renderer/Private/PrimitiveSceneInfo.cpp b/Engine/Source/Runtime/Renderer/Private/PrimitiveSceneInfo.cpp
--- a/Engine/Source/Runtime/Renderer/Private/PrimitiveSceneInfo.cpp
+++ b/Engine/Source/Runtime/Renderer/Private/PrimitiveSceneInfo.cpp
@@ -350,7 +350,9 @@ FPrimitiveSceneInfo::FPrimitiveSceneInfo(const FPrimitiveSceneInfoAdapter& InAda
 	
 	LODParentComponentId = InAdapter.LODParentComponentId;
 
+	PRAGMA_DISABLE_DEPRECATION_WARNINGS
 	FMemory::Memzero(CachedReflectionCaptureProxies);
+	PRAGMA_ENABLE_DEPRECATION_WARNINGS
 
 	if (FInstanceCullingContext::IsGPUCullingEnabled())
 	{
@@ -2558,28 +2560,38 @@ bool FPrimitiveSceneInfo::NeedsReflectionCaptureUpdate() const
 		(Scene->GetShadingPath() == EShadingPath::Mobile || IsForwardShadingEnabled(Scene->GetShaderPlatform()));
 }
 
-void FPrimitiveSceneInfo::CacheReflectionCaptures()
+bool FPrimitiveSceneInfo::CacheReflectionCaptures()
 {
 	// do not use Scene->PrimitiveBounds here, as it may be not initialized yet
 	FBoxSphereBounds BoxSphereBounds = Proxy->GetBounds(); 
-	
+
+	const FPlanarReflectionSceneProxy* PreviousPlanarReflectionProxy = CachedPlanarReflectionProxy;
+	const FReflectionCaptureProxy* PreviousReflectionCaptureProxy = CachedReflectionCaptureProxy;
+
 	CachedReflectionCaptureProxy = Scene->FindClosestReflectionCapture(BoxSphereBounds.Origin);
 	CachedPlanarReflectionProxy = Scene->FindClosestPlanarReflection(BoxSphereBounds);
-	if (Scene->GetShadingPath() == EShadingPath::Mobile)
-	{
-		// mobile HQ reflections
-		Scene->FindClosestReflectionCaptures(BoxSphereBounds.Origin, CachedReflectionCaptureProxies);
-	}
 	
 	bNeedsCachedReflectionCaptureUpdate = false;
+
+	return PreviousPlanarReflectionProxy != CachedPlanarReflectionProxy || PreviousReflectionCaptureProxy != CachedReflectionCaptureProxy;
 }
 
-void FPrimitiveSceneInfo::RemoveCachedReflectionCaptures()
+void FPrimitiveSceneInfo::RemoveCachedReflectionCapture(const FReflectionCaptureProxy* RemovedProxy)
 {
-	CachedReflectionCaptureProxy = nullptr;
-	CachedPlanarReflectionProxy = nullptr;
-	FMemory::Memzero(CachedReflectionCaptureProxies);
-	bNeedsCachedReflectionCaptureUpdate = true;
+	if (CachedReflectionCaptureProxy && (RemovedProxy == CachedReflectionCaptureProxy || RemovedProxy == nullptr))
+	{
+		CachedReflectionCaptureProxy = nullptr;
+		bNeedsCachedReflectionCaptureUpdate = true;
+	}
+}
+
+void FPrimitiveSceneInfo::RemoveCachedPlanarReflection(const FPlanarReflectionSceneProxy* RemovedProxy)
+{
+	if (CachedPlanarReflectionProxy && (RemovedProxy == CachedPlanarReflectionProxy || RemovedProxy == nullptr))
+	{
+		CachedPlanarReflectionProxy = nullptr;
+		bNeedsCachedReflectionCaptureUpdate = true;
+	}
 }
 
 void FPrimitiveSceneInfo::UpdateComponentLastRenderTime(float CurrentWorldTime, bool bUpdateLastRenderTimeOnScreen)
diff --git a/Engine/Source/Runtime/Renderer/Private/ReflectionEnvironment.cpp b/Engine/Source/Runtime/Renderer/Private/ReflectionEnvironment.cpp
--- a/Engine/Source/Runtime/Renderer/Private/ReflectionEnvironment.cpp
+++ b/Engine/Source/Runtime/Renderer/Private/ReflectionEnvironment.cpp
@@ -526,11 +526,26 @@ void FReflectionEnvironmentSceneData::Reset(FScene* Scene)
 			AllocatedReflectionCaptureState.Empty();
 			CubemapArraySlotsUsed.Empty();
 			SortedCaptures.Empty();
+			PreChangeSortOrder.Empty();
 			NumBoxCaptures = 0;
 			NumSphereCaptures = 0;
 		});
 }
 
+void FReflectionEnvironmentSceneData::SetRegisteredReflectionCapturesHasChanged()
+{
+	bRegisteredReflectionCapturesHasChanged = true;
+
+	// Snapshot the current sort order (key data only) for post-change sort order comparison in UpdateReflectionSceneData.
+	if (PreChangeSortOrder.IsEmpty())
+	{
+		for (const FReflectionCaptureSortData& Entry : SortedCaptures)
+		{
+			PreChangeSortOrder.Add(Entry);
+		}
+	}
+}
+
 void FReflectionEnvironmentCubemapArray::ResizeCubemapArrayGPU(uint32 InMaxCubemaps, int32 InCubemapSize, const TArray<int32>& IndexRemapping)
 {
 	FRHICommandListImmediate& RHICmdList = FRHICommandListImmediate::Get();
diff --git a/Engine/Source/Runtime/Renderer/Private/RendererScene.cpp b/Engine/Source/Runtime/Renderer/Private/RendererScene.cpp
--- a/Engine/Source/Runtime/Renderer/Private/RendererScene.cpp
+++ b/Engine/Source/Runtime/Renderer/Private/RendererScene.cpp
@@ -2808,7 +2808,7 @@ void FScene::AddReflectionCapture(UReflectionCaptureComponent* Component)
 				FPlatformAtomics::InterlockedIncrement(&Scene->NumUnbuiltReflectionCaptures);
 			}
 
-			Scene->ReflectionSceneData.bRegisteredReflectionCapturesHasChanged = true;
+			Scene->ReflectionSceneData.SetRegisteredReflectionCapturesHasChanged();
 			const int32 PackedIndex = Scene->ReflectionSceneData.RegisteredReflectionCaptures.Add(Proxy);
 
 			Proxy->PackedIndex = PackedIndex;
@@ -2839,12 +2839,12 @@ void FScene::RemoveReflectionCapture(UReflectionCaptureComponent* Component)
 				FPlatformAtomics::InterlockedDecrement(&Scene->NumUnbuiltReflectionCaptures);
 			}
 
-			Scene->ReflectionSceneData.bRegisteredReflectionCapturesHasChanged = true;
+			Scene->ReflectionSceneData.SetRegisteredReflectionCapturesHasChanged();
 
 			// Need to clear out all reflection captures on removal to avoid dangling pointers.
 			for (int32 PrimitiveIndex = 0; PrimitiveIndex < Scene->Primitives.Num(); ++PrimitiveIndex)
 			{
-				Scene->Primitives[PrimitiveIndex]->RemoveCachedReflectionCaptures();
+				Scene->Primitives[PrimitiveIndex]->RemoveCachedReflectionCapture(Proxy);
 			}
 
 			int32 CaptureIndex = Proxy->PackedIndex;
@@ -2877,9 +2877,10 @@ void FScene::UpdateReflectionCaptureTransform(UReflectionCaptureComponent* Compo
 		FScene* Scene = this;
 		FReflectionCaptureProxy* Proxy = Component->SceneProxy;
 		FMatrix Transform = Component->GetComponentTransform().ToMatrixWithScale();
+		float InfluenceRadius = Component->GetInfluenceBoundingRadius();
 
 		ENQUEUE_RENDER_COMMAND(UpdateReflectionCaptureTransformCommand)(
-			[Scene, Proxy, Transform, bUsingPreviewCaptureData](FRHICommandListBase& RHICmdList)
+			[Scene, Proxy, Transform, InfluenceRadius, bUsingPreviewCaptureData](FRHICommandListBase& RHICmdList)
 		{
 			if (Proxy->bUsingPreviewCaptureData)
 			{
@@ -2893,8 +2894,12 @@ void FScene::UpdateReflectionCaptureTransform(UReflectionCaptureComponent* Compo
 				FPlatformAtomics::InterlockedIncrement(&Scene->NumUnbuiltReflectionCaptures);
 			}
 
-			Scene->ReflectionSceneData.bRegisteredReflectionCapturesHasChanged = true;
+			Scene->ReflectionSceneData.SetRegisteredReflectionCapturesHasChanged();
 			Proxy->SetTransform(Transform);
+			Proxy->InfluenceRadius = InfluenceRadius;
+
+			int32 CaptureIndex = Proxy->PackedIndex;
+			Scene->ReflectionSceneData.RegisteredReflectionCapturePositionAndRadius[CaptureIndex] = FSphere(Transform.GetOrigin(), InfluenceRadius);
 
 			if (Scene->GetFeatureLevel() <= ERHIFeatureLevel::ES3_1)
 			{
@@ -3013,68 +3018,6 @@ const FPlanarReflectionSceneProxy* FScene::GetForwardPassGlobalPlanarReflection(
 	return nullptr;
 }
 
-void FScene::FindClosestReflectionCaptures(FVector Position, const FReflectionCaptureProxy* (&SortedByDistanceOUT)[FPrimitiveSceneInfo::MaxCachedReflectionCaptureProxies]) const
-{
-	checkSlow(IsInParallelRenderingThread());
-	static const int32 ArraySize = FPrimitiveSceneInfo::MaxCachedReflectionCaptureProxies;
-
-	struct FReflectionCaptureDistIndex
-	{
-		int32 CaptureIndex;
-		float CaptureDistance;
-		const FReflectionCaptureProxy* CaptureProxy;
-	};
-
-	// Find the nearest n captures to this primitive. 
-	const int32 NumRegisteredReflectionCaptures = ReflectionSceneData.RegisteredReflectionCapturePositionAndRadius.Num();
-	const int32 PopulateCaptureCount = FMath::Min(ArraySize, NumRegisteredReflectionCaptures);
-
-	TArray<FReflectionCaptureDistIndex, TFixedAllocator<ArraySize>> ClosestCaptureIndices;
-	ClosestCaptureIndices.AddUninitialized(PopulateCaptureCount);
-
-	for (int32 CaptureIndex = 0; CaptureIndex < PopulateCaptureCount; CaptureIndex++)
-	{
-		ClosestCaptureIndices[CaptureIndex].CaptureIndex = CaptureIndex;
-		ClosestCaptureIndices[CaptureIndex].CaptureDistance = (ReflectionSceneData.RegisteredReflectionCapturePositionAndRadius[CaptureIndex].Center - Position).SizeSquared();
-	}
-	
-	for (int32 CaptureIndex = PopulateCaptureCount; CaptureIndex < NumRegisteredReflectionCaptures; CaptureIndex++)
-	{
-		const float DistanceSquared = (ReflectionSceneData.RegisteredReflectionCapturePositionAndRadius[CaptureIndex].Center - Position).SizeSquared();
-		for (int32 i = 0; i < ArraySize; i++)
-		{
-			if (DistanceSquared<ClosestCaptureIndices[i].CaptureDistance)
-			{
-				ClosestCaptureIndices[i].CaptureDistance = DistanceSquared;
-				ClosestCaptureIndices[i].CaptureIndex = CaptureIndex;
-				break;
-			}
-		}
-	}
-
-	for (int32 CaptureIndex = 0; CaptureIndex < PopulateCaptureCount; CaptureIndex++)
-	{
-		FReflectionCaptureProxy* CaptureProxy = ReflectionSceneData.RegisteredReflectionCaptures[ClosestCaptureIndices[CaptureIndex].CaptureIndex];		
-		ClosestCaptureIndices[CaptureIndex].CaptureProxy = CaptureProxy;
-	}
-	// Sort by influence radius.
-	ClosestCaptureIndices.Sort([](const FReflectionCaptureDistIndex& A, const FReflectionCaptureDistIndex& B)
-		{
-			if (A.CaptureProxy->InfluenceRadius != B.CaptureProxy->InfluenceRadius)
-			{
-				return (A.CaptureProxy->InfluenceRadius < B.CaptureProxy->InfluenceRadius);
-			}
-			return A.CaptureProxy->Guid < B.CaptureProxy->Guid;
-		});
-
-	FMemory::Memzero(SortedByDistanceOUT);
-
-	for (int32 CaptureIndex = 0; CaptureIndex < PopulateCaptureCount; CaptureIndex++)
-	{
-		SortedByDistanceOUT[CaptureIndex] = ClosestCaptureIndices[CaptureIndex].CaptureProxy;
-	}
-}
-
 int64 FScene::GetCachedWholeSceneShadowMapsSize() const
 {
 	int64 CachedShadowmapMemory = 0;
@@ -5195,32 +5138,73 @@ void UpdateReflectionSceneData(FScene* Scene)
 
 
 		// Mobile needs to re-cache all mesh commands when scene capture data has changed
-		bool bNeedsStaticMeshUpdate = GetFeatureLevelShadingPath(Scene->GetFeatureLevel()) == EShadingPath::Mobile;
+		bool bPlatformNeedsStaticMeshUpdate = GetFeatureLevelShadingPath(Scene->GetFeatureLevel()) == EShadingPath::Mobile;
 
 		if (ReflectionSceneData.bRegisteredReflectionCapturesHasChanged)
 		{
-			// Mark all primitives as needing an update
-			// Note: Only visible primitives will actually update their reflection proxy
-			const int32 NumPrimitives = Scene->Primitives.Num();
-			for (int32 PrimitiveIndex = 0; PrimitiveIndex < NumPrimitives; PrimitiveIndex++)
+			// Check if the sorted capture list changed.  If it did, we need to update all primitives that point to a reflection capture.
+			bool bSortedCapturesChanged = false;
+			if (ReflectionSceneData.SortedCaptures.Num() != ReflectionSceneData.PreChangeSortOrder.Num())
+			{
+				bSortedCapturesChanged = true;
+			}
+			else
+			{
+				for (int32 CaptureIndex = 0; CaptureIndex < ReflectionSceneData.SortedCaptures.Num(); CaptureIndex++)
+				{
+					if (ReflectionSceneData.SortedCaptures[CaptureIndex] != ReflectionSceneData.PreChangeSortOrder[CaptureIndex])
+					{
+						bSortedCapturesChanged = true;
+						break;
+					}
+				}
+			}
+			ReflectionSceneData.PreChangeSortOrder.Empty();
+
+			// Mark all primitives as needing an update, that are affected by a capture that was updated.
+			for (FPrimitiveSceneInfo* Primitive : Scene->Primitives)
 			{
-				FPrimitiveSceneInfo* Primitive = Scene->Primitives[PrimitiveIndex];
-				Primitive->RemoveCachedReflectionCaptures();
+				if (bPlatformNeedsStaticMeshUpdate)
+				{
+					// If we have been flagged as requiring update, track that we need an update, even if CacheReflectionCaptures (which clears this flag)
+					// returns false.  This handles the case where a primitive pointed to a modified reflection capture, but no longer does.  In that case,
+					// CacheReflectionCaptures() will return false, since the capture pointers will already have been nulled out, and won't change.  We
+					// need to generate new draw calls in that case, that don't use the reflection capture.
+					bool bNeedsUpdate = Primitive->bNeedsCachedReflectionCaptureUpdate;
+
+					// We need to call "CacheReflectionCaptures" for every primitive, even ones not flagged as affected by a changed capture, because it's
+					// possible a change to the size or location of a capture causes the primitive to now intersect it.
+					if (Primitive->CacheReflectionCaptures())
+					{
+						bNeedsUpdate = true;
+					}
 
-				if (bNeedsStaticMeshUpdate)
+					// Finally, if the sorted captures changed, any primitive pointing to a capture whatsoever must be updated, since the sort indices
+					// will have changed globally.
+					if (bSortedCapturesChanged && Primitive->CachedReflectionCaptureProxy)
+					{
+						bNeedsUpdate = true;
+					}
+
+					if (bNeedsUpdate)
+					{
+						Primitive->RequestStaticMeshUpdate();
+					}
+				}
+				else
 				{
-					Primitive->CacheReflectionCaptures();
-					Primitive->RequestStaticMeshUpdate();
+					// For platforms that don't need caching, remove the cache.
+					Primitive->RemoveCachedReflectionCapture(nullptr);
+					Primitive->RemoveCachedPlanarReflection(nullptr);
 				}
 			}
 
 			ReflectionSceneData.bRegisteredReflectionCapturesHasChanged = false;
-			bNeedsStaticMeshUpdate = false;
 		}
 		
 		if (ReflectionSceneData.bSkylightHasChanged)
 		{
-			if (bNeedsStaticMeshUpdate)
+			if (bPlatformNeedsStaticMeshUpdate)
 			{
 				// Mark all primitives as needing an update
 				// Note: Only visible primitives will actually update their reflection proxy
diff --git a/Engine/Source/Runtime/Renderer/Private/ScenePrivate.h b/Engine/Source/Runtime/Renderer/Private/ScenePrivate.h
--- a/Engine/Source/Runtime/Renderer/Private/ScenePrivate.h
+++ b/Engine/Source/Runtime/Renderer/Private/ScenePrivate.h
@@ -9,6 +9,7 @@
 // Dependencies.
 
 #include "Containers/SparseSet.h"
+#include "Hash/xxhash.h"
 #include "Misc/Guid.h"
 #include "Math/RandomStream.h"
 #include "Templates/PimplPtr.h"
@@ -1682,19 +1683,13 @@ public:
 	}
 };
 
-struct FReflectionCaptureSortData
+// Subset of FReflectionCaptureSortData, to support separate sorting of the key by itself.
+struct FReflectionCaptureSortKey
 {
-	FDFVector3 Position;
-	FMatrix44f BoxTransform;
-	uint32 Guid;
-	int32 CubemapIndex;
+	FXxHash64 Guid;
 	float Radius;
-	FVector4f CaptureProperties;
-	FVector4f BoxScales;
-	FVector4f CaptureOffsetAndAverageBrightness;
-	FReflectionCaptureProxy* CaptureProxy;
 
-	bool operator < (const FReflectionCaptureSortData& Other) const
+	bool operator<(const FReflectionCaptureSortKey& Other) const
 	{
 		if (Radius != Other.Radius)
 		{
@@ -1705,6 +1700,22 @@ struct FReflectionCaptureSortData
 			return Guid < Other.Guid;
 		}
 	}
+
+	bool operator!=(const FReflectionCaptureSortKey& Other) const
+	{
+		return Guid != Other.Guid || Radius != Other.Radius;
+	}
+};
+
+struct FReflectionCaptureSortData : FReflectionCaptureSortKey
+{
+	FDFVector3 Position;
+	FMatrix44f BoxTransform;
+	int32 CubemapIndex;
+	FVector4f CaptureProperties;
+	FVector4f BoxScales;
+	FVector4f CaptureOffsetAndAverageBrightness;
+	FReflectionCaptureProxy* CaptureProxy;
 };
 
 
@@ -1810,6 +1821,30 @@ public:
 	TUniformBufferRef<FReflectionCaptureShaderData> ReflectionCaptureUniformBuffer;
 	TUniformBufferRef<FMobileReflectionCaptureShaderData> MobileReflectionCaptureUniformBuffer;
 
+	/**
+	 * As a mobile renderer optimization, when reflection captures change, we need to selectively update only primitives that are affected by the
+	 * change.  Previously, a full primitive scene info and subsequent mesh draw command invalidation would occur for every static mesh, on any
+	 * type of reflection capture edit.  This can drop the frame rate to low single digits depending on the complexity of scene loaded, and make
+	 * dragging or slider adjusting values impossible.
+	 *
+	 * The cases where static meshes need to be updated on reflection capture changes for mobile are threefold.  First, a uniform buffer is referenced
+	 * by static mesh draw calls, which contains a pointer to the texture for reflection captures that have a user defined texture, and settings for
+	 * intensity and blending.  The proxy and its uniform buffer are recreated whenever a reflection capture is modified, and references need to be
+	 * updated to the new uniform buffer.  Second, a reflection capture index, into the sorted reflection capture array, is stored on each primitive.
+	 * The sort order can change if the capture is resized to a point where its radius changes relative to another capture, as radius is part of the
+	 * sort key.  Finally, a change to the size or transform of a capture can cause it to intersect or no longer intersect a given primitive.
+	 *
+	 * Any static mesh that references, or previously referenced, a changed reflection capture needs to be updated.  A flag on the primitive
+	 * scene proxy tracks whether its cached reflection capture has been modified.  Note that any property edit besides a transform will remove and
+	 * recreate the render proxy, so we can just handle property changes in "remove".  Transform changes only matter for cases where a capture is newly
+	 * or no longer intersected, which is detected in the call to CacheReflectionCaptures, which is run on all primitives regardless of their flags.
+	 *
+	 * Handling sort order changes is done by caching the previous sort order keys during a change, so it can be compared with the current sort order.
+	 * On a sort order change, all primitives referencing ANY capture need to be updated, since their indices may have changed, even if they are not
+	 * the specifically edited reflection capture.
+	 */
+	TArray<FReflectionCaptureSortKey> PreChangeSortOrder;
+
 	/** 
 	 * Game thread list of reflection components that have been allocated in the cubemap array. 
 	 * These are not necessarily all visible or being rendered, but their scene state is stored in the cubemap array.
@@ -1856,6 +1891,9 @@ public:
 
 	/** Resets the structure to empty, useful if you want to shrink the allocation. */
 	void Reset(FScene* Scene);
+
+	/** Mark reflection captures as changed, also initializing PreChangeSortOrder. */
+	void SetRegisteredReflectionCapturesHasChanged();
 };
 
 /** Scene state used to manage hair strands. */
@@ -3688,8 +3726,6 @@ public:
 
 	const class FPlanarReflectionSceneProxy* GetForwardPassGlobalPlanarReflection() const;
 
-	void FindClosestReflectionCaptures(FVector Position, const FReflectionCaptureProxy* (&SortedByDistanceOUT)[FPrimitiveSceneInfo::MaxCachedReflectionCaptureProxies]) const;
-
 	int64 GetCachedWholeSceneShadowMapsSize() const;
 
 	/**
diff --git a/Engine/Source/Runtime/Renderer/Public/PrimitiveSceneInfo.h b/Engine/Source/Runtime/Renderer/Public/PrimitiveSceneInfo.h
--- a/Engine/Source/Runtime/Renderer/Public/PrimitiveSceneInfo.h
+++ b/Engine/Source/Runtime/Renderer/Public/PrimitiveSceneInfo.h
@@ -332,6 +332,7 @@ public:
 	TArray<int32, TInlineAllocator<1>> LumenPrimitiveGroupIndices;
 
 	static const uint32 MaxCachedReflectionCaptureProxies = 3;
+	UE_DEPRECATED(5.8, "CachedReflectionCaptureProxies is no longer initialized, and will be removed in a future release.  Use CachedReflectionCaptureProxy.")
 	const FReflectionCaptureProxy* CachedReflectionCaptureProxies[MaxCachedReflectionCaptureProxies];
 	
 	/** The hit proxies used by the primitive. */
@@ -502,11 +503,14 @@ public:
 	/** Returns whether the primitive needs to call CacheReflectionCaptures. */
 	bool NeedsReflectionCaptureUpdate() const;
 
-	/** Cache per-primitive reflection captures used for mobile/forward rendering. */
-	void CacheReflectionCaptures();
+	/** Cache per-primitive reflection captures used for mobile/forward rendering.  Returns true if the cached reflection captures changed. */
+	bool CacheReflectionCaptures();
 
-	/** Nulls out the cached per-primitive reflection captures. */
-	void RemoveCachedReflectionCaptures();
+	/** Nulls out the cached per-primitive reflection capture, if it matches ChangedProxy or ChangedProxy is null. */
+	void RemoveCachedReflectionCapture(const FReflectionCaptureProxy* ChangedProxy = nullptr);
+
+	/** Nulls out the cached per-primitive planar reflection, if it matches ChangedProxy or ChangedProxy is null. */
+	void RemoveCachedPlanarReflection(const FPlanarReflectionSceneProxy* ChangedProxy = nullptr);
 
 	/** Helper function for writing out to the last render times to the game thread */
 	void UpdateComponentLastRenderTime(float CurrentWorldTime, bool bUpdateLastRenderTimeOnScreen);
