Cherry-pick CL 51781430: Plumb Directional Light Shadow Amount to Distance Field Shadowing

#Synced-CL 51749323

diff --git a/Engine/Shaders/Private/DistanceFieldShadowing.usf b/Engine/Shaders/Private/DistanceFieldShadowing.usf
--- a/Engine/Shaders/Private/DistanceFieldShadowing.usf
+++ b/Engine/Shaders/Private/DistanceFieldShadowing.usf
@@ -819,6 +819,7 @@ float InvNearFadePlaneLength;
 
 float OneOverDownsampleFactor;
 float2 ShadowFactorsUVBilinearMax;
+float ShadowFadeFraction;
 
 void DistanceFieldShadowingUpsamplePS(
 	in float4 SVPos : SV_POSITION,
@@ -833,6 +834,7 @@ void DistanceFieldShadowingUpsamplePS(
 
 	float NearBlendFactor = saturate((SceneDepth - NearFadePlaneOffset) * InvNearFadePlaneLength);
 	Output = lerp(1, Output, NearBlendFactor);
+	Output = lerp(1, Output, ShadowFadeFraction);
 
 	OutColor = EncodeLightAttenuation(half4(Output, Output, Output, Output));
 }
diff --git a/Engine/Source/Runtime/Renderer/Private/DistanceFieldShadowing.cpp b/Engine/Source/Runtime/Renderer/Private/DistanceFieldShadowing.cpp
--- a/Engine/Source/Runtime/Renderer/Private/DistanceFieldShadowing.cpp
+++ b/Engine/Source/Runtime/Renderer/Private/DistanceFieldShadowing.cpp
@@ -405,6 +405,7 @@ class FDistanceFieldShadowingUpsamplePS : public FGlobalShader
 		SHADER_PARAMETER(float, NearFadePlaneOffset)
 		SHADER_PARAMETER(float, InvNearFadePlaneLength)
 		SHADER_PARAMETER(float, OneOverDownsampleFactor)
+		SHADER_PARAMETER(float, ShadowFadeFraction)
 		SHADER_PARAMETER(float, MinDepth)
 		SHADER_PARAMETER(float, MaxDepth)
 	END_SHADER_PARAMETER_STRUCT()
@@ -1123,6 +1124,7 @@ void FProjectedShadowInfo::RenderRayTracedDistanceFieldProjection(
 	const FMinimalSceneTextures& SceneTextures,
 	FRDGTextureRef ScreenShadowMaskTexture,
 	const FViewInfo& View,
+	const int32 ViewIndex,
 	FIntRect ScissorRect,
 	bool bProjectingForForwardShading,
 	bool bForceRGBModulation,
@@ -1150,6 +1152,7 @@ void FProjectedShadowInfo::RenderRayTracedDistanceFieldProjection(
 		PassParameters->PS.ScissorRectMinAndSize = FIntRect(ScissorRect.Min, ScissorRect.Size());
 		PassParameters->PS.ShadowFactorsUVBilinearMax = (RayTracedShadowsTextureViewportMax - 0.5f) * RayTracedShadowsTextureExtentInverse;
 		PassParameters->PS.OneOverDownsampleFactor = 1.0f / GetDFShadowDownsampleFactor();
+		PassParameters->PS.ShadowFadeFraction = FadeAlphas[ViewIndex];
 
 		GetDistanceFieldShadowRange(this, DFPT_SignedDistanceField, PassParameters->PS.MinDepth, PassParameters->PS.MaxDepth);
 
diff --git a/Engine/Source/Runtime/Renderer/Private/ShadowRendering.cpp b/Engine/Source/Runtime/Renderer/Private/ShadowRendering.cpp
--- a/Engine/Source/Runtime/Renderer/Private/ShadowRendering.cpp
+++ b/Engine/Source/Runtime/Renderer/Private/ShadowRendering.cpp
@@ -2243,6 +2243,7 @@ void FSceneRenderer::RenderShadowProjections(
 							SceneTextures,
 							ScreenShadowMaskTexture,
 							View,
+							ViewIndex,
 							ScissorRect,
 							bProjectingForForwardShading);
 					}
diff --git a/Engine/Source/Runtime/Renderer/Private/ShadowRendering.h b/Engine/Source/Runtime/Renderer/Private/ShadowRendering.h
--- a/Engine/Source/Runtime/Renderer/Private/ShadowRendering.h
+++ b/Engine/Source/Runtime/Renderer/Private/ShadowRendering.h
@@ -647,6 +647,7 @@ public:
 		const FMinimalSceneTextures& SceneTextures,
 		FRDGTextureRef ScreenShadowMaskTexture,
 		const FViewInfo& View,
+		const int32 ViewIndex,
 		FIntRect ScissorRect,
 		bool bProjectingForForwardShading,
 		bool bForceRGBModulation = false,
diff --git a/Engine/Source/Runtime/Renderer/Private/SingleLayerWaterRendering.cpp b/Engine/Source/Runtime/Renderer/Private/SingleLayerWaterRendering.cpp
--- a/Engine/Source/Runtime/Renderer/Private/SingleLayerWaterRendering.cpp
+++ b/Engine/Source/Runtime/Renderer/Private/SingleLayerWaterRendering.cpp
@@ -1566,6 +1566,7 @@ void FDeferredShadingSceneRenderer::RenderSingleLayerWaterReflections(
 						SceneTextures,
 						SceneWithoutWaterTextures.SeparatedMainDirLightTexture,
 						View,
+						ViewIndex,
 						ScissorRect,
 						bProjectingForForwardShading,
 						bForceRGBModulation,
