Re-add FSceneTextureExtracts::GetDepthTexture to renderer API

This function was implictly removed in a recent change by moving the implementation from the header to the cpp file. This CL simply restores the state of the API from a couple of days ago.

#Synced-CL 51833721

diff --git a/Engine/Source/Runtime/Renderer/Public/SceneRenderTargetParameters.h b/Engine/Source/Runtime/Renderer/Public/SceneRenderTargetParameters.h
--- a/Engine/Source/Runtime/Renderer/Public/SceneRenderTargetParameters.h
+++ b/Engine/Source/Runtime/Renderer/Public/SceneRenderTargetParameters.h
@@ -168,7 +168,7 @@ public:
 		return Parameters;
 	}
 
-	FRHITexture* GetDepthTexture() const;
+	RENDERER_API FRHITexture* GetDepthTexture() const;
 
 	RENDERER_API void QueueExtractions(FRDGBuilder& GraphBuilder, const FSceneTextures& SceneTextures);
 
