Updating check to compare only local lights as it should

#Synced-CL 51854089

diff --git a/Engine/Source/Runtime/Renderer/Private/LightGridInjection.cpp b/Engine/Source/Runtime/Renderer/Private/LightGridInjection.cpp
--- a/Engine/Source/Runtime/Renderer/Private/LightGridInjection.cpp
+++ b/Engine/Source/Runtime/Renderer/Private/LightGridInjection.cpp
@@ -1617,7 +1617,7 @@ FComputeLightGridOutput FSceneRenderer::ComputeLightGrid(FRDGBuilder& GraphBuild
 		{
 			// The visibility lists should be exactly the same
 			//checkSlow(PrimaryView->ForwardLightingResources.LocalLightVisibleLightInfosIndex == LocalLightVisibleLightInfosIndex);
-			check(PerViewForwardLightUniformParameters[View.PrimaryViewIndex]->NumLocalLights == NumLightsFinal);
+			check(PerViewForwardLightUniformParameters[View.PrimaryViewIndex]->NumLocalLights == NumVisibleLocalLights);
 			ForwardLightUniformParameters->ForwardLightBuffer = PerViewForwardLightUniformParameters[View.PrimaryViewIndex]->ForwardLightBuffer;
 			ForwardLightUniformParameters->LightNamesBuffer = PerViewForwardLightUniformParameters[View.PrimaryViewIndex]->LightNamesBuffer;
 		}
