Fixed Aliasable format on Debug.PrimitivesDepth

#Synced-CL 51619317

diff --git a/Engine/Source/Runtime/Renderer/Private/PostProcess/PostProcessCompositeDebugPrimitives.cpp b/Engine/Source/Runtime/Renderer/Private/PostProcess/PostProcessCompositeDebugPrimitives.cpp
--- a/Engine/Source/Runtime/Renderer/Private/PostProcess/PostProcessCompositeDebugPrimitives.cpp
+++ b/Engine/Source/Runtime/Renderer/Private/PostProcess/PostProcessCompositeDebugPrimitives.cpp
@@ -49,13 +49,16 @@ static FDebugPrimitivesDrawingContext InitializeDebugPrimitiveTextures(
 	const FRDGTextureDesc ColorDesc(FRDGTextureDesc::Create2D(OutputExtent, PF_B8G8R8A8, FClearValueBinding::Transparent, TexCreate_ShaderResource | TexCreate_RenderTargetable, 1, NumMSAASamples));
 	FRDGTextureRef DebugPrimitiveColor = GraphBuilder.CreateTexture(ColorDesc, TEXT("Debug.PrimitivesColor"));
 
-	const FRDGTextureDesc DepthDesc(FRDGTextureDesc::Create2D(OutputExtent, PF_DepthStencil, FClearValueBinding::DepthFar, TexCreate_ShaderResource | TexCreate_DepthStencilTargetable, 1, NumMSAASamples));
+	FRDGTextureDesc DepthDesc(FRDGTextureDesc::Create2D(OutputExtent, PF_DepthStencil, FClearValueBinding::DepthFar, TexCreate_ShaderResource | TexCreate_DepthStencilTargetable, 1, NumMSAASamples));
+	DepthDesc.AliasableFormats.Add(PF_X24_G8);
 	FRDGTextureRef DebugPrimitiveDepth = GraphBuilder.CreateTexture(DepthDesc, TEXT("Debug.PrimitivesDepth"));
-
+	
+	
 	AddClearRenderTargetPass(GraphBuilder, DebugPrimitiveColor);
 	AddClearDepthStencilPass(GraphBuilder, DebugPrimitiveDepth);
 
-	const FRDGTextureDesc UpscaledSceneDepthDesc(FRDGTextureDesc::Create2D(SceneColorExtent, PF_DepthStencil, FClearValueBinding::DepthFar, TexCreate_ShaderResource | TexCreate_DepthStencilTargetable, 1, Inputs.SceneColor.Texture->Desc.NumSamples));
+	FRDGTextureDesc UpscaledSceneDepthDesc(FRDGTextureDesc::Create2D(SceneColorExtent, PF_DepthStencil, FClearValueBinding::DepthFar, TexCreate_ShaderResource | TexCreate_DepthStencilTargetable, 1, Inputs.SceneColor.Texture->Desc.NumSamples));
+	
 	FRDGTextureRef UpscaledSceneDepth = GraphBuilder.CreateTexture(UpscaledSceneDepthDesc, TEXT("Debug.UpscaledWorldDepth"));
 
 	// Debug primitives with SDPG_World are depth tested against the scene depth, and immediately composed on SceneColor.
