Deprecate more unused includes from public rendering headers.

#rnx
#Synced-CL 51867312

diff --git a/Engine/Plugins/FX/Niagara/Source/NiagaraVertexFactories/Private/NiagaraRibbonVertexFactory.cpp b/Engine/Plugins/FX/Niagara/Source/NiagaraVertexFactories/Private/NiagaraRibbonVertexFactory.cpp
--- a/Engine/Plugins/FX/Niagara/Source/NiagaraVertexFactories/Private/NiagaraRibbonVertexFactory.cpp
+++ b/Engine/Plugins/FX/Niagara/Source/NiagaraVertexFactories/Private/NiagaraRibbonVertexFactory.cpp
@@ -10,6 +10,7 @@
 #include "ShaderParameterUtils.h"
 #include "MeshDrawShaderBindings.h"
 #include "MeshMaterialShader.h"
+#include "PipelineStateCache.h"
 
 IMPLEMENT_GLOBAL_SHADER_PARAMETER_STRUCT(FNiagaraRibbonUniformParameters, "NiagaraRibbonVF");
 IMPLEMENT_GLOBAL_SHADER_PARAMETER_STRUCT(FNiagaraRibbonVFLooseParameters, "NiagaraRibbonVFLooseParameters");
diff --git a/Engine/Plugins/FX/Niagara/Source/NiagaraVertexFactories/Private/NiagaraSpriteVertexFactory.cpp b/Engine/Plugins/FX/Niagara/Source/NiagaraVertexFactories/Private/NiagaraSpriteVertexFactory.cpp
--- a/Engine/Plugins/FX/Niagara/Source/NiagaraVertexFactories/Private/NiagaraSpriteVertexFactory.cpp
+++ b/Engine/Plugins/FX/Niagara/Source/NiagaraVertexFactories/Private/NiagaraSpriteVertexFactory.cpp
@@ -11,6 +11,7 @@
 #include "ShaderParameterUtils.h"
 #include "MeshDrawShaderBindings.h"
 #include "MeshMaterialShader.h"
+#include "PipelineStateCache.h"
 
 IMPLEMENT_GLOBAL_SHADER_PARAMETER_STRUCT(FNiagaraSpriteUniformParameters,"NiagaraSpriteVF");
 IMPLEMENT_GLOBAL_SHADER_PARAMETER_STRUCT(FNiagaraSpriteVFLooseParameters, "NiagaraSpriteVFLooseParameters");
diff --git a/Engine/Plugins/Runtime/MeshModelingToolset/Source/ModelingComponents/Private/Scene/WorldRenderCapture.cpp b/Engine/Plugins/Runtime/MeshModelingToolset/Source/ModelingComponents/Private/Scene/WorldRenderCapture.cpp
--- a/Engine/Plugins/Runtime/MeshModelingToolset/Source/ModelingComponents/Private/Scene/WorldRenderCapture.cpp
+++ b/Engine/Plugins/Runtime/MeshModelingToolset/Source/ModelingComponents/Private/Scene/WorldRenderCapture.cpp
@@ -17,6 +17,7 @@
 #include "ImageWriteStream.h"
 #include "ImageCore.h"
 #include "Async/ParallelFor.h"
+#include "RendererInterface.h"
 
 #include "PreviewScene.h"
 #include "Components/DirectionalLightComponent.h"
diff --git a/Engine/Plugins/Runtime/nDisplay/Source/DisplayCluster/Private/Render/Viewport/DisplayClusterViewportManagerViewExtension.cpp b/Engine/Plugins/Runtime/nDisplay/Source/DisplayCluster/Private/Render/Viewport/DisplayClusterViewportManagerViewExtension.cpp
--- a/Engine/Plugins/Runtime/nDisplay/Source/DisplayCluster/Private/Render/Viewport/DisplayClusterViewportManagerViewExtension.cpp
+++ b/Engine/Plugins/Runtime/nDisplay/Source/DisplayCluster/Private/Render/Viewport/DisplayClusterViewportManagerViewExtension.cpp
@@ -12,6 +12,7 @@
 #include "ScreenPass.h"
 #include "CommonRenderResources.h"
 #include "PostProcess/PostProcessMaterialInputs.h"
+#include "SceneTextures.h"
 
 ///////////////////////////////////////////////////////////////////////////////////////
 namespace UE::DisplayCluster::ViewportManagerViewExtension
diff --git a/Engine/Plugins/Runtime/nDisplay/Source/DisplayCluster/Private/Render/Viewport/DisplayClusterViewportManagerViewExtension.h b/Engine/Plugins/Runtime/nDisplay/Source/DisplayCluster/Private/Render/Viewport/DisplayClusterViewportManagerViewExtension.h
--- a/Engine/Plugins/Runtime/nDisplay/Source/DisplayCluster/Private/Render/Viewport/DisplayClusterViewportManagerViewExtension.h
+++ b/Engine/Plugins/Runtime/nDisplay/Source/DisplayCluster/Private/Render/Viewport/DisplayClusterViewportManagerViewExtension.h
@@ -6,6 +6,7 @@
 #include "Render/Viewport/Containers/DisplayClusterViewportProxy_Context.h"
 
 class FDisplayClusterViewportProxy;
+struct FSceneTextures;
 
 #define DISPLAYCLUSTER_SCENE_VIEW_EXTENSION_PRIORITY -1
 
diff --git a/Engine/Plugins/Runtime/nDisplay/Source/DisplayClusterLightCardEditorShaders/Private/DisplayClusterMeshProjectionRenderer.cpp b/Engine/Plugins/Runtime/nDisplay/Source/DisplayClusterLightCardEditorShaders/Private/DisplayClusterMeshProjectionRenderer.cpp
--- a/Engine/Plugins/Runtime/nDisplay/Source/DisplayClusterLightCardEditorShaders/Private/DisplayClusterMeshProjectionRenderer.cpp
+++ b/Engine/Plugins/Runtime/nDisplay/Source/DisplayClusterLightCardEditorShaders/Private/DisplayClusterMeshProjectionRenderer.cpp
@@ -20,6 +20,7 @@
 #include "Materials/Material.h"
 #include "DataDrivenShaderPlatformInfo.h"
 #include "ScreenPass.h"
+#include "ShaderCompilerCore.h"
 
 namespace UE::DisplayClusterMeshProjectionRenderer
 {
diff --git a/Engine/Plugins/VirtualProduction/TextureShare/Source/TextureShare/Private/Game/ViewExtension/TextureShareSceneViewExtension.cpp b/Engine/Plugins/VirtualProduction/TextureShare/Source/TextureShare/Private/Game/ViewExtension/TextureShareSceneViewExtension.cpp
--- a/Engine/Plugins/VirtualProduction/TextureShare/Source/TextureShare/Private/Game/ViewExtension/TextureShareSceneViewExtension.cpp
+++ b/Engine/Plugins/VirtualProduction/TextureShare/Source/TextureShare/Private/Game/ViewExtension/TextureShareSceneViewExtension.cpp
@@ -17,9 +17,11 @@
 #include "RHI.h"
 #include "RHICommandList.h"
 #include "RHIResources.h"
+#include "RenderGraph.h"
 #include "SceneRenderTargetParameters.h"
 #include "SceneView.h"
 #include "SceneRendering.h"
+#include "SceneTextures.h"
 
 namespace UE::TextureShare::SceneViewExtension
 {
diff --git a/Engine/Plugins/VirtualProduction/TextureShare/Source/TextureShare/Private/Game/ViewExtension/TextureShareSceneViewExtension.h b/Engine/Plugins/VirtualProduction/TextureShare/Source/TextureShare/Private/Game/ViewExtension/TextureShareSceneViewExtension.h
--- a/Engine/Plugins/VirtualProduction/TextureShare/Source/TextureShare/Private/Game/ViewExtension/TextureShareSceneViewExtension.h
+++ b/Engine/Plugins/VirtualProduction/TextureShare/Source/TextureShare/Private/Game/ViewExtension/TextureShareSceneViewExtension.h
@@ -9,6 +9,12 @@
 
 class ITextureShareObjectProxy;
 class SWindow;
+class FSceneViewFamily;
+class FSceneView;
+class FViewport;
+struct FSceneTextures;
+class FRDGBuilder;
+class FRHICommandListImmediate;
 
 /** A container with a saved scene view for TS. */
 struct FTextureShareSceneView
diff --git a/Engine/Plugins/VirtualProduction/TextureShare/Source/TextureShare/Private/Module/TextureShareAPI.cpp b/Engine/Plugins/VirtualProduction/TextureShare/Source/TextureShare/Private/Module/TextureShareAPI.cpp
--- a/Engine/Plugins/VirtualProduction/TextureShare/Source/TextureShare/Private/Module/TextureShareAPI.cpp
+++ b/Engine/Plugins/VirtualProduction/TextureShare/Source/TextureShare/Private/Module/TextureShareAPI.cpp
@@ -13,6 +13,7 @@
 #include "Framework/Application/SlateApplication.h"
 #include "Slate/SlateViewportProvider.h"
 #include "Engine/GameViewportClient.h"
+#include "SceneTextures.h"
 
 //////////////////////////////////////////////////////////////////////////////////////////////
 namespace UE::TextureShare::APIHelpers
diff --git a/Engine/Source/Editor/MovieSceneTools/Private/TrackEditorThumbnail/TrackThumbnailUtils.cpp b/Engine/Source/Editor/MovieSceneTools/Private/TrackEditorThumbnail/TrackThumbnailUtils.cpp
--- a/Engine/Source/Editor/MovieSceneTools/Private/TrackEditorThumbnail/TrackThumbnailUtils.cpp
+++ b/Engine/Source/Editor/MovieSceneTools/Private/TrackEditorThumbnail/TrackThumbnailUtils.cpp
@@ -8,6 +8,7 @@
 #include "LegacyScreenPercentageDriver.h"
 #include "MovieSceneToolsUserSettings.h"
 #include "SceneViewExtension.h"
+#include "RendererInterface.h"
 
 static TAutoConsoleVariable<int32> CVarSequencerUsePostProcessThumbnails(
 	TEXT("Sequencer.UsePostProcessThumbnails"),
diff --git a/Engine/Source/Runtime/Engine/Internal/MaterialCache/MaterialCacheAttribute.h b/Engine/Source/Runtime/Engine/Internal/MaterialCache/MaterialCacheAttribute.h
--- a/Engine/Source/Runtime/Engine/Internal/MaterialCache/MaterialCacheAttribute.h
+++ b/Engine/Source/Runtime/Engine/Internal/MaterialCache/MaterialCacheAttribute.h
@@ -3,7 +3,7 @@
 #pragma once
 
 #include "PixelFormat.h"
-#include "VirtualTexturing.h"
+#include "VirtualTexturingFwd.h"
 #include "MaterialValueType.h"
 #include "Containers/Array.h"
 #include "MaterialCacheAttribute.generated.h"
diff --git a/Engine/Source/Runtime/Engine/Public/MaterialExpressionIO.h b/Engine/Source/Runtime/Engine/Public/MaterialExpressionIO.h
--- a/Engine/Source/Runtime/Engine/Public/MaterialExpressionIO.h
+++ b/Engine/Source/Runtime/Engine/Public/MaterialExpressionIO.h
@@ -3,8 +3,11 @@
 #pragma once
 
 #include "CoreMinimal.h"
-#include "UObject/Class.h"
 #include "SceneTypes.h"
+#include "Math/MathFwd.h"
+#include "Math/Color.h"
+#include "UObject/NameTypes.h"
+#include "UObject/StructOpsTypeTraits.h"
 
 struct FExpressionOutput;
 
diff --git a/Engine/Source/Runtime/Engine/Public/MaterialShared.h b/Engine/Source/Runtime/Engine/Public/MaterialShared.h
--- a/Engine/Source/Runtime/Engine/Public/MaterialShared.h
+++ b/Engine/Source/Runtime/Engine/Public/MaterialShared.h
@@ -27,14 +27,12 @@
 #include "Serialization/MemoryWriter.h"
 #include "Serialization/ArchiveProxy.h"
 #include "MaterialKeyIncludeEnum.h"
-#include "MaterialRecursionGuard.h"
 #include "MaterialSceneTextureId.h"
 #include "MaterialShaderPrecompileMode.h"
 #include "MaterialValueType.h"
-#include "VirtualTexturing.h"
+#include "VirtualTexturingFwd.h"
 #include "Templates/UnrealTemplate.h"
-#include "ShaderCompilerCore.h"
-#include "PSOPrecache.h"
+#include "PSOPrecacheFwd.h"
 #include "UObject/ObjectMacros.h"
 #include "Rendering/SubstrateMaterialShared.h"
 #include "Shader/Preshader.h"
diff --git a/Engine/Source/Runtime/Engine/Public/SceneViewExtension.h b/Engine/Source/Runtime/Engine/Public/SceneViewExtension.h
--- a/Engine/Source/Runtime/Engine/Public/SceneViewExtension.h
+++ b/Engine/Source/Runtime/Engine/Public/SceneViewExtension.h
@@ -8,8 +8,12 @@
 
 #include "CoreMinimal.h"
 #include "Delegates/DelegateCombinations.h"
-#include "RendererInterface.h"
 #include "SceneViewExtensionContext.h"
+#include "SceneRenderingAllocator.h"
+
+#if UE_ENABLE_INCLUDE_ORDER_DEPRECATED_IN_5_8
+#include "RendererInterface.h"
+#endif // UE_ENABLE_INCLUDE_ORDER_DEPRECATED_IN_5_8
 
 /**
  *  SCENE VIEW EXTENSIONS
@@ -90,7 +94,7 @@ struct FMobilePostProcessingInputs;
 struct FPostProcessMaterialInputs;
 struct FScreenPassTexture;
 class FViewport;
-
+class ISceneRenderer;
 
 /** This is used to add more flexibility to Post Processing, so that users can subscribe to any after Post Processing Pass events. */
 using FPostProcessingPassDelegate = TDelegate<FScreenPassTexture(FRDGBuilder& GraphBuilder, const FSceneView& View, const FPostProcessMaterialInputs& Inputs)>;
diff --git a/Engine/Source/Runtime/Renderer/Private/RayTracing/RayTracingInstance.cpp b/Engine/Source/Runtime/Renderer/Private/RayTracing/RayTracingInstance.cpp
--- a/Engine/Source/Runtime/Renderer/Private/RayTracing/RayTracingInstance.cpp
+++ b/Engine/Source/Runtime/Renderer/Private/RayTracing/RayTracingInstance.cpp
@@ -5,6 +5,7 @@
 =============================================================================*/
 
 #include "RayTracingInstance.h"
+#include "RHIDefinitions.h"
 
 #if RHI_RAYTRACING
 
diff --git a/Engine/Source/Runtime/Renderer/Private/SceneRendering.cpp b/Engine/Source/Runtime/Renderer/Private/SceneRendering.cpp
--- a/Engine/Source/Runtime/Renderer/Private/SceneRendering.cpp
+++ b/Engine/Source/Runtime/Renderer/Private/SceneRendering.cpp
@@ -5140,6 +5140,12 @@ void FSceneRenderer::UpdatePrimitiveIndirectLightingCacheBuffers(FRHICommandList
 	FRendererModule
 -----------------------------------------------------------------------------*/
 
+TRDGUniformBufferRef<FSceneUniformParameters> GetSceneUniformBufferRef(FRDGBuilder& GraphBuilder, const FSceneView& View)
+{
+	ISceneRenderer* SceneRender = View.Family ? View.Family->GetSceneRenderer() : nullptr;
+	return SceneRender ? SceneRender->GetSceneUniformBufferRef(GraphBuilder) : nullptr;
+}
+
 UE::Renderer::Private::IShadowInvalidatingInstances *FSceneRenderer::GetShadowInvalidatingInstancesInterface(const FSceneView *SceneView)
 {
 	checkf(IsInRenderingThread(), TEXT("Accessing the ShadowInvalidatingInstancesInterface should only be allowed from the rendering thread!"));
diff --git a/Engine/Source/Runtime/Renderer/Private/VariableRateShading/ContrastAdaptiveImageGenerator.h b/Engine/Source/Runtime/Renderer/Private/VariableRateShading/ContrastAdaptiveImageGenerator.h
--- a/Engine/Source/Runtime/Renderer/Private/VariableRateShading/ContrastAdaptiveImageGenerator.h
+++ b/Engine/Source/Runtime/Renderer/Private/VariableRateShading/ContrastAdaptiveImageGenerator.h
@@ -4,10 +4,8 @@
 
 #include "CoreMinimal.h"
 #include "RHI.h"
-#include "RendererInterface.h"
 #include "RenderGraphDefinitions.h"
 #include "VariableRateShadingImageManager.h"
-#include "Engine/Engine.h"
 
 class FContrastAdaptiveImageGenerator : public IVariableRateShadingImageGenerator
 {
diff --git a/Engine/Source/Runtime/Renderer/Private/VariableRateShading/FoveatedImageGenerator.h b/Engine/Source/Runtime/Renderer/Private/VariableRateShading/FoveatedImageGenerator.h
--- a/Engine/Source/Runtime/Renderer/Private/VariableRateShading/FoveatedImageGenerator.h
+++ b/Engine/Source/Runtime/Renderer/Private/VariableRateShading/FoveatedImageGenerator.h
@@ -4,10 +4,8 @@
 
 #include "CoreMinimal.h"
 #include "RHI.h"
-#include "RendererInterface.h"
 #include "RenderGraphDefinitions.h"
 #include "VariableRateShadingImageManager.h"
-#include "Engine/Engine.h"
 
 class FFoveatedImageGenerator : public IVariableRateShadingImageGenerator
 {
diff --git a/Engine/Source/Runtime/Renderer/Public/FXRenderingUtils.h b/Engine/Source/Runtime/Renderer/Public/FXRenderingUtils.h
--- a/Engine/Source/Runtime/Renderer/Public/FXRenderingUtils.h
+++ b/Engine/Source/Runtime/Renderer/Public/FXRenderingUtils.h
@@ -4,11 +4,14 @@
 
 #include "HAL/Platform.h"
 #include "RenderGraphFwd.h"
-#include "RendererUtils.h"
+#include "Shader.h"
 #include "RHIDefinitions.h"
 #include "SceneRenderTargetParameters.h"
-#include "UObject/ObjectMacros.h"
 
+#if UE_ENABLE_INCLUDE_ORDER_DEPRECATED_IN_5_8
+#include "UObject/ObjectMacros.h"
+#include "RendererUtils.h"
+#endif // UE_ENABLE_INCLUDE_ORDER_DEPRECATED_IN_5_8
 
 class FMaterial;
 class FPrimitiveSceneProxy;
@@ -19,6 +22,11 @@ class FSceneView;
 class FSceneViewFamily;
 class FShaderParametersMetadata;
 class FSceneUniformBuffer;
+class FViewInfo;
+class FGlobalDistanceFieldParameterData;
+class FSceneTextureUniformParameters;
+class FSubstratePublicGlobalUniformParameters;
+class FMobileSceneTextureUniformParameters;
 DECLARE_UNIFORM_BUFFER_STRUCT(FSceneUniformParameters, RENDERER_API)
 
 #if RHI_RAYTRACING
diff --git a/Engine/Source/Runtime/Renderer/Public/GPUSceneWriter.h b/Engine/Source/Runtime/Renderer/Public/GPUSceneWriter.h
--- a/Engine/Source/Runtime/Renderer/Public/GPUSceneWriter.h
+++ b/Engine/Source/Runtime/Renderer/Public/GPUSceneWriter.h
@@ -5,10 +5,13 @@
 #include "CoreMinimal.h"
 #include "ShaderParameterMacros.h"
 #include "ShaderParameters.h"
-#include "MeshBatch.h"
 #include "RenderGraphFwd.h"
 #include "SceneManagement.h"
 
+#if UE_ENABLE_INCLUDE_ORDER_DEPRECATED_IN_5_8
+#include "MeshBatch.h"
+#endif // UE_ENABLE_INCLUDE_ORDER_DEPRECATED_IN_5_8
+
 class FSceneView;
 class FRDGBuilder;
 class FMeshElementCollector;
diff --git a/Engine/Source/Runtime/Renderer/Public/GlobalDistanceFieldParameters.h b/Engine/Source/Runtime/Renderer/Public/GlobalDistanceFieldParameters.h
--- a/Engine/Source/Runtime/Renderer/Public/GlobalDistanceFieldParameters.h
+++ b/Engine/Source/Runtime/Renderer/Public/GlobalDistanceFieldParameters.h
@@ -8,10 +8,13 @@
 #include "ShaderParameterMacros.h"
 #include "ShaderParameters.h"
 #include "ShaderParameterUtils.h"
-#include "RenderUtils.h"
 #include "RHIStaticStates.h"
 #include "RenderGraphResources.h"
 
+#if UE_ENABLE_INCLUDE_ORDER_DEPRECATED_IN_5_8
+#include "RenderUtils.h"
+#endif // UE_ENABLE_INCLUDE_ORDER_DEPRECATED_IN_5_8
+
 class FShaderParameterMap;
 
 class FGlobalDistanceFieldParameterData
diff --git a/Engine/Source/Runtime/Renderer/Public/HairStrandsInterface.h b/Engine/Source/Runtime/Renderer/Public/HairStrandsInterface.h
--- a/Engine/Source/Runtime/Renderer/Public/HairStrandsInterface.h
+++ b/Engine/Source/Runtime/Renderer/Public/HairStrandsInterface.h
@@ -7,23 +7,34 @@
 #pragma once
 
 #include "CoreMinimal.h"
-#include "RendererInterface.h"
 #include "Containers/Array.h"
 #include "Containers/BitArray.h"
-#include "Engine/EngineTypes.h"
+#include "Templates/RefCounting.h"
+#include "RHIFwd.h"
+#include "RenderGraphFwd.h"
 #include "Shader.h"
+#include "RenderGraphDefinitions.h"
+#include "HairStrandsDefinitions.h"
+#include "CachedGeometry.h"
+#include "Serialization/BulkData.h"
+
+#if UE_ENABLE_INCLUDE_ORDER_DEPRECATED_IN_5_8
+#include "RendererInterface.h"
+#include "Engine/EngineTypes.h"
 #include "RenderResource.h"
 #include "RenderGraphResources.h"
 #include "ShaderPrintParameters.h"
 #include "GroomVisualizationData.h"
-#include "HairStrandsDefinitions.h"
-#include "CachedGeometry.h"
-#include "Serialization/BulkData.h"
+#endif // UE_ENABLE_INCLUDE_ORDER_DEPRECATED_IN_5_8
+
 
 class UTexture2D;
 class FSceneInterface;
 class FGlobalShaderMap;
 
+class FViewInfo;
+struct FShaderPrintData;
+
 ////////////////////////////////////////////////////////////////////////////////////////////////////
 // Shader parameters
 
diff --git a/Engine/Source/Runtime/Renderer/Public/PrimitiveSceneInfo.h b/Engine/Source/Runtime/Renderer/Public/PrimitiveSceneInfo.h
--- a/Engine/Source/Runtime/Renderer/Public/PrimitiveSceneInfo.h
+++ b/Engine/Source/Runtime/Renderer/Public/PrimitiveSceneInfo.h
@@ -14,12 +14,15 @@
 #include "PrimitiveComponentId.h"
 #include "PrimitiveDirtyState.h"
 #include "MeshPassProcessor.h"
-#include "MeshBatch.h"
-#include "RendererInterface.h"
 #include "ShaderParameterMacros.h"
 #include "RayTracingMeshDrawCommands.h"
 #include "RayTracingGeometryInstance.h"
 
+#if UE_ENABLE_INCLUDE_ORDER_DEPRECATED_IN_5_8
+#include "MeshBatch.h"
+#include "RendererInterface.h"
+#endif // UE_ENABLE_INCLUDE_ORDER_DEPRECATED_IN_5_8
+
 enum class ERayTracingPrimitiveFlags : uint8;
 
 class FIndirectLightingCacheUniformParameters;
@@ -38,6 +41,7 @@ struct FPrimitiveSceneInfoData;
 struct FInstanceDataBufferHeader;
 class FInstanceSceneDataBuffers;
 class FInstanceDataUpdateTaskInfo;
+struct FMeshBatch;
 
 struct FNaniteMaterialSlot;
 struct FNaniteRasterBin;
diff --git a/Engine/Source/Runtime/Renderer/Public/RayTracingInstance.h b/Engine/Source/Runtime/Renderer/Public/RayTracingInstance.h
--- a/Engine/Source/Runtime/Renderer/Public/RayTracingInstance.h
+++ b/Engine/Source/Runtime/Renderer/Public/RayTracingInstance.h
@@ -2,10 +2,18 @@
 
 #pragma once
 
+#include "Containers/ContainersFwd.h"
+#include "Containers/Array.h"
+#include "Containers/ArrayView.h"
+#include "Math/MathFwd.h"
+
+#if UE_ENABLE_INCLUDE_ORDER_DEPRECATED_IN_5_8
 #include "RenderResource.h"
 #include "MeshBatch.h"
+#endif // UE_ENABLE_INCLUDE_ORDER_DEPRECATED_IN_5_8
 
 class FRayTracingGeometry;
+struct FMeshBatch;
 
 enum class ERayTracingInstanceLayer : uint8
 {
diff --git a/Engine/Source/Runtime/Renderer/Public/SceneRendererInterface.h b/Engine/Source/Runtime/Renderer/Public/SceneRendererInterface.h
--- a/Engine/Source/Runtime/Renderer/Public/SceneRendererInterface.h
+++ b/Engine/Source/Runtime/Renderer/Public/SceneRendererInterface.h
@@ -2,13 +2,18 @@
 
 #pragma once
 
-#include "SceneView.h"
 #include "ShaderParameterStructDeclaration.h"
+#include "RenderGraphFwd.h"
+
+#if UE_ENABLE_INCLUDE_ORDER_DEPRECATED_IN_5_8
+#include "SceneView.h"
+#endif // UE_ENABLE_INCLUDE_ORDER_DEPRECATED_IN_5_8
 
 class FScene;
 class FSceneUniformBuffer;
 class FPrimitiveSceneInfo;
 struct FPersistentPrimitiveIndex;
+class FSceneView;
 
 DECLARE_UNIFORM_BUFFER_STRUCT(FSceneUniformParameters, RENDERER_API)
 
@@ -62,8 +67,4 @@ public:
 	virtual UE::Renderer::Private::IShadowInvalidatingInstances *GetShadowInvalidatingInstancesInterface(const FSceneView *SceneView) { return nullptr; }
 };
 
-inline TRDGUniformBufferRef<FSceneUniformParameters> GetSceneUniformBufferRef(FRDGBuilder& GraphBuilder, const FSceneView& View)
-{
-	ISceneRenderer* SceneRender = View.Family ? View.Family->GetSceneRenderer() : nullptr;
-	return SceneRender ? SceneRender->GetSceneUniformBufferRef(GraphBuilder) : nullptr;
-}
+RENDERER_API TRDGUniformBufferRef<FSceneUniformParameters> GetSceneUniformBufferRef(FRDGBuilder& GraphBuilder, const FSceneView& View);
diff --git a/Engine/Source/Runtime/Renderer/Public/SimpleMeshDrawCommandPass.h b/Engine/Source/Runtime/Renderer/Public/SimpleMeshDrawCommandPass.h
--- a/Engine/Source/Runtime/Renderer/Public/SimpleMeshDrawCommandPass.h
+++ b/Engine/Source/Runtime/Renderer/Public/SimpleMeshDrawCommandPass.h
@@ -6,12 +6,15 @@
 
 #pragma once
 
-#include "MeshMaterialShader.h"
-#include "SceneUtils.h"
 #include "MeshPassProcessor.h"
 #include "InstanceCulling/InstanceCullingContext.h"
 #include "RenderGraphBuilder.h"
 
+#if UE_ENABLE_INCLUDE_ORDER_DEPRECATED_IN_5_8
+#include "MeshMaterialShader.h"
+#include "SceneUtils.h"
+#endif // UE_ENABLE_INCLUDE_ORDER_DEPRECATED_IN_5_8
+
 class FGPUScene;
 class FInstanceCullingDrawParams;
 
diff --git a/Engine/Source/Runtime/Renderer/Public/VariableRateShadingImageManager.h b/Engine/Source/Runtime/Renderer/Public/VariableRateShadingImageManager.h
--- a/Engine/Source/Runtime/Renderer/Public/VariableRateShadingImageManager.h
+++ b/Engine/Source/Runtime/Renderer/Public/VariableRateShadingImageManager.h
@@ -9,12 +9,20 @@
 
 #include "Logging/LogMacros.h"
 #include "RHI.h"
-#include "RendererInterface.h"
 #include "RenderGraphDefinitions.h"
 #include "RenderResource.h"
 
+#if UE_ENABLE_INCLUDE_ORDER_DEPRECATED_IN_5_8
+#include "RendererInterface.h"
+#endif // UE_ENABLE_INCLUDE_ORDER_DEPRECATED_IN_5_8
+
 struct FMinimalSceneTextures; // Forward declaration from SceneTextures.
 
+class FViewInfo;
+class FSceneView;
+class FSceneViewFamily;
+struct IPooledRenderTarget;
+
 class IVariableRateShadingImageGenerator;
 
 DECLARE_LOG_CATEGORY_EXTERN(LogVRS, Log, All);
