HZB | Fix corrupted results for split screen

diff --git a/Engine/Source/Runtime/Renderer/Private/DeferredShadingRenderer.cpp b/Engine/Source/Runtime/Renderer/Private/DeferredShadingRenderer.cpp
--- a/Engine/Source/Runtime/Renderer/Private/DeferredShadingRenderer.cpp
+++ b/Engine/Source/Runtime/Renderer/Private/DeferredShadingRenderer.cpp
@@ -568,11 +568,17 @@ bool FDeferredShadingSceneRenderer::RenderHzb(FRDGBuilder& GraphBuilder, FRDGTex
 			FRDGTextureRef ClosestHZBTexture = nullptr;
 			FRDGTextureRef FurthestHZBTexture = nullptr;
 
+			// In split-screen, multiple views share a depth buffer. UpperBoundViewRect (used by
+			// GetHZBViewRect) can extend beyond a view's rendered depth region into an adjacent
+			// view's area, contaminating the HZB mipchain. Clamp to ViewRect in multi-view to
+			// prevent reading another view's depth. See UE-347165.
+			const FIntRect HZBViewRect = Views.Num() > 1 ? View.ViewRect : GetHZBViewRect(View);
+
 			BuildHZB(
 				GraphBuilder,
 				SceneDepthTexture,
 				/* VisBufferTexture = */ nullptr,
-				GetHZBViewRect(View),
+				HZBViewRect,
 				View.GetFeatureLevel(),
 				View.GetShaderPlatform(),
 				TEXT("HZBClosest"),
diff --git a/Engine/Source/Runtime/Renderer/Private/MobileShadingRenderer.cpp b/Engine/Source/Runtime/Renderer/Private/MobileShadingRenderer.cpp
--- a/Engine/Source/Runtime/Renderer/Private/MobileShadingRenderer.cpp
+++ b/Engine/Source/Runtime/Renderer/Private/MobileShadingRenderer.cpp
@@ -3428,11 +3428,17 @@ void FMobileSceneRenderer::RenderHZB(FRDGBuilder& GraphBuilder, FRDGTextureRef S
 
 				FRDGTextureRef FurthestHZBTexture = nullptr;
 
+				// In split-screen, multiple views share a depth buffer. UpperBoundViewRect (used by
+				// GetHZBViewRect) can extend beyond a view's rendered depth region into an adjacent
+				// view's area, contaminating the HZB mipchain. Clamp to ViewRect in multi-view to
+				// prevent reading another view's depth. See UE-347165.
+				const FIntRect HZBViewRect = Views.Num() > 1 ? View.ViewRect : GetHZBViewRect(View);
+
 				BuildHZBFurthest(
 					GraphBuilder,
 					SceneDepthTexture,
 					/* VisBufferTexture = */ nullptr,
-					GetHZBViewRect(View),
+					HZBViewRect,
 					View.GetFeatureLevel(),
 					View.GetShaderPlatform(),
 					TEXT("MobileHZBFurthest"),
