Add PrimitiveComponentIdToInfoMap for scene lookup by component ID in GetPrimitiveSceneInfo(FPrimitiveComponentId), replacing a linear scan over all primitives.

#Synced-CL 51926088

diff --git a/Engine/Source/Runtime/Engine/Public/PrimitiveComponentId.h b/Engine/Source/Runtime/Engine/Public/PrimitiveComponentId.h
--- a/Engine/Source/Runtime/Engine/Public/PrimitiveComponentId.h
+++ b/Engine/Source/Runtime/Engine/Public/PrimitiveComponentId.h
@@ -25,6 +25,11 @@ public:
 		return PrimIDValue == OtherId.PrimIDValue;
 	}
 
+	inline bool operator!=(FPrimitiveComponentId OtherId) const
+	{
+		return PrimIDValue != OtherId.PrimIDValue;
+	}
+
 	friend uint32 GetTypeHash(FPrimitiveComponentId Id)
 	{
 		return GetTypeHash(Id.PrimIDValue);
diff --git a/Engine/Source/Runtime/Renderer/Private/PrimitiveSceneInfo.cpp b/Engine/Source/Runtime/Renderer/Private/PrimitiveSceneInfo.cpp
--- a/Engine/Source/Runtime/Renderer/Private/PrimitiveSceneInfo.cpp
+++ b/Engine/Source/Runtime/Renderer/Private/PrimitiveSceneInfo.cpp
@@ -1962,6 +1962,7 @@ void FPrimitiveSceneInfo::AddToScene(FScene* Scene, TArrayView<FPrimitiveSceneIn
 
 			// Store the component.
 			Scene->PrimitiveComponentIds[PackedIndex] = SceneInfo->PrimitiveComponentId;
+			Scene->PrimitiveComponentIdToInfoMap.Add(SceneInfo->PrimitiveComponentId, SceneInfo);
 
 #if RHI_RAYTRACING
 			// Set group id
diff --git a/Engine/Source/Runtime/Renderer/Private/RendererScene.cpp b/Engine/Source/Runtime/Renderer/Private/RendererScene.cpp
--- a/Engine/Source/Runtime/Renderer/Private/RendererScene.cpp
+++ b/Engine/Source/Runtime/Renderer/Private/RendererScene.cpp
@@ -2015,7 +2015,8 @@ FPrimitiveSceneInfo* FScene::GetPrimitiveSceneInfo(int32 PrimitiveIndex) const
 
 FPrimitiveSceneInfo* FScene::GetPrimitiveSceneInfo(FPrimitiveComponentId PrimitiveId) const
 {
-	return GetPrimitiveSceneInfo(PrimitiveComponentIds.Find(PrimitiveId));
+	FPrimitiveSceneInfo* const* Found = PrimitiveComponentIdToInfoMap.Find(PrimitiveId);
+	return Found ? *Found : nullptr;
 }
 
 FPrimitiveSceneInfo* FScene::GetPrimitiveSceneInfo(const FPersistentPrimitiveIndex& PersistentPrimitiveIndex) const
@@ -5655,6 +5656,8 @@ void FScene::Update(FRDGBuilder& GraphBuilder, const FUpdateParameters& Paramete
 					VelocityData.RemoveFromScene(PrimitiveSceneInfo->PrimitiveComponentId, false);
 				}
 
+				PrimitiveComponentIdToInfoMap.Remove(PrimitiveSceneInfo->PrimitiveComponentId);
+
 				// Unlink the primitive from its shadow parent.
 				PrimitiveSceneInfo->UnlinkAttachmentGroup();
 
diff --git a/Engine/Source/Runtime/Renderer/Private/ScenePrivate.h b/Engine/Source/Runtime/Renderer/Private/ScenePrivate.h
--- a/Engine/Source/Runtime/Renderer/Private/ScenePrivate.h
+++ b/Engine/Source/Runtime/Renderer/Private/ScenePrivate.h
@@ -1685,6 +1685,8 @@ public:
 	TScenePrimitiveArray<FBoxSphereBounds> PrimitiveOcclusionBounds;
 	/** Packed array of primitive components associated with the primitive. */
 	TArray<FPrimitiveComponentId> PrimitiveComponentIds;
+	/** Map from FPrimitiveComponentId to FPrimitiveSceneInfo*. */
+	TMap<FPrimitiveComponentId, FPrimitiveSceneInfo*> PrimitiveComponentIdToInfoMap;
 #if RHI_RAYTRACING
 	/** Packed array of ray tracing primitive caching flags*/
 	TArray<ERayTracingPrimitiveFlags> PrimitiveRayTracingFlags;
diff --git a/Engine/Source/Runtime/Renderer/Private/SceneVisibility.cpp b/Engine/Source/Runtime/Renderer/Private/SceneVisibility.cpp
--- a/Engine/Source/Runtime/Renderer/Private/SceneVisibility.cpp
+++ b/Engine/Source/Runtime/Renderer/Private/SceneVisibility.cpp
@@ -6166,10 +6166,10 @@ void FLODSceneTree::AddChildNode(const FPrimitiveComponentId ParentId, FPrimitiv
 			// Create parent SceneNode, assign correct SceneInfo
 			Parent = &SceneNodes.Add(ParentId, FLODSceneNode());
 
-			int32 ParentIndex = Scene->PrimitiveComponentIds.Find(ParentId);
-			if (ParentIndex != INDEX_NONE)
+			FPrimitiveSceneInfo* ParentSceneInfo = Scene->GetPrimitiveSceneInfo(ParentId);
+			if (ParentSceneInfo)
 			{
-				Parent->SceneInfo = Scene->Primitives[ParentIndex];
+				Parent->SceneInfo = ParentSceneInfo;
 				check(Parent->SceneInfo->PrimitiveComponentId == ParentId);
 			}
 		}
