Support for applying a post process material to a runtime reflection capture (a.k.a. Diffuse Ambient Capture), to allow artists to fine tune the result (for example, color adjustment, directional tinting, sky brightness adjustment).  Note that by design, runtime reflection captures currently only run on mobile, and so this feature has only been implemented on mobile.  In the future, debug visualization of Diffuse Ambient Captures in the deferred renderer may be implemented, at which point features like this will be ported over.

#rnx

#Synced-CL 51952831

diff --git a/Engine/Source/Runtime/Engine/Classes/Components/PlaneReflectionCaptureComponent.h b/Engine/Source/Runtime/Engine/Classes/Components/PlaneReflectionCaptureComponent.h
--- a/Engine/Source/Runtime/Engine/Classes/Components/PlaneReflectionCaptureComponent.h
+++ b/Engine/Source/Runtime/Engine/Classes/Components/PlaneReflectionCaptureComponent.h
@@ -9,7 +9,7 @@
 #include "PlaneReflectionCaptureComponent.generated.h"
 
 	// -> will be exported to EngineDecalClasses.h
-UCLASS(hidecategories=(Collision, Object, Physics, SceneComponent), MinimalAPI)
+UCLASS(hidecategories=(Collision, Object, Physics, SceneComponent, RuntimeCapture), MinimalAPI)
 class UPlaneReflectionCaptureComponent : public UReflectionCaptureComponent
 {
 	GENERATED_UCLASS_BODY()
diff --git a/Engine/Source/Runtime/Engine/Classes/Components/ReflectionCaptureComponent.h b/Engine/Source/Runtime/Engine/Classes/Components/ReflectionCaptureComponent.h
--- a/Engine/Source/Runtime/Engine/Classes/Components/ReflectionCaptureComponent.h
+++ b/Engine/Source/Runtime/Engine/Classes/Components/ReflectionCaptureComponent.h
@@ -65,10 +65,14 @@ class UReflectionCaptureComponent : public USceneComponent
 	/** Cubemap texture resource used for rendering with the encoded HDR values. */
 	class FReflectionTextureCubeResource* EncodedHDRCubemapTexture;
 
-	/** Reflection capture is generated at runtime.  Ignored for planar reflections.  If CVar r.ReflectionCapture.Runtime==2, all captures are runtime, regardless of this setting. */
-	UPROPERTY(EditAnywhere, BlueprintReadOnly, Category = ReflectionCapture)
+	/** Reflection capture is generated at runtime.  If CVar r.ReflectionCapture.Runtime==2, all captures are runtime, regardless of this setting. */
+	UPROPERTY(EditAnywhere, BlueprintReadOnly, Category=RuntimeCapture)
 	bool bRuntimeCapture;
 
+	/** Post process material applied to runtime reflection captures.  Must use Blendable Location "Scene Color Before Bloom". */
+	UPROPERTY(EditAnywhere, BlueprintReadWrite, Category=RuntimeCapture, meta=(EditCondition="bRuntimeCapture", EditConditionHides))
+	TObjectPtr<class UMaterialInterface> PostProcessMaterial;
+
 #if WITH_EDITOR
 	/** Check to see if MapBuildDataId was loaded - otherwise we need to display a warning on cook */
 	bool bMapBuildDataIdLoaded;
diff --git a/Engine/Source/Runtime/Engine/Private/Materials/Material.cpp b/Engine/Source/Runtime/Engine/Private/Materials/Material.cpp
--- a/Engine/Source/Runtime/Engine/Private/Materials/Material.cpp
+++ b/Engine/Source/Runtime/Engine/Private/Materials/Material.cpp
@@ -2056,7 +2056,23 @@ void UMaterialInterface::OverrideBlendableSettings(class FSceneView& View, float
 
 	//	should we use UMaterial::GetDefaultMaterial(Domain) instead of skipping the material
 
-	if(!Base || Base->MaterialDomain != MD_PostProcess || !View.State)
+	if (!Base || Base->MaterialDomain != MD_PostProcess)
+	{
+		return;
+	}
+
+	if (View.IsRuntimeReflectionCapture())
+	{
+		// Runtime reflection captures support applying a post process material, even when no View.State is set.
+		// These can't use a MID because they have no View.State to store it on, but since it's a singleton material,
+		// it can't blend, and so it doesn't need one.  The "const_cast" is unavoidable here, as FPostProcessMaterialNode
+		// supports modification of the material interface when it's a MID, forcing the member to be non-const, but it
+		// won't actually be modified when not a MID.
+		FPostProcessMaterialNode MaterialNode(const_cast<UMaterialInterface*>(this), GetBlendableLocation(Base), 0, false);
+		Dest.BlendableManager.PushBlendableData(1.0f, MaterialNode);
+		return;
+	}
+	else if(!View.State)
 	{
 		return;
 	}
diff --git a/Engine/Source/Runtime/Engine/Private/SceneView.cpp b/Engine/Source/Runtime/Engine/Private/SceneView.cpp
--- a/Engine/Source/Runtime/Engine/Private/SceneView.cpp
+++ b/Engine/Source/Runtime/Engine/Private/SceneView.cpp
@@ -1565,7 +1565,7 @@ bool FSceneView::ProjectWorldToScreen(const FVector& WorldPosition, const FIntRe
 #define IF_PP(NAME)   if(Src.bOverride_ ## NAME && Src . NAME)
 
 // @param Weight 0..1
-void FSceneView::OverridePostProcessSettings(const FPostProcessSettings& Src, float Weight)
+void FSceneView::OverridePostProcessSettings(const FPostProcessSettings& Src, float Weight, bool bAllowBlendables)
 {
 	SCOPE_CYCLE_COUNTER(STAT_OverridePostProcessSettings);
 
@@ -1932,6 +1932,7 @@ void FSceneView::OverridePostProcessSettings(const FPostProcessSettings& Src, fl
 	}
 
 	// Blendable objects
+	if (bAllowBlendables)
 	{
 		uint32 Count = Src.WeightedBlendables.Array.Num();
 
diff --git a/Engine/Source/Runtime/Engine/Private/World.cpp b/Engine/Source/Runtime/Engine/Private/World.cpp
--- a/Engine/Source/Runtime/Engine/Private/World.cpp
+++ b/Engine/Source/Runtime/Engine/Private/World.cpp
@@ -10614,7 +10614,13 @@ static void DoPostProcessVolume(IInterface_PostProcessVolume* Volume, FVector Vi
 
 	if (LocalWeight > 0)
 	{
-		SceneView->OverridePostProcessSettings(*VolumeProperties.Settings, LocalWeight);
+		// Suppress blendables from post process volumes for runtime reflection captures.  Runtime reflection captures don't run post processing,
+		// and only support a singleton post process material to allow artists to apply things like color or directional tint adjustments.  Runtime
+		// reflection captures don't have an FSceneViewState to store MIDs, which would be required to handle generalized post process material
+		// blending, but generally we also don't want other post process materials by design.
+		const bool bAllowBlendables = !SceneView->IsRuntimeReflectionCapture();
+
+		SceneView->OverridePostProcessSettings(*VolumeProperties.Settings, LocalWeight, bAllowBlendables);
 	}
 
 #if DEBUG_POST_PROCESS_VOLUME_ENABLE
diff --git a/Engine/Source/Runtime/Engine/Public/SceneView.h b/Engine/Source/Runtime/Engine/Public/SceneView.h
--- a/Engine/Source/Runtime/Engine/Public/SceneView.h
+++ b/Engine/Source/Runtime/Engine/Public/SceneView.h
@@ -2058,7 +2058,7 @@ public:
 	 * custom layers can be combined with the existing settings
 	 * @param Weight usually 0..1 but outside range is clamped
 	 */
-	ENGINE_API void OverridePostProcessSettings(const FPostProcessSettings& Src, float Weight);
+	ENGINE_API void OverridePostProcessSettings(const FPostProcessSettings& Src, float Weight, bool bAllowBlendables = true);
 
 	/** applied global restrictions from show flags */
 	ENGINE_API void EndFinalPostprocessSettings(const FSceneViewInitOptions& ViewInitOptions);
diff --git a/Engine/Source/Runtime/Renderer/Private/MobileShadingRenderer.cpp b/Engine/Source/Runtime/Renderer/Private/MobileShadingRenderer.cpp
--- a/Engine/Source/Runtime/Renderer/Private/MobileShadingRenderer.cpp
+++ b/Engine/Source/Runtime/Renderer/Private/MobileShadingRenderer.cpp
@@ -1811,6 +1811,13 @@ void FMobileSceneRenderer::Render(FRDGBuilder& GraphBuilder, const FSceneRenderU
 		FRendererModule& RendererModule = static_cast<FRendererModule&>(GetRendererModule());
 		RendererModule.RenderPostOpaqueExtensions(GraphBuilder, Views, SceneTextures);
 
+		if (Views[0].IsRuntimeReflectionCapture())
+		{
+			// Needs to run before CopySceneCaptureComponentToTarget, as it disables the scene color copy that occurs there,
+			// if a post pass has already handled copying to the final output.
+			AddRuntimeReflectionCapturePostPass(GraphBuilder, ViewFamily, ViewFamilyTexture, ViewFamilyDepthTexture, Views);
+		}
+
 		CopySceneCaptureComponentToTarget(GraphBuilder, SceneTextures, ViewFamilyTexture, ViewFamilyDepthTexture, ViewFamily, Views);
 
 		RenderOpaqueFX(GraphBuilder, GetSceneViews(), GetSceneUniforms(), FXSystem, SceneTextures.MobileUniformBuffer);
diff --git a/Engine/Source/Runtime/Renderer/Private/PostProcess/PostProcessMaterial.cpp b/Engine/Source/Runtime/Renderer/Private/PostProcess/PostProcessMaterial.cpp
--- a/Engine/Source/Runtime/Renderer/Private/PostProcess/PostProcessMaterial.cpp
+++ b/Engine/Source/Runtime/Renderer/Private/PostProcess/PostProcessMaterial.cpp
@@ -994,6 +994,13 @@ FScreenPassTexture AddPostProcessMaterialPass(
 
 static bool IsPostProcessMaterialsEnabledForView(const FViewInfo& View)
 {
+	// Runtime reflection captures have a special case to support a singleton post processing material, even though
+	// post processing isn't enabled in the show flags.
+	if (View.IsRuntimeReflectionCapture())
+	{
+		return CVarPostProcessingDisableMaterials.GetValueOnRenderThread() == 0;
+	}
+
 	if (!View.Family->EngineShowFlags.PostProcessing ||
 		!View.Family->EngineShowFlags.PostProcessMaterial ||
 		View.Family->EngineShowFlags.VisualizeShadingModels ||
diff --git a/Engine/Source/Runtime/Renderer/Private/PostProcess/PostProcessing.cpp b/Engine/Source/Runtime/Renderer/Private/PostProcess/PostProcessing.cpp
--- a/Engine/Source/Runtime/Renderer/Private/PostProcess/PostProcessing.cpp
+++ b/Engine/Source/Runtime/Renderer/Private/PostProcess/PostProcessing.cpp
@@ -3748,6 +3748,50 @@ void AddMobilePostProcessingPasses(FRDGBuilder& GraphBuilder, FScene* Scene, con
 	}
 }
 
+void AddRuntimeReflectionCapturePostPass(FRDGBuilder& GraphBuilder, FViewFamilyInfo& ViewFamily, FRDGTextureRef ViewFamilyTexture, FRDGTextureRef ViewFamilyDepthTexture, TConstArrayView<FViewInfo> Views)
+{
+	// We are assuming this logic replaces regular post processing -- make sure that won't be running later.
+	check(ViewFamily.bResolveScene == false);
+
+	FMobilePostProcessingInputs Inputs;
+	Inputs.ViewFamilyTexture = ViewFamilyTexture;
+	Inputs.ViewFamilyDepthTexture = ViewFamilyDepthTexture;
+	Inputs.SceneTextures = CreateMobileSceneTextureUniformBuffer(GraphBuilder, &ViewFamily.GetSceneTextures(), EMobileSceneTextureSetupMode::All);
+
+	for (int32 ViewIndex = 0; ViewIndex < Views.Num(); ViewIndex++)
+	{
+		const FViewInfo& View = Views[ViewIndex];
+
+		// Runtime reflection captures apply only BL_SceneColorBeforeBloom post process materials.
+		const FPostProcessMaterialChain MaterialChain = GetPostProcessMaterialChain(View, BL_SceneColorBeforeBloom);
+
+		if (MaterialChain.Num())
+		{
+			const FScreenPassRenderTarget ViewFamilyOutput = FScreenPassRenderTarget::CreateViewFamilyOutput(Inputs.ViewFamilyTexture, View);
+			const FScreenPassTexture CustomDepth((*Inputs.SceneTextures)->CustomDepthTexture, View.ViewRect);
+			FScreenPassTexture SceneColor((*Inputs.SceneTextures)->SceneColorTexture, View.ViewRect);
+
+			FPostProcessMaterialInputs PostProcessMaterialInputs;
+			PostProcessMaterialInputs.SetInput(GraphBuilder, EPostProcessMaterialInput::SceneColor, SceneColor);
+			PostProcessMaterialInputs.CustomDepthTexture = CustomDepth.Texture;
+			PostProcessMaterialInputs.SceneTextures = GetSceneTextureShaderParameters(Inputs.SceneTextures);
+			SceneColor = AddPostProcessMaterialChain(GraphBuilder, View, ViewIndex, PostProcessMaterialInputs, MaterialChain);
+
+			// Copy the resulting post processed texture to the output.  Note that it doesn't work to directly write to ViewFamilyOutput via
+			// PostProcessMaterialInputs.OverrideOutput, because the output rect generally won't be the same as the scene texture rect where
+			// rendering occurred, as those are compacted to the upper left corner.  Various SvPosition related math gets messed up.  You
+			// can't fix up the view uniforms to point to the final output rect either, because it messes up the math for certain inputs.
+			// A bunch of stuff would need to be cleaned up to make the math consistent, but this works fine, and the texture in question is
+			// low resolution, and cheap to copy.
+			AddDrawTexturePass(GraphBuilder, View, SceneColor, ViewFamilyOutput);
+
+			// Override the SceneCaptureSource to communicate that final color has been written by a post pass.
+			// This prevents CopySceneCaptureComponentToTarget from overwriting it with SceneColor later.
+			ViewFamily.SceneCaptureSource = SCS_FinalColorHDR;
+		}
+	}
+}
+
 FRDGTextureRef AddProcessPlanarReflectionPass(
 	FRDGBuilder& GraphBuilder,
 	const FViewInfo& View,
diff --git a/Engine/Source/Runtime/Renderer/Private/PostProcess/PostProcessing.h b/Engine/Source/Runtime/Renderer/Private/PostProcess/PostProcessing.h
--- a/Engine/Source/Runtime/Renderer/Private/PostProcess/PostProcessing.h
+++ b/Engine/Source/Runtime/Renderer/Private/PostProcess/PostProcessing.h
@@ -17,6 +17,7 @@ enum class EReflectionsMethod;
 enum class EDiffuseIndirectMethod;
 
 class FScreenPassVS;
+class FViewFamilyInfo;
 class FViewInfo;
 class FVirtualShadowMapArray;
 struct FLumenSceneFrameTemporaries;
@@ -90,3 +91,5 @@ void AddMobilePostProcessingPasses(FRDGBuilder& GraphBuilder, FScene* Scene, con
 void AddBasicPostProcessPasses(FRDGBuilder& GraphBuilder, const FViewInfo& View);
 
 FRDGTextureRef AddProcessPlanarReflectionPass(FRDGBuilder& GraphBuilder, const FViewInfo& View, FRDGTextureRef SceneColorTexture);
+
+void AddRuntimeReflectionCapturePostPass(FRDGBuilder& GraphBuilder, FViewFamilyInfo& ViewFamily, FRDGTextureRef ViewFamilyTexture, FRDGTextureRef ViewFamilyDepthTexture, TConstArrayView<FViewInfo> Views);
diff --git a/Engine/Source/Runtime/Renderer/Private/ReflectionEnvironmentCapture.cpp b/Engine/Source/Runtime/Renderer/Private/ReflectionEnvironmentCapture.cpp
--- a/Engine/Source/Runtime/Renderer/Private/ReflectionEnvironmentCapture.cpp
+++ b/Engine/Source/Runtime/Renderer/Private/ReflectionEnvironmentCapture.cpp
@@ -33,6 +33,8 @@
 #include "HAL/IConsoleManager.h"
 #include "LegacyScreenPercentageDriver.h"
 #include "Logging/LogMacros.h"
+#include "Materials/Material.h"
+#include "Materials/MaterialInterface.h"
 #include "Math/Color.h"
 #include "Math/Float16Color.h"
 #include "Math/SHMath.h"
@@ -1291,6 +1293,26 @@ void FScene::AllocateReflectionCaptures(const TArray<UReflectionCaptureComponent
 					SceneCaptureComponent->ReflectionTimeSliceCount = RuntimeCaptureTimesliceCount;
 					SceneCaptureComponent->SkylightScale = RuntimeCaptureToUpdate->RuntimeSkylightScale;
 
+					// Propagate post process material to the scene capture's PostProcessSettings
+					SceneCaptureComponent->PostProcessSettings.WeightedBlendables.Array.Reset();
+					SceneCaptureComponent->PostProcessBlendWeight = 0.0f;
+
+					if (RuntimeCaptureToUpdate->PostProcessMaterial)
+					{
+						UMaterialInterface* PostMaterial = RuntimeCaptureToUpdate->PostProcessMaterial;
+						UMaterial* BaseMaterial = PostMaterial->GetMaterial();
+
+						if (BaseMaterial->IsPostProcessMaterial() &&
+							PostMaterial->GetUserSceneTextureOutput(BaseMaterial).IsNone() &&
+							PostMaterial->GetBlendableLocation(BaseMaterial) == BL_SceneColorBeforeBloom)
+						{
+							FWeightedBlendable& Blendable = SceneCaptureComponent->PostProcessSettings.WeightedBlendables.Array.AddDefaulted_GetRef();
+							Blendable.Weight = 1.0f;
+							Blendable.Object = PostMaterial;
+							SceneCaptureComponent->PostProcessBlendWeight = 1.0f;
+						}
+					}
+
 					// Set these every frame, to allow for the CVars to be changed at runtime
 					SceneCaptureComponent->ShowFlags.InstancedFoliage = CVarReflectionCaptureRuntimeFoliage.GetValueOnGameThread();
 					SceneCaptureComponent->ShowFlags.InstancedGrass = SceneCaptureComponent->ShowFlags.InstancedFoliage;
diff --git a/Engine/Source/Runtime/Renderer/Private/ScreenPass.cpp b/Engine/Source/Runtime/Renderer/Private/ScreenPass.cpp
--- a/Engine/Source/Runtime/Renderer/Private/ScreenPass.cpp
+++ b/Engine/Source/Runtime/Renderer/Private/ScreenPass.cpp
@@ -181,7 +181,7 @@ FScreenPassRenderTarget FScreenPassRenderTarget::CreateViewFamilyOutput(FRDGText
 
 	ERenderTargetLoadAction LoadAction = ERenderTargetLoadAction::ENoAction;
 
-	if (!View.IsFirstInFamily() || View.Family->bAdditionalViewFamily)
+	if (!View.IsFirstInFamily() || View.Family->bAdditionalViewFamily || View.IsRuntimeReflectionCapture())
 	{
 		LoadAction = ERenderTargetLoadAction::ELoad;
 	}
