[SceneExtensions] Fix dll import linking errors with the GetExtensionID method.

#rnx
#Synced-CL 51901069

diff --git a/Engine/Source/Runtime/Renderer/Private/SceneExtensions.h b/Engine/Source/Runtime/Renderer/Private/SceneExtensions.h
--- a/Engine/Source/Runtime/Renderer/Private/SceneExtensions.h
+++ b/Engine/Source/Runtime/Renderer/Private/SceneExtensions.h
@@ -402,7 +402,7 @@ public:
 /** Use these macros in the class definitions of your extension. */
 #define DECLARE_SCENE_EXTENSION(ModuleExport, ClassName) \
 	public: \
-		ModuleExport static int32 GetExtensionID() { return ExtensionRegistration.GetExtensionID();  } \
+		ModuleExport static int32 GetExtensionID(); \
 	private: \
 		ModuleExport static TSceneExtensionRegistration<ClassName> ExtensionRegistration
 
@@ -416,4 +416,5 @@ public:
 
 /** Use this macro in the implementation source file of your extension. */
 #define IMPLEMENT_SCENE_EXTENSION(ClassName) \
+	int32 ClassName::GetExtensionID() { return ExtensionRegistration.GetExtensionID();  } \
 	TSceneExtensionRegistration<ClassName> ClassName::ExtensionRegistration
