Fixed crash due to NumAnimBones being smaller than the number of bones in the skeleton (due to virtual bones).

#Synced-CL 51928592

diff --git a/Engine/Source/Runtime/Renderer/Private/Skinning/AnimSequenceTransformProvider.cpp b/Engine/Source/Runtime/Renderer/Private/Skinning/AnimSequenceTransformProvider.cpp
--- a/Engine/Source/Runtime/Renderer/Private/Skinning/AnimSequenceTransformProvider.cpp
+++ b/Engine/Source/Runtime/Renderer/Private/Skinning/AnimSequenceTransformProvider.cpp
@@ -1547,7 +1547,7 @@ void FAnimSequenceTransformProvider::ProvideGPUTransforms(FSkinningTransformProv
 					const FCompressedAnimSequence& PlatformCompressedData = CompressedAnimSequence.Get();
 
 					const TArray<FTransform>& AnimRefLocalPoses = AnimSequence->GetSkeleton()->GetRefLocalPoses();
-					check(Sequence.NumAnimBones == AnimRefLocalPoses.Num());
+					check(Sequence.NumAnimBones <= (uint16)AnimRefLocalPoses.Num());
 
 					const int32 NumTracks = PlatformCompressedData.CompressedTrackToSkeletonMapTable.Num();
 					TrackToBoneIndexMap.Reset(NumTracks);
@@ -1595,7 +1595,7 @@ void FAnimSequenceTransformProvider::ProvideGPUTransforms(FSkinningTransformProv
 
 						TArrayView<FBoneTransformWithScale> DstTransforms = ScatterUploader.Add_GetRef<FBoneTransformWithScale>(Sequence.BufferAllocation.Offset + Sequence.NumAnimBones * FrameIndex, Sequence.NumAnimBones);
 
-						for (int32 BoneIndex = 0; BoneIndex < AnimLocalPoses.Num(); ++BoneIndex)
+						for (uint16 BoneIndex = 0; BoneIndex < Sequence.NumAnimBones; ++BoneIndex)
 						{
 							DstTransforms[BoneIndex].SetTransform(FTransform3f(AnimLocalPoses[BoneIndex]));
 						}
