Fixed FGBufferPickingCS not compiling on some non-PC platforms (causing to high register pressure)

#tests PC editor dx12 - test should compile enginetest succesfully.
#rnx
#Synced-CL 51999473

diff --git a/Engine/Source/Runtime/Renderer/Private/PostProcess/PostProcessing.cpp b/Engine/Source/Runtime/Renderer/Private/PostProcess/PostProcessing.cpp
--- a/Engine/Source/Runtime/Renderer/Private/PostProcess/PostProcessing.cpp
+++ b/Engine/Source/Runtime/Renderer/Private/PostProcess/PostProcessing.cpp
@@ -4227,12 +4227,12 @@ class FGBufferPickingCS : public FGlobalShader
 	END_SHADER_PARAMETER_STRUCT()
 
 public:
-	static bool IsSupported(EShaderPlatform Platform) 
-	{ 
-		return 
+	static bool IsSupported(EShaderPlatform Platform)
+	{
+		return
+			IsPCPlatform(Platform) && // Only on PC because some platforms have trouble to compile this shader. See JIRA UE-370697.
 			IsFeatureLevelSupported(Platform, ERHIFeatureLevel::SM5) &&
-			ShaderPrint::IsSupported(Platform) && 
-			!IsMobilePlatform(Platform) && 
+			ShaderPrint::IsSupported(Platform) &&
 			(!Substrate::IsSubstrateEnabled() || Substrate::IsSubstrateBlendableGBufferEnabled(Platform));
 	}
 	static bool ShouldCompilePermutation(const FGlobalShaderPermutationParameters& Parameters) 
