Fix crash when previewing SM5 Vulkan in editor with Substrate Adaptive GBuffer.

#Synced-CL 51943914

diff --git a/Engine/Source/Runtime/Renderer/Private/StochasticLighting/StochasticLighting.cpp b/Engine/Source/Runtime/Renderer/Private/StochasticLighting/StochasticLighting.cpp
--- a/Engine/Source/Runtime/Renderer/Private/StochasticLighting/StochasticLighting.cpp
+++ b/Engine/Source/Runtime/Renderer/Private/StochasticLighting/StochasticLighting.cpp
@@ -645,7 +645,7 @@ StochasticLighting::FFrameTemporaries StochasticLighting::TileClassificationMark
 			VisualizeExtent,
 			TEXT("StochasticLighting.NormalAndShadingInfoHistory"));
 
-		if (Substrate::IsSubstrateEnabled() && !Substrate::IsSubstrateBlendableGBufferEnabled(GMaxRHIShaderPlatform))
+		if (Substrate::IsSubstrateEnabled() && !Substrate::IsSubstrateBlendableGBufferEnabled(ReferenceView.GetShaderPlatform()))
 		{
 			ClosureHistory = FrameTemporaries.ClosureHistory.CreateSharedRT(GraphBuilder,
 				FRDGTextureDesc::Create2D(SceneTextures.Config.Extent, PF_R8_UINT, FClearValueBinding::Black, TexCreate_ShaderResource | TexCreate_UAV),
