[Light shape visualization] Visualize light shapes by drawing translucent wireframes conformed to the shapes of lights

#tests tested in editor and on a console
#Synced-CL 51949275

diff --git a/Engine/Shaders/Private/Tools/VisualizeLightShape.usf b/Engine/Shaders/Private/Tools/VisualizeLightShape.usf
new file mode 100644
--- /dev/null
+++ b/Engine/Shaders/Private/Tools/VisualizeLightShape.usf
@@ -0,0 +1,139 @@
+// Copyright Epic Games, Inc. All Rights Reserved.
+
+/*=============================================================================
+	VisualizeLightShape.usf: Shaders for visualizing light shapes
+=============================================================================*/
+
+#include "../Common.ush"
+#include "../Hash.ush"
+#include "../Visualization.ush"
+#include "../MonteCarlo.ush"
+#include "../SceneTexturesCommon.ush"
+#include "../LightGridCommon.ush"
+
+float3 SafeNormalize(float3 V)
+{
+	float L2 = dot(V, V);
+	float InvLength = L2 > 1e-8 ? rsqrt(L2) : 0.0;
+	return V * InvLength;
+}
+
+float3 GetLightShapeVertexTranslatedWorldPosition(float3 LocalPosition, FDeferredLightData LightData)
+{
+	float LightRadius = rcp(LightData.InvRadius);
+	float3 OutPosition;
+
+#if LIGHT_TYPE == LIGHT_TYPE_POINT
+	OutPosition = LocalPosition * LightRadius + LightData.TranslatedWorldPosition;
+#elif LIGHT_TYPE == LIGHT_TYPE_SPOT
+	float ConeAngle = acos(LightData.SpotAngles.x) * LocalPosition.x;
+
+	float SinConeAngle, CosConeAngle;
+	sincos(ConeAngle, SinConeAngle, CosConeAngle);
+
+	OutPosition = float3(CosConeAngle, LocalPosition.yz * SinConeAngle) * LightRadius;
+	OutPosition = select(LocalPosition.xxx > 1.0001, 0.0, OutPosition);
+
+	float3x3 LocalToWorldRotation = GetTangentBasis(-LightData.Direction);
+	OutPosition = mul(OutPosition.zyx, LocalToWorldRotation) + LightData.TranslatedWorldPosition;
+#else // LIGHT_TYPE_RECT
+	float Width = LightData.SourceRadius * 2.0;
+	float Height = LightData.SourceLength * 2.0;
+	float CosBarnAngle = LightData.RectLightData.BarnCosAngle;
+	float SinBarnAngle = sqrt(1.0 - CosBarnAngle * CosBarnAngle);
+	float BarnLength = LightData.RectLightData.BarnLength;
+	
+	float3 LightForward = float3(1, 0, 0);
+	float3 NormPos = LocalPosition * 2.0;
+	float3 BaseRect = LocalPosition * float3(0, Width, Height);
+	float3 BarnRect = BaseRect + float3(CosBarnAngle, SinBarnAngle * NormPos.yz) * BarnLength;
+
+	// Project onto bounding sphere by intersecting the projection ray with the sphere
+	// (i.e. substitute x^2 + y^2 + z^2 = r2 with o + dt and solve for t)
+	float3 Dir = BarnRect + BaseRect;
+	float A = dot(Dir, Dir);
+	float B = -2.0 * dot(Dir, BaseRect);
+	float C = dot(BaseRect, BaseRect) - LightRadius * LightRadius;
+	float T = (-B + sqrt(B * B - 4.0 * A * C)) / (2.0 * A);
+	float3 OuterRect = -BaseRect + Dir * T;
+
+	// Rotate towards the intersection of light forward and the bounding sphere
+	Dir = normalize(OuterRect);
+	float AngleBetween = acos(dot(Dir, LightForward)) * NormPos.x;
+	float SinAngleBetween, CosAngleBetween;
+	sincos(AngleBetween, SinAngleBetween, CosAngleBetween);
+	float3 RotationAxis = SafeNormalize(cross(Dir, LightForward));
+
+	OutPosition = OuterRect * CosAngleBetween + cross(RotationAxis, OuterRect) * SinAngleBetween;
+	OutPosition = select(NormPos.xxx < -0.0001, BarnRect, OutPosition);
+
+	float3x3 LightToWorldRotation = float3x3(-LightData.Direction, cross(LightData.Tangent, -LightData.Direction), LightData.Tangent);
+	OutPosition = mul(OutPosition, LightToWorldRotation) + LightData.TranslatedWorldPosition;
+#endif
+
+	return OutPosition;
+}
+
+struct FVSToPS
+{
+	uint LightSceneId : TEXCOORD0;
+};
+
+void VisualizeLightShapeVS(
+	float3 LocalPosition : ATTRIBUTE0,
+	uint2 ForwardLightIndex_LightType : ATTRIBUTE1,
+	out FVSToPS VSToPS,
+	out float4 OutPosition : SV_POSITION)
+{
+	// TODO: instanced stereo?
+	ResolvedView = ResolveView();
+
+	uint ForwardLightIndex = ForwardLightIndex_LightType.x;
+	FForwardLightData ForwardLightData = GetForwardLightData(ForwardLightIndex, 0);
+	FDeferredLightData LightData = ConvertToDeferredLight(ForwardLightData);
+
+	float3 TranslatedWorldPosition = GetLightShapeVertexTranslatedWorldPosition(LocalPosition, LightData);
+
+	VSToPS.LightSceneId = ForwardLightData.LightSceneId;
+    OutPosition = mul(float4(TranslatedWorldPosition, 1.0), View.TranslatedWorldToClip);
+}
+
+float VisualizationToggleTime;
+float InvFadeDuration;
+float InvMaxDelayDistance;
+float MaxDelay;
+float MaxOpacity;
+
+void VisualizeLightShapePS(
+	FVSToPS VSToPS,
+    float4 SvPosition : SV_POSITION,
+    out float4 OutColor : SV_Target0)
+{
+	// TODO: instanced stereo?
+	ResolvedView = ResolveView();
+
+	float Random = Rand32ToFloat(PCGHash(VSToPS.LightSceneId));
+	float3 TranslatedWorldPosition = SvPositionToTranslatedWorld(SvPosition);
+	float ViewDistance = length(TranslatedWorldPosition - View.TranslatedWorldCameraOrigin);
+	
+	// When visualization is just toggled on, delay fade in over distance so that lights far
+	// away fade in first. When just toggled off, all lights fade out together
+	bool bVisualize = VisualizationToggleTime > 0;
+	float Delay = bVisualize ? saturate(1.0 - ViewDistance * InvMaxDelayDistance) * MaxDelay : 0.0;
+	float Alpha = View.GameTime - Delay - abs(VisualizationToggleTime);
+	Alpha = saturate(Alpha * InvFadeDuration);
+	Alpha = bVisualize ? Alpha : (1.0 - Alpha);
+	Alpha *= MaxOpacity;
+
+	// Depth fade
+	float4 ScreenPosition = SvPositionToScreenPosition(SvPosition);
+	float PixelDepth = GetScreenPositionDepth(ScreenPosition);
+	float OpaqueDepth = CalcSceneDepth(SvPositionToBufferUV(SvPosition));
+	float DepthFadeDistance = 100.0;
+	float DepthFadeAlpha = saturate((OpaqueDepth - PixelDepth) / DepthFadeDistance);
+	Alpha *= DepthFadeAlpha;
+
+	float3 Color = HSVtoRGB(float3(Random, 1, 1));
+
+    OutColor = float4(Color, Alpha);
+}
diff --git a/Engine/Source/Runtime/Renderer/Private/DeferredShadingRenderer.cpp b/Engine/Source/Runtime/Renderer/Private/DeferredShadingRenderer.cpp
--- a/Engine/Source/Runtime/Renderer/Private/DeferredShadingRenderer.cpp
+++ b/Engine/Source/Runtime/Renderer/Private/DeferredShadingRenderer.cpp
@@ -3851,6 +3851,11 @@ void FDeferredShadingSceneRenderer::Render(FRDGBuilder& GraphBuilder, const FSce
 			RenderHairComposition(GraphBuilder, Views, SceneTextures.Color.Target, SceneTextures.Depth.Target, SceneTextures.Velocity, TranslucencyResourceMap, SingleLayerWaterPrePassResult);
 		}
 
+		if (!bHasRayTracedOverlay)
+		{
+			RenderLightShapeVisualization(GraphBuilder, SceneTextures, TranslucencyResourceMap);
+		}
+
 #if !UE_BUILD_SHIPPING
 		if (CVarForceBlackVelocityBuffer.GetValueOnRenderThread())
 		{
diff --git a/Engine/Source/Runtime/Renderer/Private/DeferredShadingRenderer.h b/Engine/Source/Runtime/Renderer/Private/DeferredShadingRenderer.h
--- a/Engine/Source/Runtime/Renderer/Private/DeferredShadingRenderer.h
+++ b/Engine/Source/Runtime/Renderer/Private/DeferredShadingRenderer.h
@@ -676,6 +676,11 @@ private:
 		const FFrontLayerTranslucencyData& FrontLayerTranslucencyData,
 		ERDGPassFlags ComputePassFlags);
 
+	void RenderLightShapeVisualization(
+		FRDGBuilder& GraphBuilder,
+		const FSceneTextures& SceneTextures,
+		const FTranslucencyPassResourcesMap& TranslucencyResourceMap);
+
 	FSSDSignalTextures RenderLumenFinalGather(
 		FRDGBuilder& GraphBuilder,
 		const FSceneTextures& SceneTextures,
diff --git a/Engine/Source/Runtime/Renderer/Private/LightGridInjection.cpp b/Engine/Source/Runtime/Renderer/Private/LightGridInjection.cpp
--- a/Engine/Source/Runtime/Renderer/Private/LightGridInjection.cpp
+++ b/Engine/Source/Runtime/Renderer/Private/LightGridInjection.cpp
@@ -42,6 +42,7 @@ LightGridInjection.cpp
 #include "LightViewData.h"
 #include "HairStrands/HairStrandsTransmittance.h"
 #include "GPUMessaging.h"
+#include "Templates/Greater.h"
 
 static TAutoConsoleVariable<bool> CVarLightGridAsyncCompute(
 	TEXT("r.Forward.LightGridAsyncCompute"),
@@ -1109,6 +1110,9 @@ FComputeLightGridOutput FSceneRenderer::ComputeLightGrid(FRDGBuilder& GraphBuild
 		const uint32 LightNameBytes = LightNameDwords * sizeof(uint32);
 		ForwardLightUniformParameters->LightNameDwords = LightNameDwords;
 
+		TArray<VisualizeLightShape::FSortKey> VisualizeLightShapeSortKeys;
+		const bool bVisualizeLightShape = VisualizeLightShape::UpdateLastToggleTime(View);
+
 		float FurthestLight = 1000;
 
 		int32 ConflictingLightCountForForwardShading = 0;
@@ -1151,6 +1155,11 @@ FComputeLightGridOutput FSceneRenderer::ComputeLightGrid(FRDGBuilder& GraphBuild
 				ViewSpaceDirAndPreprocAngleData.AddZeroed(SimpleLightsEnd);
 				ViewSpaceRectPlanesData.AddZeroed(SimpleLightsEnd * NUM_PLANES_PER_RECT_LIGHT);
 				ViewSpaceWeightEstimateData.AddZeroed(SimpleLightsEnd);
+
+				if (bVisualizeLightShape)
+				{
+					VisualizeLightShapeSortKeys.AddUninitialized(SimpleLightsEnd);
+				}
 			}
 
 			const uint32 LightShaderParameterFlags = bRenderRectLightsAsSpotLights ? ELightShaderParameterFlags::RectAsSpotLight : 0u;
@@ -1186,6 +1195,11 @@ FComputeLightGridOutput FSceneRenderer::ComputeLightGrid(FRDGBuilder& GraphBuild
 						ViewSpaceWeightEstimateData.Reserve(ViewSpaceWeightEstimateData.Num() + NumMegaSimpleLights);
 						IndirectionIndices.Reserve(IndirectionIndices.Num() + NumMegaSimpleLights);
 
+						if (bVisualizeLightShape)
+						{
+							VisualizeLightShapeSortKeys.Reserve(VisualizeLightShapeSortKeys.Num() + NumMegaSimpleLights);
+						}
+
 						if (bNeedsLightNames)
 						{
 							LightNames.Reserve(LightNames.Num() + NumMegaSimpleLights * LightNameBytes);
@@ -1230,6 +1244,14 @@ FComputeLightGridOutput FSceneRenderer::ComputeLightGrid(FRDGBuilder& GraphBuild
 
 					IndirectionIndices.Add(IndexInBuffer);
 
+					if (bVisualizeLightShape)
+					{
+						VisualizeLightShape::FSortKey& SortKey = VisualizeLightShapeSortKeys[VisualizeLightShapeSortKeys.AddUninitialized()];
+						SortKey.LightData.ForwardLightIndex = IndexInBuffer;
+						SortKey.LightData.LightType = LightType_Point;
+						SortKey.Distance = FVector3f(ViewSpacePosAndRadius).SizeSquared();
+					}
+
 					if (bNeedsLightNames)
 					{
 						const int32 FirstCharIndex = LightNames.AddUninitialized(LightNameBytes);
@@ -1376,7 +1398,8 @@ FComputeLightGridOutput FSceneRenderer::ComputeLightGrid(FRDGBuilder& GraphBuild
 						}
 
 						const FSphere BoundingSphere = LightProxy->GetBoundingSphere();
-						const float Distance = View.ViewMatrices.GetWorldToView().TransformPosition(BoundingSphere.Center).Z + BoundingSphere.W;
+						const FVector ViewSpaceBoundingSphereCenter(View.ViewMatrices.GetWorldToView().TransformPosition(BoundingSphere.Center));
+						const float Distance = ViewSpaceBoundingSphereCenter.Z + BoundingSphere.W;
 						FurthestLight = FMath::Max(FurthestLight, Distance);
 
 						const FVector3f LightViewPosition = FVector4f(View.ViewMatrices.GetWorldToView().TransformPosition(LightParameters.WorldPosition)); // LWC_TODO: precision loss
@@ -1425,6 +1448,14 @@ FComputeLightGridOutput FSceneRenderer::ComputeLightGrid(FRDGBuilder& GraphBuild
 						WeightEstimateData.W = 0.0f;
 						ViewSpaceWeightEstimateData.Add(WeightEstimateData);
 
+						if (bVisualizeLightShape)
+						{
+							VisualizeLightShape::FSortKey& SortKey = VisualizeLightShapeSortKeys[VisualizeLightShapeSortKeys.AddUninitialized()];
+							SortKey.LightData.ForwardLightIndex = IndexInBuffer;
+							SortKey.LightData.LightType = SortedLightInfo.SortKey.Fields.LightType;
+							SortKey.Distance = (float)ViewSpaceBoundingSphereCenter.SizeSquared();
+						}
+
 						bHasRectLights |= bIsRectLight;
 						bHasTexturedLights |= LightProxy->HasSourceTexture();
 						LightGridLightingChannelAndMask &= LightProxy->GetLightingChannelMask();
@@ -1587,6 +1618,14 @@ FComputeLightGridOutput FSceneRenderer::ComputeLightGrid(FRDGBuilder& GraphBuild
 					ViewSpacePosAndRadiusData[SortedIndex] = ViewSpacePosAndRadius;
 
 					IndirectionIndices[SortedIndex] = IndexInBuffer;
+
+					if (bVisualizeLightShape)
+					{
+						VisualizeLightShape::FSortKey& SortKey = VisualizeLightShapeSortKeys[VisualizeLightShapeSortKeys.AddUninitialized()];
+						SortKey.LightData.ForwardLightIndex = IndexInBuffer;
+						SortKey.LightData.LightType = LightType_Point;
+						SortKey.Distance = FVector3f(ViewSpacePosAndRadius).SizeSquared();
+					}
 				}
 			}
 		}
@@ -1683,6 +1722,26 @@ FComputeLightGridOutput FSceneRenderer::ComputeLightGrid(FRDGBuilder& GraphBuild
 		check(ViewSpaceWeightEstimateData.Num() == NumVisibleLocalLights);
 		check(IndirectionIndices.Num() == NumVisibleLocalLights);
 
+		if (bVisualizeLightShape && VisualizeLightShapeSortKeys.Num() > 0)
+		{
+			FRDGUploadData<VisualizeLightShape::FLightInfo> VisualizeLightShapeLightInfo(GraphBuilder, VisualizeLightShapeSortKeys.Num());
+
+			GraphBuilder.AddSetupTask([SortKeys = MoveTemp(VisualizeLightShapeSortKeys), VisualizeLightShapeLightInfo]()
+			{
+				SCOPED_NAMED_EVENT(VisualizeLightShapeSortLights, FColor::Emerald);
+
+				Algo::Sort(SortKeys, TGreater<>());
+
+				for (int32 Index = 0; Index < SortKeys.Num(); ++Index)
+				{
+					VisualizeLightShapeLightInfo[Index] = SortKeys[Index].LightData;
+				}
+			});
+
+			View.ForwardLightingResources.VisualizeLightShapeLightInfos = TConstArrayView<VisualizeLightShape::FLightInfo>(VisualizeLightShapeLightInfo.GetData(), VisualizeLightShapeLightInfo.Num());
+			View.ForwardLightingResources.VisualizeLightShapeLightInfoBuffer = CreateUploadBuffer(GraphBuilder, TEXT("VisualizeLightShape.LightInfos"), VisualizeLightShapeLightInfo);
+		}
+
 		FRDGBufferRef LightViewSpacePositionAndRadius = CreateStructuredBuffer(GraphBuilder, TEXT("ViewSpacePosAndRadiusData"), TConstArrayView<FVector4f>(ViewSpacePosAndRadiusData));
 		FRDGBufferRef LightViewSpaceDirAndPreprocAngle = CreateStructuredBuffer(GraphBuilder, TEXT("ViewSpaceDirAndPreprocAngleData"), TConstArrayView<FVector4f>(ViewSpaceDirAndPreprocAngleData));
 		FRDGBufferRef LightViewSpaceRectPlanes = CreateStructuredBuffer(GraphBuilder, TEXT("ViewSpaceRectPlanesData"), TConstArrayView<FVector4f>(ViewSpaceRectPlanesData));
diff --git a/Engine/Source/Runtime/Renderer/Private/Renderer.cpp b/Engine/Source/Runtime/Renderer/Private/Renderer.cpp
--- a/Engine/Source/Runtime/Renderer/Private/Renderer.cpp
+++ b/Engine/Source/Runtime/Renderer/Private/Renderer.cpp
@@ -220,6 +220,8 @@ void FRendererModule::ShutdownModule()
 	// Free up global resources in Lumen
 	Lumen::Shutdown();
 
+	VisualizeLightShape::Shutdown();
+
 	void CleanupOcclusionSubmittedFence();
 	CleanupOcclusionSubmittedFence();
 }
diff --git a/Engine/Source/Runtime/Renderer/Private/SceneRendering.h b/Engine/Source/Runtime/Renderer/Private/SceneRendering.h
--- a/Engine/Source/Runtime/Renderer/Private/SceneRendering.h
+++ b/Engine/Source/Runtime/Renderer/Private/SceneRendering.h
@@ -100,6 +100,7 @@
 #include "MegaFiltersInterface.h"
 #include "LightShaftRendering.h"
 #include "SceneViewOcclusionHistory.h"
+#include "VisualizeLightShape.h"
 
 #if UE_ENABLE_INCLUDE_ORDER_DEPRECATED_IN_5_8
 #include "RHI.h"
@@ -798,6 +799,9 @@ public:
 	FRDGBufferUAVRef CulledLightDataGridUAV;
 	FRDGBufferSRVRef NumCulledLightsGridSRV;
 	FRDGBufferUAVRef NumCulledLightsGridUAV;
+
+	TConstArrayView<VisualizeLightShape::FLightInfo> VisualizeLightShapeLightInfos;
+	FRDGBufferRef VisualizeLightShapeLightInfoBuffer = nullptr;
 };
 
 BEGIN_GLOBAL_SHADER_PARAMETER_STRUCT_WITH_CONSTRUCTOR(FVolumetricFogGlobalData,) 
diff --git a/Engine/Source/Runtime/Renderer/Private/SceneViewState.h b/Engine/Source/Runtime/Renderer/Private/SceneViewState.h
--- a/Engine/Source/Runtime/Renderer/Private/SceneViewState.h
+++ b/Engine/Source/Runtime/Renderer/Private/SceneViewState.h
@@ -320,6 +320,8 @@ public:
 
 	FTranslucencyLightingViewState TranslucencyLighting;
 
+	VisualizeLightShape::FViewState VisualizeLightShape;
+
 	// Heterogeneous Volumes cached data stores
 	TRDGUniformBufferRef<FOrthoVoxelGridUniformBufferParameters> OrthoVoxelGridUniformBuffer = nullptr;
 	TRDGUniformBufferRef<FFrustumVoxelGridUniformBufferParameters> FrustumVoxelGridUniformBuffer = nullptr;
diff --git a/Engine/Source/Runtime/Renderer/Private/TranslucentRendering.cpp b/Engine/Source/Runtime/Renderer/Private/TranslucentRendering.cpp
--- a/Engine/Source/Runtime/Renderer/Private/TranslucentRendering.cpp
+++ b/Engine/Source/Runtime/Renderer/Private/TranslucentRendering.cpp
@@ -1382,6 +1382,11 @@ static FViewShaderParameters GetSeparateTranslucencyViewParameters(const FViewIn
 	return ViewParameters;
 }
 
+FViewShaderParameters GetSeparateTranslucencyViewParameters(const FViewInfo& View, const FSeparateTranslucencyDimensions& SeparateTranslucencyDimensions, ETranslucencyPass::Type TranslucencyPass)
+{
+	return GetSeparateTranslucencyViewParameters(View, SeparateTranslucencyDimensions.Extent, SeparateTranslucencyDimensions.Scale, TranslucencyPass);
+}
+
 BEGIN_SHADER_PARAMETER_STRUCT(FTranslucentBasePassParameters, )
 	SHADER_PARAMETER_STRUCT_INCLUDE(FViewShaderParameters, View)
 	SHADER_PARAMETER_STRUCT_REF(FReflectionCaptureShaderData, ReflectionCapture)
diff --git a/Engine/Source/Runtime/Renderer/Private/TranslucentRendering.h b/Engine/Source/Runtime/Renderer/Private/TranslucentRendering.h
--- a/Engine/Source/Runtime/Renderer/Private/TranslucentRendering.h
+++ b/Engine/Source/Runtime/Renderer/Private/TranslucentRendering.h
@@ -13,6 +13,7 @@
 #include "ProfilingDebugging/RealtimeGPUProfiler.h"
 #include "PostProcess/SceneRenderTargets.h"
 #include "ScreenPass.h"
+#include "SceneView.h"
 
 class FSceneRenderer;
 class FViewInfo;
@@ -79,6 +80,7 @@ const FRDGTextureDesc GetPostDOFTranslucentTextureDesc(ETranslucencyPass::Type T
 /** Shared function used to create Post DOF translucent textures */
 FRDGTextureMSAA CreatePostDOFTranslucentTexture(FRDGBuilder& GraphBuilder, ETranslucencyPass::Type TranslucencyPass, const FSeparateTranslucencyDimensions& SeparateTranslucencyDimensions, bool bIsModulate, EShaderPlatform ShaderPlatform);
 
+FViewShaderParameters GetSeparateTranslucencyViewParameters(const FViewInfo& View, const FSeparateTranslucencyDimensions& SeparateTranslucencyDimensions, ETranslucencyPass::Type TranslucencyPass);
 
 /** Add a pass to compose separate translucency. */
 struct FTranslucencyComposition
diff --git a/Engine/Source/Runtime/Renderer/Private/VisualizeLightShape.cpp b/Engine/Source/Runtime/Renderer/Private/VisualizeLightShape.cpp
new file mode 100644
--- /dev/null
+++ b/Engine/Source/Runtime/Renderer/Private/VisualizeLightShape.cpp
@@ -0,0 +1,646 @@
+// Copyright Epic Games, Inc. All Rights Reserved.
+
+#include "VisualizeLightShape.h"
+#include "BasePassRendering.h"
+#include "SceneViewState.h"
+#include "TranslucentPassResource.h"
+#include "DeferredShadingRenderer.h"
+
+static TAutoConsoleVariable<bool> CVarVisualizeLightShapeEnable(
+	TEXT("r.VisualizeLightShape.Enable"),
+	false,
+	TEXT("Whether to enable light shape visualization."),
+	ECVF_RenderThreadSafe);
+
+static TAutoConsoleVariable<float> CVarVisualizeLightShapeFadeDuration(
+	TEXT("r.VisualizeLightShape.FadeDuration"),
+	0.5f,
+	TEXT("Fade duration in seconds when visualization is toggled."),
+	ECVF_RenderThreadSafe);
+
+static TAutoConsoleVariable<float> CVarVisualizeLightShapeMaxDelayDistance(
+	TEXT("r.VisualizeLightShape.MaxDelayDistance"),
+	4000.0f,
+	TEXT("Max distance where the fade-in animation can be delayed. ")
+	TEXT("Delay time is scaled linearly within this range and closer lights have more delay."),
+	ECVF_RenderThreadSafe);
+
+static TAutoConsoleVariable<float> CVarVisualizeLightShapeMaxDelay(
+	TEXT("r.VisualizeLightShape.MaxDelay"),
+	2.5f,
+	TEXT("Max distance based delay in seconds. Delay is only applied when fading in visualization."),
+	ECVF_RenderThreadSafe);
+
+static TAutoConsoleVariable<float> CVarVisualizeLightShapeMaxOpacity(
+	TEXT("r.VisualizeLightShape.MaxOpacity"),
+	0.1f,
+	TEXT("Max opacity of each visualization mesh."),
+	ECVF_RenderThreadSafe);
+
+DECLARE_GPU_STAT(VisualizeLightShape);
+
+namespace VisualizeLightShape
+{
+	bool GetUserEnable()
+	{
+		return CVarVisualizeLightShapeEnable.GetValueOnRenderThread();
+	}
+
+	float GetMaxDelay()
+	{
+		return FMath::Max(CVarVisualizeLightShapeMaxDelay.GetValueOnRenderThread(), 0.0f);
+	}
+
+	float GetFadeDuration()
+	{
+		return FMath::Max(CVarVisualizeLightShapeFadeDuration.GetValueOnRenderThread(), 0.0f);
+	}
+
+	float GetMaxDelayDistance()
+	{
+		return FMath::Max(CVarVisualizeLightShapeMaxDelayDistance.GetValueOnRenderThread(), 100.0f);
+	}
+
+	float GetMaxOpacity()
+	{
+		return FMath::Clamp(CVarVisualizeLightShapeMaxOpacity.GetValueOnRenderThread(), 0.0f, 1.0f);
+	}
+
+	// Returns whether this view should render
+	bool UpdateLastToggleTime(FViewInfo& View)
+	{
+		if (!View.ViewState)
+		{
+			return false;
+		}
+
+		FViewState& ViewState = View.ViewState->VisualizeLightShape;
+		const double CurrentGameTime = View.Family->Time.GetWorldTimeSeconds();
+
+		if (CurrentGameTime < ViewState.LastToggleTime)
+		{
+			ViewState.LastToggleTime = 0.0;
+		}
+
+		const bool bWasEnabled = ViewState.LastToggleTime > 0.0;
+		if (GetUserEnable() != bWasEnabled)
+		{
+			ViewState.LastToggleTime = bWasEnabled ? -CurrentGameTime : CurrentGameTime;
+		}
+
+		const double FadeDuration = GetFadeDuration();
+		return ViewState.LastToggleTime != 0.0 && (ViewState.LastToggleTime > 0.0 || CurrentGameTime + ViewState.LastToggleTime < FadeDuration);
+	}
+}
+
+struct FVisualizationMeshBuffers
+{
+	uint32 NumVertices = 0;
+	uint32 NumTriangles = 0;
+	FBufferRHIRef VertexBufferRHI;
+	FBufferRHIRef IndexBufferRHI;
+};
+
+static FVisualizationMeshBuffers CreatePointLightVisualizationMesh(FRDGBuilder& GraphBuilder)
+{
+	const int32 NumSections = 16;
+	const int32 NumSlices = NumSections * 2;
+	const int32 NumVertices = (NumSlices + 1) * (NumSections - 1) + NumSlices * 2;
+	const int32 NumTriangles = NumSlices * 2 * (NumSections - 2) + NumSlices * 2;
+	const float Radius = 1.0f;
+
+	FRHICommandListImmediate& RHICmdList = GraphBuilder.RHICmdList;
+	FVisualizationMeshBuffers Out;
+
+	Out.NumVertices = NumVertices;
+	Out.NumTriangles = NumTriangles;
+
+	{
+		FRHIBufferCreateDesc CreateDesc = FRHIBufferCreateDesc::CreateVertex(TEXT("VisualizeLightShape.Point.VertexBuffer"), sizeof(FVector3f) * NumVertices)
+			.AddUsage(EBufferUsageFlags::ShaderResource)
+			.SetInitialState(ERHIAccess::VertexOrIndexBuffer | ERHIAccess::SRVMask);
+
+		Out.VertexBufferRHI = RHICmdList.CreateBuffer(CreateDesc);
+	}
+
+	{
+		FRHIBufferCreateDesc CreateDesc = FRHIBufferCreateDesc::CreateIndex(TEXT("VisualizeLightShape.Point.IndexBuffer"), sizeof(uint16) * NumTriangles * 3, sizeof(uint16))
+			.AddUsage(EBufferUsageFlags::ShaderResource)
+			.SetInitialState(ERHIAccess::VertexOrIndexBuffer | ERHIAccess::SRVMask);
+
+		Out.IndexBufferRHI = RHICmdList.CreateBuffer(CreateDesc);
+	}
+
+	FVector3f* VertexData = (FVector3f*)RHICmdList.LockBuffer(Out.VertexBufferRHI, 0, sizeof(FVector3f) * NumVertices, RLM_WriteOnly);
+	uint16* IndexData = (uint16*)RHICmdList.LockBuffer(Out.IndexBufferRHI, 0, sizeof(uint16) * NumTriangles * 3, RLM_WriteOnly);
+
+	int32 VertexIndex = 0;
+	for (int32 X = 0; X <= NumSlices; ++X)
+	{
+		const float Phi = UE_TWO_PI / NumSlices * X;
+		const float SinPhi = FMath::Sin(Phi);
+		const float CosPhi = FMath::Cos(Phi);
+
+		const int32 StartY = X == NumSlices ? 1 : 0;
+		const int32 EndY = NumSections - (X == NumSlices ? 1 : 0);
+
+		for (int32 Y = StartY; Y <= EndY; ++Y)
+		{
+			const float Theta = UE_PI / NumSections * Y;
+			const float SinTheta = FMath::Sin(Theta);
+			const float CosTheta = FMath::Cos(Theta);
+
+			VertexData[VertexIndex] = FVector3f(SinTheta * CosPhi, SinTheta * SinPhi, CosTheta) * Radius;
+			++VertexIndex;
+		}
+	}
+	check(VertexIndex == NumVertices);
+
+	int32 TriangleIndex = 0;
+	for (int32 X = 0; X < NumSlices; ++X)
+	{
+		const int32 Offset = (NumSections + 1) * X;
+
+		for (int32 Y = 0; Y < NumSections; ++Y)
+		{
+			// I3 - I2
+			// |  \  |
+			// I0 - I1
+			const int32 I0 = Offset + Y + 1;
+			const int32 I1 = Offset + Y + (NumSections + 1) + (X == NumSlices - 1 ? 0 : 1);
+			const int32 I2 = Offset + Y + (NumSections + 1) + (X == NumSlices - 1 ? -1 : 0);
+			const int32 I3 = Offset + Y;
+
+			if (Y == 0)
+			{
+				IndexData[TriangleIndex * 3 + 0] = I0;
+				IndexData[TriangleIndex * 3 + 1] = I1;
+				IndexData[TriangleIndex * 3 + 2] = I3;
+				++TriangleIndex;
+			}
+			else if (Y == NumSections - 1)
+			{
+				IndexData[TriangleIndex * 3 + 0] = I0;
+				IndexData[TriangleIndex * 3 + 1] = I2;
+				IndexData[TriangleIndex * 3 + 2] = I3;
+				++TriangleIndex;
+			}
+			else
+			{
+				IndexData[TriangleIndex * 3 + 0] = I0;
+				IndexData[TriangleIndex * 3 + 1] = I1;
+				IndexData[TriangleIndex * 3 + 2] = I3;
+				++TriangleIndex;
+
+				IndexData[TriangleIndex * 3 + 0] = I1;
+				IndexData[TriangleIndex * 3 + 1] = I2;
+				IndexData[TriangleIndex * 3 + 2] = I3;
+				++TriangleIndex;
+			}
+		}
+	}
+	check(TriangleIndex == NumTriangles);
+
+	RHICmdList.UnlockBuffer(Out.VertexBufferRHI);
+	RHICmdList.UnlockBuffer(Out.IndexBufferRHI);
+
+	return Out;
+}
+
+static FVisualizationMeshBuffers CreateSpotLightVisualizationMesh(FRDGBuilder& GraphBuilder)
+{
+	const int32 NumSections = 11;
+	const int32 NumSlices = 32;
+	const int32 NumVertices = (NumSlices + 1) * (NumSections + 1);
+	const int32 NumTriangles = NumSlices * 2 * NumSections;
+	const float Radius = 1.0f;
+	const float Length = 1.1f;
+
+	FRHICommandListImmediate& RHICmdList = GraphBuilder.RHICmdList;
+	FVisualizationMeshBuffers Out;
+
+	Out.NumVertices = NumVertices;
+	Out.NumTriangles = NumTriangles;
+
+	{
+		FRHIBufferCreateDesc CreateDesc = FRHIBufferCreateDesc::CreateVertex(TEXT("VisualizeLightShape.Spot.VertexBuffer"), sizeof(FVector3f) * NumVertices)
+			.AddUsage(EBufferUsageFlags::ShaderResource)
+			.SetInitialState(ERHIAccess::VertexOrIndexBuffer | ERHIAccess::SRVMask);
+
+		Out.VertexBufferRHI = RHICmdList.CreateBuffer(CreateDesc);
+	}
+
+	{
+		FRHIBufferCreateDesc CreateDesc = FRHIBufferCreateDesc::CreateIndex(TEXT("VisualizeLightShape.Spot.IndexBuffer"), sizeof(uint16) * NumTriangles * 3, sizeof(uint16))
+			.AddUsage(EBufferUsageFlags::ShaderResource)
+			.SetInitialState(ERHIAccess::VertexOrIndexBuffer | ERHIAccess::SRVMask);
+
+		Out.IndexBufferRHI = RHICmdList.CreateBuffer(CreateDesc);
+	}
+
+	FVector3f* VertexData = (FVector3f*)RHICmdList.LockBuffer(Out.VertexBufferRHI, 0, sizeof(FVector3f) * NumVertices, RLM_WriteOnly);
+	uint16* IndexData = (uint16*)RHICmdList.LockBuffer(Out.IndexBufferRHI, 0, sizeof(uint16) * NumTriangles * 3, RLM_WriteOnly);
+
+	int32 VertexIndex = 0;
+	for (int32 X = 0; X <= NumSlices; ++X)
+	{
+		const float Phi = UE_TWO_PI / NumSlices * X;
+		const float SinPhi = FMath::Sin(Phi);
+		const float CosPhi = FMath::Cos(Phi);
+
+		for (int32 Y = 0; Y <= NumSections; ++Y)
+		{
+			VertexData[VertexIndex] = FVector3f(Length / NumSections * Y, Radius * CosPhi, Radius * SinPhi);
+			++VertexIndex;
+		}
+	}
+	check(VertexIndex == NumVertices);
+
+	int32 TriangleIndex = 0;
+	for (int32 X = 0; X < NumSlices; ++X)
+	{
+		const int32 Offset = (NumSections + 1) * X;
+
+		for (int32 Y = 0; Y < NumSections; ++Y)
+		{
+			// I3 - I2
+			// |  \  |
+			// I0 - I1
+			const int32 I0 = Offset + Y;
+			const int32 I1 = Offset + Y + 1;
+			const int32 I2 = Offset + Y + (NumSections + 1) + 1;
+			const int32 I3 = Offset + Y + (NumSections + 1);
+
+			IndexData[TriangleIndex * 3 + 0] = I0;
+			IndexData[TriangleIndex * 3 + 1] = I1;
+			IndexData[TriangleIndex * 3 + 2] = I3;
+			++TriangleIndex;
+
+			IndexData[TriangleIndex * 3 + 0] = I1;
+			IndexData[TriangleIndex * 3 + 1] = I2;
+			IndexData[TriangleIndex * 3 + 2] = I3;
+			++TriangleIndex;
+		}
+	}
+	check(TriangleIndex == NumTriangles);
+
+	RHICmdList.UnlockBuffer(Out.VertexBufferRHI);
+	RHICmdList.UnlockBuffer(Out.IndexBufferRHI);
+
+	return Out;
+}
+
+static FVisualizationMeshBuffers CreateRectLightVisualizationMesh(FRDGBuilder& GraphBuilder)
+{
+	const int32 NumSections = 14;
+	const int32 NumVertices = (NumSections + 1) * (NumSections * 4 + 1);
+	const int32 NumTriangles = NumSections * NumSections * 4 * 2;
+	const float Length = 1.0f;
+
+	FRHICommandListImmediate& RHICmdList = GraphBuilder.RHICmdList;
+	FVisualizationMeshBuffers Out;
+
+	Out.NumVertices = NumVertices;
+	Out.NumTriangles = NumTriangles;
+
+	{
+		FRHIBufferCreateDesc CreateDesc = FRHIBufferCreateDesc::CreateVertex(TEXT("VisualizeLightShape.Rect.VertexBuffer"), sizeof(FVector3f) * NumVertices)
+			.AddUsage(EBufferUsageFlags::ShaderResource)
+			.SetInitialState(ERHIAccess::VertexOrIndexBuffer | ERHIAccess::SRVMask);
+
+		Out.VertexBufferRHI = RHICmdList.CreateBuffer(CreateDesc);
+	}
+
+	{
+		FRHIBufferCreateDesc CreateDesc = FRHIBufferCreateDesc::CreateIndex(TEXT("VisualizeLightShape.Rect.IndexBuffer"), sizeof(uint16) * NumTriangles * 3, sizeof(uint16))
+			.AddUsage(EBufferUsageFlags::ShaderResource)
+			.SetInitialState(ERHIAccess::VertexOrIndexBuffer | ERHIAccess::SRVMask);
+
+		Out.IndexBufferRHI = RHICmdList.CreateBuffer(CreateDesc);
+	}
+
+	FVector3f* VertexData = (FVector3f*)RHICmdList.LockBuffer(Out.VertexBufferRHI, 0, sizeof(FVector3f) * NumVertices, RLM_WriteOnly);
+	uint16* IndexData = (uint16*)RHICmdList.LockBuffer(Out.IndexBufferRHI, 0, sizeof(uint16) * NumTriangles * 3, RLM_WriteOnly);
+
+	int32 VertexIndex = 0;
+	for (int32 Side = 0; Side < 4; ++Side)
+	{
+		for (int32 Y = 0; Y <= NumSections + (Side == 3 ? 0 : -1); ++Y)
+		{
+			for (int32 X = 0; X <= NumSections; ++X)
+			{
+				VertexData[VertexIndex][0] = Length / NumSections * X - Length * 0.5f;
+				VertexData[VertexIndex][Side % 2 + 1] = Length * 0.5f * (Side >= 2 ? -1.0f : 1.0f);
+				VertexData[VertexIndex][(Side + 1) % 2 + 1] = (Length / NumSections * Y - Length * 0.5f) * (Side == 1 || Side == 2 ? -1.0f : 1.0f);
+				++VertexIndex;
+			}
+		}
+	}
+	check(VertexIndex == NumVertices);
+
+	int32 TriangleIndex = 0;
+	for (int32 Side = 0; Side < 4; ++Side)
+	{
+		const int32 SideOffset = (NumSections + 1) * NumSections * Side;
+
+		for (int32 Y = 0; Y < NumSections; ++Y)
+		{
+			const int32 Offset = SideOffset + (NumSections + 1) * Y;
+
+			for (int32 X = 0; X < NumSections; ++X)
+			{
+				// I3 - I2
+				// |  \  |
+				// I0 - I1
+				const int32 I0 = Offset + X;
+				const int32 I1 = Offset + X + 1;
+				const int32 I2 = Offset + X + (NumSections + 1) + 1;
+				const int32 I3 = Offset + X + (NumSections + 1);
+
+				IndexData[TriangleIndex * 3 + 0] = I0;
+				IndexData[TriangleIndex * 3 + 1] = I1;
+				IndexData[TriangleIndex * 3 + 2] = I3;
+				++TriangleIndex;
+
+				IndexData[TriangleIndex * 3 + 0] = I1;
+				IndexData[TriangleIndex * 3 + 1] = I2;
+				IndexData[TriangleIndex * 3 + 2] = I3;
+				++TriangleIndex;
+			}
+		}
+	}
+	check(TriangleIndex == NumTriangles);
+
+	RHICmdList.UnlockBuffer(Out.VertexBufferRHI);
+	RHICmdList.UnlockBuffer(Out.IndexBufferRHI);
+
+	return Out;
+}
+
+static constexpr uint32 GNumLocalLightTypes = LIGHT_TYPE_MAX - 1;
+static bool bGVisualizationMeshInitialized = false;
+static FVisualizationMeshBuffers GVisualizationMeshBuffers[GNumLocalLightTypes];
+
+namespace VisualizeLightShape
+{
+	void InitVisualizationMeshes(FRDGBuilder& GraphBuilder)
+	{
+		check(IsInRenderingThread());
+
+		if (!bGVisualizationMeshInitialized)
+		{
+			bGVisualizationMeshInitialized = true;
+			GVisualizationMeshBuffers[LIGHT_TYPE_POINT - 1] = CreatePointLightVisualizationMesh(GraphBuilder);
+			GVisualizationMeshBuffers[LIGHT_TYPE_SPOT - 1] = CreateSpotLightVisualizationMesh(GraphBuilder);
+			GVisualizationMeshBuffers[LIGHT_TYPE_RECT - 1] = CreateRectLightVisualizationMesh(GraphBuilder);
+		}
+	}
+
+	void Shutdown()
+	{
+		if (bGVisualizationMeshInitialized)
+		{
+			for (int32 Index = 0; Index < GNumLocalLightTypes; ++Index)
+			{
+				GVisualizationMeshBuffers[Index] = FVisualizationMeshBuffers();
+			}
+			bGVisualizationMeshInitialized = false;
+		}
+	}
+}
+
+class FVisualizeLightShapeVS : public FGlobalShader
+{
+	DECLARE_GLOBAL_SHADER(FVisualizeLightShapeVS)
+	SHADER_USE_PARAMETER_STRUCT(FVisualizeLightShapeVS, FGlobalShader)
+
+	BEGIN_SHADER_PARAMETER_STRUCT(FParameters, )
+		SHADER_PARAMETER_STRUCT_REF(FViewUniformShaderParameters, ViewUniformBuffer)
+		SHADER_PARAMETER_RDG_UNIFORM_BUFFER(FForwardLightUniformParameters, ForwardLightStruct)
+	END_SHADER_PARAMETER_STRUCT()
+
+	class FLightType : SHADER_PERMUTATION_RANGE_INT("LIGHT_TYPE", LIGHT_TYPE_POINT, LIGHT_TYPE_MAX - LIGHT_TYPE_POINT);
+	using FPermutationDomain = TShaderPermutationDomain<FLightType>;
+
+	static void ModifyCompilationEnvironment(const FGlobalShaderPermutationParameters& Parameters, FShaderCompilerEnvironment& OutEnvironment)
+	{
+		FGlobalShader::ModifyCompilationEnvironment(Parameters, OutEnvironment);
+		FForwardLightingParameters::ModifyCompilationEnvironment(Parameters.Platform, OutEnvironment);
+		// CFLAG_IndirectDraw is needed for vertex and instance offsets to work on some platforms
+		OutEnvironment.CompilerFlags.Add(CFLAG_IndirectDraw);
+	}
+};
+
+IMPLEMENT_GLOBAL_SHADER(FVisualizeLightShapeVS, "/Engine/Private/Tools/VisualizeLightShape.usf", "VisualizeLightShapeVS", SF_Vertex);
+
+class FVisualizeLightShapePS : public FGlobalShader
+{
+	DECLARE_GLOBAL_SHADER(FVisualizeLightShapePS);
+	SHADER_USE_PARAMETER_STRUCT(FVisualizeLightShapePS, FGlobalShader);
+
+	BEGIN_SHADER_PARAMETER_STRUCT(FParameters, )
+		SHADER_PARAMETER_STRUCT_REF(FViewUniformShaderParameters, ViewUniformBuffer)
+		SHADER_PARAMETER_RDG_UNIFORM_BUFFER(FSceneTextureUniformParameters, SceneTextures)
+		SHADER_PARAMETER(float, VisualizationToggleTime)
+		SHADER_PARAMETER(float, InvFadeDuration)
+		SHADER_PARAMETER(float, InvMaxDelayDistance)
+		SHADER_PARAMETER(float, MaxDelay)
+		SHADER_PARAMETER(float, MaxOpacity)
+	END_SHADER_PARAMETER_STRUCT()
+};
+
+IMPLEMENT_GLOBAL_SHADER(FVisualizeLightShapePS, "/Engine/Private/Tools/VisualizeLightShape.usf", "VisualizeLightShapePS", SF_Pixel);
+
+BEGIN_SHADER_PARAMETER_STRUCT(FVisualizeLightShapeParameters, )
+	SHADER_PARAMETER_STRUCT_INCLUDE(FVisualizeLightShapeVS::FParameters, VS)
+	SHADER_PARAMETER_STRUCT_INCLUDE(FVisualizeLightShapePS::FParameters, PS)
+	RDG_BUFFER_ACCESS(LightInfoBuffer, ERHIAccess::VertexOrIndexBuffer)
+	RENDER_TARGET_BINDING_SLOTS()
+END_SHADER_PARAMETER_STRUCT()
+
+class FVisualizeLightShapeVertexDeclaration : public FRenderResource
+{
+public:
+	FVertexDeclarationRHIRef VertexDeclarationRHI;
+
+	virtual void InitRHI(FRHICommandListBase& RHICmdList) override
+	{
+		FVertexDeclarationElementList Elements;
+		Elements.Add(FVertexElement(0, 0, VET_Float3, 0, sizeof(FVector3f)));
+		Elements.Add(FVertexElement(1, 0, VET_UShort2, 1, sizeof(VisualizeLightShape::FLightInfo), true));
+		VertexDeclarationRHI = PipelineStateCache::GetOrCreateVertexDeclaration(Elements);
+	}
+	
+	virtual void ReleaseRHI() override
+	{
+		VertexDeclarationRHI.SafeRelease();
+	}
+};
+
+static TGlobalResource<FVisualizeLightShapeVertexDeclaration> GVisualizeLightShapeVertexDeclaration;
+
+static void RenderLightShapeVisualizationView(
+	FRDGBuilder& GraphBuilder,
+	int32 ViewIndex,
+	const FViewInfo& View,
+	const FSceneTextures& SceneTextures,
+	const FTranslucencyPassResourcesMap& TranslucencyResourceMap,
+	const FSeparateTranslucencyDimensions& SeparateTranslucencyDimensions)
+{
+	const FTranslucencyPassResources& TranslucencyPassResources = TranslucencyResourceMap.Get(ViewIndex, ETranslucencyPass::TPT_TranslucencyAfterDOF);
+	const FRDGTextureMSAA& OutColorTexture = TranslucencyPassResources.IsValid() ? TranslucencyPassResources.ColorTexture : SceneTextures.Color;
+	const FRDGTextureMSAA& DepthTexture = TranslucencyPassResources.IsValid() ? TranslucencyPassResources.DepthTexture : SceneTextures.Depth;
+	const FIntRect ViewRect = TranslucencyPassResources.IsValid() ? TranslucencyPassResources.ViewRect : View.ViewRect;
+	const FViewShaderParameters ViewShaderParameters = TranslucencyPassResources.IsValid() ?
+		GetSeparateTranslucencyViewParameters(View, SeparateTranslucencyDimensions, ETranslucencyPass::TPT_TranslucencyAfterDOF) : View.GetShaderParameters();
+	
+	ESceneTextureSetupMode SceneTextureSetupMode = SceneTextures.SetupMode;
+	EnumRemoveFlags(SceneTextureSetupMode, ESceneTextureSetupMode::SceneColor);
+
+	FVisualizeLightShapeParameters* PassParameters = GraphBuilder.AllocParameters<FVisualizeLightShapeParameters>();
+	PassParameters->VS.ViewUniformBuffer = ViewShaderParameters.View;
+	PassParameters->VS.ForwardLightStruct = View.ForwardLightingResources.ForwardLightUniformBuffer;
+	PassParameters->PS.ViewUniformBuffer = ViewShaderParameters.View;
+	PassParameters->PS.SceneTextures = CreateSceneTextureUniformBuffer(GraphBuilder, &SceneTextures, View.GetFeatureLevel(), SceneTextureSetupMode);
+	PassParameters->PS.VisualizationToggleTime = View.ViewState->VisualizeLightShape.LastToggleTime;
+	PassParameters->PS.InvFadeDuration = 1.0f / VisualizeLightShape::GetFadeDuration();
+	PassParameters->PS.InvMaxDelayDistance = 1.0f / VisualizeLightShape::GetMaxDelayDistance();
+	PassParameters->PS.MaxDelay = VisualizeLightShape::GetMaxDelay();
+	PassParameters->PS.MaxOpacity = VisualizeLightShape::GetMaxOpacity();
+	PassParameters->LightInfoBuffer = View.ForwardLightingResources.VisualizeLightShapeLightInfoBuffer;
+	PassParameters->RenderTargets[0] = FRenderTargetBinding(OutColorTexture.Target, ERenderTargetLoadAction::ELoad);
+	PassParameters->RenderTargets.DepthStencil = FDepthStencilBinding(DepthTexture.Target, ERenderTargetLoadAction::ELoad, FExclusiveDepthStencil::DepthRead_StencilNop);
+
+	TShaderRef<FVisualizeLightShapeVS> VertexShaders[GNumLocalLightTypes];
+	TShaderRef<FVisualizeLightShapePS> PixelShaders[GNumLocalLightTypes];
+
+	for (uint32 LightTypeMinusOne = 0; LightTypeMinusOne < GNumLocalLightTypes; ++LightTypeMinusOne)
+	{
+		FVisualizeLightShapeVS::FPermutationDomain VSPermutationVector;
+		VSPermutationVector.Set<FVisualizeLightShapeVS::FLightType>(LightTypeMinusOne + 1);
+
+		VertexShaders[LightTypeMinusOne] = View.ShaderMap->GetShader<FVisualizeLightShapeVS>(VSPermutationVector);
+		PixelShaders[LightTypeMinusOne] = View.ShaderMap->GetShader<FVisualizeLightShapePS>();
+	}
+
+	const TConstArrayView<VisualizeLightShape::FLightInfo> SortedLights = View.ForwardLightingResources.VisualizeLightShapeLightInfos;
+
+	GraphBuilder.AddDispatchPass(
+		RDG_EVENT_NAME("VisualizeLightShapeParallel"),
+		PassParameters,
+		ERDGPassFlags::Raster,
+		[ViewRect, VertexShaders, PixelShaders, PassParameters, SortedLights](FRDGDispatchPassBuilder& DispatchPassBuilder)
+		{
+			extern TAutoConsoleVariable<int32> CVarRHICmdMinDrawsPerParallelCmdList;
+
+			const int32 NumLights = SortedLights.Num();
+			const int32 NumThreads = FMath::Min(FTaskGraphInterface::Get().GetNumWorkerThreads(), CVarRHICmdWidth.GetValueOnRenderThread());
+			const int32 NumTasks = FMath::Min(NumThreads, FMath::DivideAndRoundUp(NumLights, CVarRHICmdMinDrawsPerParallelCmdList.GetValueOnRenderThread()));
+			const int32 NumDrawsPerTask = FMath::DivideAndRoundUp(NumLights, NumTasks);
+
+			for (int32 TaskIndex = 0; TaskIndex < NumTasks; ++TaskIndex)
+			{
+				const int32 StartIndex = TaskIndex * NumDrawsPerTask;
+				const int32 NumDraws = FMath::Min(NumDrawsPerTask, NumLights - StartIndex);
+
+				FRHICommandList* RHICmdList = DispatchPassBuilder.CreateCommandList();
+
+				UE::Tasks::FTask Task = UE::Tasks::Launch(UE_SOURCE_LOCATION, [RHICmdList, StartIndex, NumDraws, ViewRect, VertexShaders, PixelShaders, PassParameters, SortedLights]()
+				{
+					FTaskTagScope Scope(ETaskTag::EParallelRenderingThread);
+					TRACE_CPUPROFILER_EVENT_SCOPE(RecordVisualizeLightShapeDraws);
+
+					for (int32 Index = 0; Index < NumDraws; ++Index)
+					{
+						const int32 LightIndex = StartIndex + Index;
+						const VisualizeLightShape::FLightInfo& LightInfo = SortedLights[LightIndex];
+						check(LightInfo.LightType > 0 && LightInfo.LightType < LIGHT_TYPE_MAX);
+
+						const int32 LightTypeMinusOne = LightInfo.LightType - 1;
+						const FVisualizationMeshBuffers& VisualizationMesh = GVisualizationMeshBuffers[LightTypeMinusOne];
+
+						if (Index == 0 || LightInfo.LightType != SortedLights[LightIndex - 1].LightType)
+						{
+							const TShaderRef<FVisualizeLightShapeVS>& VertexShader = VertexShaders[LightTypeMinusOne];
+							const TShaderRef<FVisualizeLightShapePS>& PixelShader = PixelShaders[LightTypeMinusOne];
+
+							FGraphicsPipelineStateInitializer GraphicsPSOInit;
+							RHICmdList->ApplyCachedRenderTargets(GraphicsPSOInit);
+
+							GraphicsPSOInit.RasterizerState = TStaticRasterizerState<FM_Wireframe, CM_None>::GetRHI();
+							GraphicsPSOInit.BlendState = TStaticBlendState<CW_RGBA, BO_Add, BF_SourceAlpha, BF_InverseSourceAlpha, BO_Add, BF_Zero, BF_InverseSourceAlpha>::GetRHI();
+							GraphicsPSOInit.DepthStencilState = TStaticDepthStencilState<false, CF_DepthNearOrEqual>::GetRHI();
+							GraphicsPSOInit.PrimitiveType = PT_TriangleList;
+
+							GraphicsPSOInit.BoundShaderState.VertexDeclarationRHI = GVisualizeLightShapeVertexDeclaration.VertexDeclarationRHI;
+							GraphicsPSOInit.BoundShaderState.VertexShaderRHI = VertexShader.GetVertexShader();
+							GraphicsPSOInit.BoundShaderState.PixelShaderRHI = PixelShader.GetPixelShader();
+
+							SetGraphicsPipelineState(*RHICmdList, GraphicsPSOInit, 0);
+
+							SetShaderParameters(*RHICmdList, VertexShader, VertexShader.GetVertexShader(), PassParameters->VS);
+							SetShaderParameters(*RHICmdList, PixelShader, PixelShader.GetPixelShader(), PassParameters->PS);
+
+							RHICmdList->SetViewport(ViewRect.Min.X, ViewRect.Min.Y, 0.0f, ViewRect.Max.X, ViewRect.Max.Y, 1.0f);
+
+							RHICmdList->SetStreamSource(0, VisualizationMesh.VertexBufferRHI, 0);
+							RHICmdList->SetStreamSource(1, PassParameters->LightInfoBuffer->GetRHI(), 0);
+						}
+						
+						RHICmdList->DrawIndexedPrimitive(VisualizationMesh.IndexBufferRHI, 0, LightIndex, VisualizationMesh.NumVertices, 0, VisualizationMesh.NumTriangles, 1);
+					}
+
+					if (!RHICmdList->IsSubCommandList())
+					{
+						RHICmdList->EndRenderPass();
+					}
+					RHICmdList->FinishRecording();
+				});
+
+				DispatchPassBuilder.AddPrerequisite(Task);
+			}
+		});
+}
+
+void FDeferredShadingSceneRenderer::RenderLightShapeVisualization(
+	FRDGBuilder& GraphBuilder,
+	const FSceneTextures& SceneTextures,
+	const FTranslucencyPassResourcesMap& TranslucencyResourceMap)
+{
+	check(IsInRenderingThread());
+
+	bool bAnyViewShouldDraw = false;
+	for (int32 ViewIndex = 0; ViewIndex < Views.Num(); ++ViewIndex)
+	{
+		const FViewInfo& View = Views[ViewIndex];
+		bAnyViewShouldDraw = bAnyViewShouldDraw || View.ForwardLightingResources.VisualizeLightShapeLightInfoBuffer != nullptr;
+	}
+
+	if (!bAnyViewShouldDraw)
+	{
+		return;
+	}
+
+	RDG_EVENT_SCOPE_STAT(GraphBuilder, VisualizeLightShape, "VisualizeLightShape");
+
+	VisualizeLightShape::InitVisualizationMeshes(GraphBuilder);
+
+	for (int32 ViewIndex = 0; ViewIndex < Views.Num(); ++ViewIndex)
+	{
+		const FViewInfo& View = Views[ViewIndex];
+
+		if (View.ForwardLightingResources.VisualizeLightShapeLightInfoBuffer)
+		{
+			RDG_GPU_MASK_SCOPE(GraphBuilder, View.GPUMask);
+			RDG_EVENT_SCOPE_CONDITIONAL(GraphBuilder, Views.Num() > 1, "View%d", ViewIndex);
+
+			RenderLightShapeVisualizationView(
+				GraphBuilder,
+				ViewIndex,
+				View,
+				SceneTextures,
+				TranslucencyResourceMap,
+				SeparateTranslucencyDimensions);
+		}
+	}
+}
diff --git a/Engine/Source/Runtime/Renderer/Private/VisualizeLightShape.h b/Engine/Source/Runtime/Renderer/Private/VisualizeLightShape.h
new file mode 100644
--- /dev/null
+++ b/Engine/Source/Runtime/Renderer/Private/VisualizeLightShape.h
@@ -0,0 +1,44 @@
+// Copyright Epic Games, Inc. All Rights Reserved.
+
+#pragma once
+
+#include "CoreMinimal.h"
+
+class FViewInfo;
+
+namespace VisualizeLightShape
+{
+	struct FLightInfo
+	{
+		uint16 ForwardLightIndex;
+		uint16 LightType;
+	};
+
+	struct FSortKey
+	{
+		union
+		{
+			struct 
+			{
+				FLightInfo LightData;
+				float Distance;
+			};
+			uint64 Packed;
+		};
+
+		bool operator<(const FSortKey& Other) const
+		{
+			return Packed < Other.Packed;
+		}
+	};
+
+	class FViewState
+	{
+	public:
+		double LastToggleTime = 0.0;
+	};
+
+	void Shutdown();
+
+	bool UpdateLastToggleTime(FViewInfo& View);
+}
\ No newline at end of file
