[Visualize light shape] Optimize rect light shape visualization perf by trimming vertices behind barn rects instead of simply collapsing them. Saves ~0.15 ms GPU time over a 1.6 ms visualization pass

#tests tested in editor and on a console
#Synced-CL 51949275

diff --git a/Engine/Source/Runtime/Renderer/Private/VisualizeLightShape.cpp b/Engine/Source/Runtime/Renderer/Private/VisualizeLightShape.cpp
--- a/Engine/Source/Runtime/Renderer/Private/VisualizeLightShape.cpp
+++ b/Engine/Source/Runtime/Renderer/Private/VisualizeLightShape.cpp
@@ -292,9 +292,12 @@ static FVisualizationMeshBuffers CreateSpotLightVisualizationMesh(FRDGBuilder& G
 
 static FVisualizationMeshBuffers CreateRectLightVisualizationMesh(FRDGBuilder& GraphBuilder)
 {
-	const int32 NumSections = 14;
-	const int32 NumVertices = (NumSections + 1) * (NumSections * 4 + 1);
-	const int32 NumTriangles = NumSections * NumSections * 4 * 2;
+	const bool bTrimBehindBarnRect = true;
+	const int32 NumSectionsY = 14;
+	const int32 NumSectionsX = bTrimBehindBarnRect ? NumSectionsY / 2 + 1 : NumSectionsY;
+	const int32 BiasX = NumSectionsY - NumSectionsX;
+	const int32 NumVertices = (NumSectionsX + 1) * (NumSectionsY * 4 + 1);
+	const int32 NumTriangles = NumSectionsX * NumSectionsY * 4 * 2;
 	const float Length = 1.0f;
 
 	FRHICommandListImmediate& RHICmdList = GraphBuilder.RHICmdList;
@@ -325,13 +328,13 @@ static FVisualizationMeshBuffers CreateRectLightVisualizationMesh(FRDGBuilder& G
 	int32 VertexIndex = 0;
 	for (int32 Side = 0; Side < 4; ++Side)
 	{
-		for (int32 Y = 0; Y <= NumSections + (Side == 3 ? 0 : -1); ++Y)
+		for (int32 Y = 0; Y <= NumSectionsY + (Side == 3 ? 0 : -1); ++Y)
 		{
-			for (int32 X = 0; X <= NumSections; ++X)
+			for (int32 X = 0; X <= NumSectionsX; ++X)
 			{
-				VertexData[VertexIndex][0] = Length / NumSections * X - Length * 0.5f;
+				VertexData[VertexIndex][0] = Length / NumSectionsY * (BiasX + X) - Length * 0.5f;
 				VertexData[VertexIndex][Side % 2 + 1] = Length * 0.5f * (Side >= 2 ? -1.0f : 1.0f);
-				VertexData[VertexIndex][(Side + 1) % 2 + 1] = (Length / NumSections * Y - Length * 0.5f) * (Side == 1 || Side == 2 ? -1.0f : 1.0f);
+				VertexData[VertexIndex][(Side + 1) % 2 + 1] = (Length / NumSectionsY * Y - Length * 0.5f) * (Side == 1 || Side == 2 ? -1.0f : 1.0f);
 				++VertexIndex;
 			}
 		}
@@ -341,21 +344,21 @@ static FVisualizationMeshBuffers CreateRectLightVisualizationMesh(FRDGBuilder& G
 	int32 TriangleIndex = 0;
 	for (int32 Side = 0; Side < 4; ++Side)
 	{
-		const int32 SideOffset = (NumSections + 1) * NumSections * Side;
+		const int32 SideOffset = (NumSectionsX + 1) * NumSectionsY * Side;
 
-		for (int32 Y = 0; Y < NumSections; ++Y)
+		for (int32 Y = 0; Y < NumSectionsY; ++Y)
 		{
-			const int32 Offset = SideOffset + (NumSections + 1) * Y;
+			const int32 Offset = SideOffset + (NumSectionsX + 1) * Y;
 
-			for (int32 X = 0; X < NumSections; ++X)
+			for (int32 X = 0; X < NumSectionsX; ++X)
 			{
 				// I3 - I2
 				// |  \  |
 				// I0 - I1
 				const int32 I0 = Offset + X;
 				const int32 I1 = Offset + X + 1;
-				const int32 I2 = Offset + X + (NumSections + 1) + 1;
-				const int32 I3 = Offset + X + (NumSections + 1);
+				const int32 I2 = Offset + X + (NumSectionsX + 1) + 1;
+				const int32 I3 = Offset + X + (NumSectionsX + 1);
 
 				IndexData[TriangleIndex * 3 + 0] = I0;
 				IndexData[TriangleIndex * 3 + 1] = I1;
