Mobile: Cleaned up the SM5Nodes cvar related to light rendering. Capsule lights are now always supported in mobile deferred shading (including translucency) and disabled in mobile forward

diff --git a/Engine/Shaders/Private/LightData.ush b/Engine/Shaders/Private/LightData.ush
--- a/Engine/Shaders/Private/LightData.ush
+++ b/Engine/Shaders/Private/LightData.ush
@@ -14,8 +14,9 @@
 
 #define MAX_RECT_ATLAS_MIP			32
 
+// Capsule lights are supported by default in mobile deferred shading (including translucency), and are disabled in mobile forward
 #ifndef MOBILE_SHADING_PATH_SUPPORT_CAPSULE_LIGHT
-#define MOBILE_SHADING_PATH_SUPPORT_CAPSULE_LIGHT 0
+#define MOBILE_SHADING_PATH_SUPPORT_CAPSULE_LIGHT MOBILE_DEFERRED_SHADING
 #endif
 
 ///////////////////////////////////////////////////////////////////////////////////////////////////
diff --git a/Engine/Source/Runtime/Renderer/Private/LightGridInjection.cpp b/Engine/Source/Runtime/Renderer/Private/LightGridInjection.cpp
--- a/Engine/Source/Runtime/Renderer/Private/LightGridInjection.cpp
+++ b/Engine/Source/Runtime/Renderer/Private/LightGridInjection.cpp
@@ -1488,99 +1488,95 @@ FComputeLightGridOutput FSceneRenderer::ComputeLightGrid(FRDGBuilder& GraphBuild
 							View.ForwardLightingResources.SelectedForwardDirectionalLightProxy = LightProxy;
 							View.ForwardLightingResources.SelectedForwardDirectionalLightIndex = IndexInBuffer;
 
-							// On mobile there is a separate FMobileDirectionalLightShaderParameters UB which holds all directional light data.
-							if(!IsMobilePlatform(ShaderPlatform) || MobileSupportsSM5MaterialNodes(ShaderPlatform))
+							ForwardLightUniformParameters->HasDirectionalLight = 1;
+							ForwardLightUniformParameters->DirectionalLightColor = FVector3f(LightParameters.Color);
+							if (LightProxy->GetUsePerPixelAtmosphereTransmittance())
 							{
-								ForwardLightUniformParameters->HasDirectionalLight = 1;
-								ForwardLightUniformParameters->DirectionalLightColor = FVector3f(LightParameters.Color);
-								if (LightProxy->GetUsePerPixelAtmosphereTransmittance())
-								{
-									// When using PerPixelTransmittance, transmittance is evaluated per pixel by sampling the transmittance texture. It gives better gradient on large scale objects such as mountains.
-									// However, to skip doing that texture sampling in translucent using ForwardShading/VolumetricFog, we use the simple planet top ground transmittance as a simplification.
-									// That will work for most of the cases for most of the map/terrain at the top of the virtual planet.
-									ForwardLightUniformParameters->DirectionalLightColor *= FVector3f(LightProxy->GetAtmosphereTransmittanceTowardSun());
-								}
-								ForwardLightUniformParameters->DirectionalLightVolumetricScatteringIntensity = LightProxy->GetVolumetricScatteringIntensity();
-								ForwardLightUniformParameters->DirectionalLightSpecularScale = FMath::Clamp(LightProxy->GetSpecularScale(), 0.f, 1.f);
-								ForwardLightUniformParameters->DirectionalLightDiffuseScale = FMath::Clamp(LightProxy->GetDiffuseScale(), 0.f, 1.f);
-								ForwardLightUniformParameters->DirectionalLightDirection = LightParameters.Direction;
-								ForwardLightUniformParameters->DirectionalLightSourceRadius = LightParameters.SourceRadius;
-								ForwardLightUniformParameters->DirectionalLightSoftSourceRadius = LightParameters.SoftSourceRadius;
-								ForwardLightUniformParameters->DirectionalLightSceneInfoExtraDataPacked = LightSceneInfoExtraDataPacked;
-								ForwardLightUniformParameters->DirectionalLightVSM = INDEX_NONE;
-								ForwardLightUniformParameters->DirectionalLightLightFunctionAtlasLightIndex = LightParameters.LightFunctionAtlasLightIndex;
-								ForwardLightUniformParameters->DirectionalLightAffectsTranslucentLighting = LightParameters.bAffectsTranslucentLighting ? 1 : 0;
-								ForwardLightUniformParameters->DirectionalLightHandledByMegaLights = SortedLightInfo.SortKey.Fields.bHandledByMegaLights ? 1 : 0;
+								// When using PerPixelTransmittance, transmittance is evaluated per pixel by sampling the transmittance texture. It gives better gradient on large scale objects such as mountains.
+								// However, to skip doing that texture sampling in translucent using ForwardShading/VolumetricFog, we use the simple planet top ground transmittance as a simplification.
+								// That will work for most of the cases for most of the map/terrain at the top of the virtual planet.
+								ForwardLightUniformParameters->DirectionalLightColor *= FVector3f(LightProxy->GetAtmosphereTransmittanceTowardSun());
+							}
+							ForwardLightUniformParameters->DirectionalLightVolumetricScatteringIntensity = LightProxy->GetVolumetricScatteringIntensity();
+							ForwardLightUniformParameters->DirectionalLightSpecularScale = FMath::Clamp(LightProxy->GetSpecularScale(), 0.f, 1.f);
+							ForwardLightUniformParameters->DirectionalLightDiffuseScale = FMath::Clamp(LightProxy->GetDiffuseScale(), 0.f, 1.f);
+							ForwardLightUniformParameters->DirectionalLightDirection = LightParameters.Direction;
+							ForwardLightUniformParameters->DirectionalLightSourceRadius = LightParameters.SourceRadius;
+							ForwardLightUniformParameters->DirectionalLightSoftSourceRadius = LightParameters.SoftSourceRadius;
+							ForwardLightUniformParameters->DirectionalLightSceneInfoExtraDataPacked = LightSceneInfoExtraDataPacked;
+							ForwardLightUniformParameters->DirectionalLightVSM = INDEX_NONE;
+							ForwardLightUniformParameters->DirectionalLightLightFunctionAtlasLightIndex = LightParameters.LightFunctionAtlasLightIndex;
+							ForwardLightUniformParameters->DirectionalLightAffectsTranslucentLighting = LightParameters.bAffectsTranslucentLighting ? 1 : 0;
+							ForwardLightUniformParameters->DirectionalLightHandledByMegaLights = SortedLightInfo.SortKey.Fields.bHandledByMegaLights ? 1 : 0;
 
-								const FVector2D FadeParams = LightProxy->GetDirectionalLightDistanceFadeParameters(View.GetFeatureLevel(), LightSceneInfo->IsPrecomputedLightingValid(), View.MaxShadowCascades);
+							const FVector2D FadeParams = LightProxy->GetDirectionalLightDistanceFadeParameters(View.GetFeatureLevel(), LightSceneInfo->IsPrecomputedLightingValid(), View.MaxShadowCascades);
 
-								ForwardLightUniformParameters->DirectionalLightDistanceFadeMAD = FVector2f(FadeParams.Y, -FadeParams.X * FadeParams.Y);	// LWC_TODO: Precision loss
+							ForwardLightUniformParameters->DirectionalLightDistanceFadeMAD = FVector2f(FadeParams.Y, -FadeParams.X * FadeParams.Y);	// LWC_TODO: Precision loss
 
-								const FMatrix TranslatedWorldToWorld = FTranslationMatrix(-View.ViewMatrices.GetPreViewTranslation());
+							const FMatrix TranslatedWorldToWorld = FTranslationMatrix(-View.ViewMatrices.GetPreViewTranslation());
 
-								if (bDynamicShadows)
-								{
-									const TArray<FProjectedShadowInfo*, SceneRenderingAllocator>& DirectionalLightShadowInfos = VisibleLightInfos[LightSceneInfo->Id].AllProjectedShadows;
+							if (bDynamicShadows)
+							{
+								const TArray<FProjectedShadowInfo*, SceneRenderingAllocator>& DirectionalLightShadowInfos = VisibleLightInfos[LightSceneInfo->Id].AllProjectedShadows;
 
-									ForwardLightUniformParameters->DirectionalLightVSM = VirtualShadowMapId;
+								ForwardLightUniformParameters->DirectionalLightVSM = VirtualShadowMapId;
 
-									ForwardLightUniformParameters->NumDirectionalLightCascades = 0;
-									// Unused cascades should compare > all scene depths
-									ForwardLightUniformParameters->CascadeEndDepths = FVector4f(MAX_FLT, MAX_FLT, MAX_FLT, MAX_FLT);
+								ForwardLightUniformParameters->NumDirectionalLightCascades = 0;
+								// Unused cascades should compare > all scene depths
+								ForwardLightUniformParameters->CascadeEndDepths = FVector4f(MAX_FLT, MAX_FLT, MAX_FLT, MAX_FLT);
 
-									for (const FProjectedShadowInfo* ShadowInfo : DirectionalLightShadowInfos)
+								for (const FProjectedShadowInfo* ShadowInfo : DirectionalLightShadowInfos)
+								{
+									if (ShadowInfo->DependentView)
 									{
-										if (ShadowInfo->DependentView)
+										// when rendering stereo views, allow using the shadows rendered for the primary view as 'close enough'
+										if (ShadowInfo->DependentView != &View && ShadowInfo->DependentView != View.GetPrimaryView())
 										{
-											// when rendering stereo views, allow using the shadows rendered for the primary view as 'close enough'
-											if (ShadowInfo->DependentView != &View && ShadowInfo->DependentView != View.GetPrimaryView())
-											{
-												continue;
-											}
+											continue;
 										}
+									}
+
+									const int32 CascadeIndex = ShadowInfo->CascadeSettings.ShadowSplitIndex;
+
+									if (ShadowInfo->IsWholeSceneDirectionalShadow() && !ShadowInfo->HasVirtualShadowMap() && ShadowInfo->bAllocated && CascadeIndex < GMaxForwardShadowCascades)
+									{
+										const FMatrix WorldToShadow = ShadowInfo->GetWorldToShadowMatrix(ForwardLightUniformParameters->DirectionalLightShadowmapMinMax[CascadeIndex]);
+										const FMatrix44f TranslatedWorldToShadow = FMatrix44f(TranslatedWorldToWorld * WorldToShadow);
 
-										const int32 CascadeIndex = ShadowInfo->CascadeSettings.ShadowSplitIndex;
+										ForwardLightUniformParameters->NumDirectionalLightCascades++;
+										ForwardLightUniformParameters->DirectionalLightTranslatedWorldToShadowMatrix[CascadeIndex] = TranslatedWorldToShadow;
+										ForwardLightUniformParameters->CascadeEndDepths[CascadeIndex] = ShadowInfo->CascadeSettings.SplitFar;
 
-										if (ShadowInfo->IsWholeSceneDirectionalShadow() && !ShadowInfo->HasVirtualShadowMap() && ShadowInfo->bAllocated && CascadeIndex < GMaxForwardShadowCascades)
+										if (CascadeIndex == 0)
 										{
-											const FMatrix WorldToShadow = ShadowInfo->GetWorldToShadowMatrix(ForwardLightUniformParameters->DirectionalLightShadowmapMinMax[CascadeIndex]);
-											const FMatrix44f TranslatedWorldToShadow = FMatrix44f(TranslatedWorldToWorld * WorldToShadow);
-
-											ForwardLightUniformParameters->NumDirectionalLightCascades++;
-											ForwardLightUniformParameters->DirectionalLightTranslatedWorldToShadowMatrix[CascadeIndex] = TranslatedWorldToShadow;
-											ForwardLightUniformParameters->CascadeEndDepths[CascadeIndex] = ShadowInfo->CascadeSettings.SplitFar;
-
-											if (CascadeIndex == 0)
-											{
-												ForwardLightUniformParameters->DirectionalLightShadowmapAtlas = GraphBuilder.RegisterExternalTexture(ShadowInfo->RenderTargets.DepthTarget);
-												ForwardLightUniformParameters->DirectionalLightDepthBias = ShadowInfo->GetShaderDepthBias();
-												FVector2D AtlasSize = ForwardLightUniformParameters->DirectionalLightShadowmapAtlas->Desc.Extent;
-												ForwardLightUniformParameters->DirectionalLightShadowmapAtlasBufferSize = FVector4f(AtlasSize.X, AtlasSize.Y, 1.0f / AtlasSize.X, 1.0f / AtlasSize.Y);
-											}
+											ForwardLightUniformParameters->DirectionalLightShadowmapAtlas = GraphBuilder.RegisterExternalTexture(ShadowInfo->RenderTargets.DepthTarget);
+											ForwardLightUniformParameters->DirectionalLightDepthBias = ShadowInfo->GetShaderDepthBias();
+											FVector2D AtlasSize = ForwardLightUniformParameters->DirectionalLightShadowmapAtlas->Desc.Extent;
+											ForwardLightUniformParameters->DirectionalLightShadowmapAtlasBufferSize = FVector4f(AtlasSize.X, AtlasSize.Y, 1.0f / AtlasSize.X, 1.0f / AtlasSize.Y);
 										}
 									}
 								}
+							}
 
-								const FStaticShadowDepthMap* StaticShadowDepthMap = LightSceneInfo->Proxy->GetStaticShadowDepthMap();
-								const uint32 bStaticallyShadowedValue = LightSceneInfo->IsPrecomputedLightingValid() 
-																		&& StaticShadowDepthMap 
-																		&& StaticShadowDepthMap->Data 
-																		&& !StaticShadowDepthMap->Data->WorldToLight.ContainsNaN()
-																		&& StaticShadowDepthMap->TextureRHI ? 1 : 0;
-								ForwardLightUniformParameters->DirectionalLightUseStaticShadowing = bStaticallyShadowedValue;
-								if (bStaticallyShadowedValue)
-								{
-									const FMatrix44f TranslatedWorldToShadow = FMatrix44f(TranslatedWorldToWorld * StaticShadowDepthMap->Data->WorldToLight);
-									ForwardLightUniformParameters->DirectionalLightStaticShadowBufferSize = FVector4f(StaticShadowDepthMap->Data->ShadowMapSizeX, StaticShadowDepthMap->Data->ShadowMapSizeY, 1.0f / StaticShadowDepthMap->Data->ShadowMapSizeX, 1.0f / StaticShadowDepthMap->Data->ShadowMapSizeY);
-									ForwardLightUniformParameters->DirectionalLightTranslatedWorldToStaticShadow = TranslatedWorldToShadow;
-									ForwardLightUniformParameters->DirectionalLightStaticShadowmap = StaticShadowDepthMap->TextureRHI;
-								}
-								else
-								{
-									ForwardLightUniformParameters->DirectionalLightStaticShadowBufferSize = FVector4f(0, 0, 0, 0);
-									ForwardLightUniformParameters->DirectionalLightTranslatedWorldToStaticShadow = FMatrix44f::Identity;
-									ForwardLightUniformParameters->DirectionalLightStaticShadowmap = GWhiteTexture->TextureRHI;
-								}
+							const FStaticShadowDepthMap* StaticShadowDepthMap = LightSceneInfo->Proxy->GetStaticShadowDepthMap();
+							const uint32 bStaticallyShadowedValue = LightSceneInfo->IsPrecomputedLightingValid() 
+																	&& StaticShadowDepthMap 
+																	&& StaticShadowDepthMap->Data 
+																	&& !StaticShadowDepthMap->Data->WorldToLight.ContainsNaN()
+																	&& StaticShadowDepthMap->TextureRHI ? 1 : 0;
+							ForwardLightUniformParameters->DirectionalLightUseStaticShadowing = bStaticallyShadowedValue;
+							if (bStaticallyShadowedValue)
+							{
+								const FMatrix44f TranslatedWorldToShadow = FMatrix44f(TranslatedWorldToWorld * StaticShadowDepthMap->Data->WorldToLight);
+								ForwardLightUniformParameters->DirectionalLightStaticShadowBufferSize = FVector4f(StaticShadowDepthMap->Data->ShadowMapSizeX, StaticShadowDepthMap->Data->ShadowMapSizeY, 1.0f / StaticShadowDepthMap->Data->ShadowMapSizeX, 1.0f / StaticShadowDepthMap->Data->ShadowMapSizeY);
+								ForwardLightUniformParameters->DirectionalLightTranslatedWorldToStaticShadow = TranslatedWorldToShadow;
+								ForwardLightUniformParameters->DirectionalLightStaticShadowmap = StaticShadowDepthMap->TextureRHI;
+							}
+							else
+							{
+								ForwardLightUniformParameters->DirectionalLightStaticShadowBufferSize = FVector4f(0, 0, 0, 0);
+								ForwardLightUniformParameters->DirectionalLightTranslatedWorldToStaticShadow = FMatrix44f::Identity;
+								ForwardLightUniformParameters->DirectionalLightStaticShadowmap = GWhiteTexture->TextureRHI;
 							}
 						}
 					}
diff --git a/Engine/Source/Runtime/Renderer/Private/MobileBasePassRendering.h b/Engine/Source/Runtime/Renderer/Private/MobileBasePassRendering.h
--- a/Engine/Source/Runtime/Renderer/Private/MobileBasePassRendering.h
+++ b/Engine/Source/Runtime/Renderer/Private/MobileBasePassRendering.h
@@ -439,10 +439,7 @@ public:
 		OutEnvironment.SetDefine(TEXT("MOBILE_TRANSLUCENT_COLOR_TRANSMITTANCE_PROGRAMMABLE_BLENDING"), TranslucentColorTransmittanceMode == EMobileTranslucentColorTransmittanceMode::PROGRAMMABLE_BLENDING ? 1u : 0u);
 		OutEnvironment.SetDefine(TEXT("MOBILE_TRANSLUCENT_COLOR_TRANSMITTANCE_SINGLE_SRC_BLENDING"), TranslucentColorTransmittanceMode == EMobileTranslucentColorTransmittanceMode::SINGLE_SRC_BLENDING ? 1u : 0u);
 		OutEnvironment.SetDefine(TEXT("MOBILE_SSR_ENABLED"), (!bDeferredShadingEnabled && AreMobileScreenSpaceReflectionsEnabled(Parameters.Platform)) ? 1u : 0u);
-
-		const bool bSupportCapsule = MobileSupportsSM5MaterialNodes(Parameters.Platform);
-		OutEnvironment.SetDefine(TEXT("MOBILE_SHADING_PATH_SUPPORT_CAPSULE_LIGHT"), bSupportCapsule ? 1u : 0u);
-
+			
 		TMobileBasePassPSBaseType<LightMapPolicyType>::ModifyCompilationEnvironment(Parameters, OutEnvironment);
 	}
 	
