-Consolidate PreviewFeatureLevelChanged and PreviewPlatformChanged into a single PreviewShaderPlatformChanged(EShaderPlatform) delegate.
-Cache EShaderPlatform on FSceneInterface at construction so GetShaderPlatform() is stable for the scene's lifetime.
-Fix renderer GetGlobalShaderMap calls to use Scene.GetShaderPlatform() instead of Scene.GetFeatureLevel()
-Deprecate ChangeFeatureLevel, OnPreviewPlatformChanged  and OnPreviewFeatureLevelChanged.

#Synced-CL 52056483

diff --git a/Engine/Plugins/Enterprise/DataprepEditor/Source/DataprepEditor/Private/Widgets/SDataprepEditorViewport.cpp b/Engine/Plugins/Enterprise/DataprepEditor/Source/DataprepEditor/Private/Widgets/SDataprepEditorViewport.cpp
--- a/Engine/Plugins/Enterprise/DataprepEditor/Source/DataprepEditor/Private/Widgets/SDataprepEditorViewport.cpp
+++ b/Engine/Plugins/Enterprise/DataprepEditor/Source/DataprepEditor/Private/Widgets/SDataprepEditorViewport.cpp
@@ -20,6 +20,7 @@
 #include "EngineLogs.h"
 #include "StaticMeshSceneProxy.h"
 #include "Engine/TextureCube.h"
+#include "SceneInterface.h"
 #include "EngineUtils.h"
 #include "Framework/Commands/UICommandList.h"
 #include "Materials/Material.h"
@@ -221,7 +222,7 @@ SDataprepEditorViewport::SDataprepEditorViewport()
 
 SDataprepEditorViewport::~SDataprepEditorViewport()
 {
-	CastChecked<UEditorEngine>(GEngine)->OnPreviewFeatureLevelChanged().Remove(PreviewFeatureLevelChangedHandle);
+	CastChecked<UEditorEngine>(GEngine)->OnPreviewShaderPlatformChanged().Remove(PreviewFeatureLevelChangedHandle);
 	ClearScene();
 }
 
@@ -233,16 +234,16 @@ void SDataprepEditorViewport::Construct(const FArguments& InArgs, TSharedPtr<FDa
 
 	// restore last used feature level
 	UWorld* PreviewSceneWorld = PreviewScene->GetWorld();
-	if (ensure(PreviewSceneWorld))
+	if (ensure(PreviewSceneWorld) && GWorld != nullptr && GWorld->Scene != nullptr)
 	{
-		PreviewSceneWorld->ChangeFeatureLevel(GWorld->GetFeatureLevel());
+		PreviewSceneWorld->ShaderPlatformChanged(GWorld->Scene->GetShaderPlatform());
 	}
 
-	// Listen to and act on changes in feature level
+	// Listen to and act on changes in preview platform
 	UEditorEngine* Editor = CastChecked<UEditorEngine>(GEngine);
-	PreviewFeatureLevelChangedHandle = Editor->OnPreviewFeatureLevelChanged().AddLambda([&PreviewSceneWorld](ERHIFeatureLevel::Type NewFeatureLevel)
+	PreviewFeatureLevelChangedHandle = Editor->OnPreviewShaderPlatformChanged().AddLambda([&PreviewSceneWorld](EShaderPlatform NewShaderPlatform)
 	{
-		PreviewSceneWorld->ChangeFeatureLevel(NewFeatureLevel);
+		PreviewSceneWorld->ShaderPlatformChanged(NewShaderPlatform);
 	});
 
 	// Create root actor of 3D viewport's world
diff --git a/Engine/Plugins/FX/Cascade/Source/Cascade/Private/SCascadePreviewViewport.cpp b/Engine/Plugins/FX/Cascade/Source/Cascade/Private/SCascadePreviewViewport.cpp
--- a/Engine/Plugins/FX/Cascade/Source/Cascade/Private/SCascadePreviewViewport.cpp
+++ b/Engine/Plugins/FX/Cascade/Source/Cascade/Private/SCascadePreviewViewport.cpp
@@ -11,6 +11,7 @@
 #include "EditorViewportClient.h"
 #include "Engine/Engine.h"
 #include "Engine/World.h"
+#include "SceneInterface.h"
 #include "Framework/Application/SlateApplication.h"
 #include "Math/BoxSphereBounds.h"
 #include "Math/UnrealMathSSE.h"
@@ -40,7 +41,7 @@ void SCascadePreviewViewport::Tick(const FGeometry& AllottedGeometry, const doub
 SCascadePreviewViewport::~SCascadePreviewViewport()
 {
 	UEditorEngine* Editor = (UEditorEngine*)GEngine;
-	Editor->OnPreviewFeatureLevelChanged().Remove(PreviewFeatureLevelChangedHandle);
+	Editor->OnPreviewShaderPlatformChanged().Remove(PreviewFeatureLevelChangedHandle);
 
 	if (ViewportClient.IsValid())
 	{
@@ -58,22 +59,22 @@ void SCascadePreviewViewport::Construct(const FArguments& InArgs)
 	if (ViewportClient.IsValid())
 	{
 		UWorld* World = ViewportClient->GetPreviewScene().GetWorld();
-		if (World != nullptr)
+		if (World != nullptr && GWorld != nullptr && GWorld->Scene != nullptr)
 		{
-			World->ChangeFeatureLevel(GWorld->GetFeatureLevel());
+			World->ShaderPlatformChanged(GWorld->Scene->GetShaderPlatform());
 		}
 	}
 
-	// Use a delegate to inform the attached world of feature level changes.
+	// Use a delegate to inform the attached world of preview platform changes.
 	UEditorEngine* Editor = (UEditorEngine*)GEngine;
-	PreviewFeatureLevelChangedHandle = Editor->OnPreviewFeatureLevelChanged().AddLambda([this](ERHIFeatureLevel::Type NewFeatureLevel)
+	PreviewFeatureLevelChangedHandle = Editor->OnPreviewShaderPlatformChanged().AddLambda([this](EShaderPlatform NewShaderPlatform)
 		{
 			if(ViewportClient.IsValid())
 			{
 				UWorld* World = ViewportClient->GetPreviewScene().GetWorld();
 				if (World != nullptr)
 				{
-					World->ChangeFeatureLevel(NewFeatureLevel);
+					World->ShaderPlatformChanged(NewShaderPlatform);
 				}
 			}
 		});
diff --git a/Engine/Plugins/FX/Niagara/Source/NiagaraEditor/Private/NiagaraEditorModule.cpp b/Engine/Plugins/FX/Niagara/Source/NiagaraEditor/Private/NiagaraEditorModule.cpp
--- a/Engine/Plugins/FX/Niagara/Source/NiagaraEditor/Private/NiagaraEditorModule.cpp
+++ b/Engine/Plugins/FX/Niagara/Source/NiagaraEditor/Private/NiagaraEditorModule.cpp
@@ -1,6 +1,7 @@
 // Copyright Epic Games, Inc. All Rights Reserved.
 
 #include "NiagaraEditorModule.h"
+#include "DataDrivenShaderPlatformInfo.h"
 #include "NiagaraModule.h"
 #include "NiagaraEditorTickables.h"
 #include "Modules/ModuleManager.h"
@@ -1148,8 +1149,7 @@ void FNiagaraEditorModule::OnPreExit()
 
 	if (GEditor)
 	{
-		CastChecked<UEditorEngine>(GEngine)->OnPreviewPlatformChanged().Remove(PreviewPlatformChangedHandle);
-		CastChecked<UEditorEngine>(GEngine)->OnPreviewFeatureLevelChanged().Remove(PreviewFeatureLevelChangedHandle);
+		CastChecked<UEditorEngine>(GEngine)->OnPreviewShaderPlatformChanged().Remove(PreviewPlatformChangedHandle);
 	}
 
 	// Ensure that we don't have any lingering compiles laying around that will explode after this module shuts down.
@@ -1897,9 +1897,7 @@ void FNiagaraEditorModule::OnPostEngineInit()
 		GEditor->OnExecParticleInvoked().AddRaw(this, &FNiagaraEditorModule::OnExecParticleInvoked);
 
 		UEditorEngine* EditorEngine = CastChecked<UEditorEngine>(GEngine);
-		PreviewPlatformChangedHandle = EditorEngine->OnPreviewPlatformChanged().AddRaw(this, &FNiagaraEditorModule::OnPreviewPlatformChanged);
-
-		PreviewFeatureLevelChangedHandle = EditorEngine->OnPreviewFeatureLevelChanged().AddStatic(UNiagaraScript::SetPreviewFeatureLevel);
+		PreviewPlatformChangedHandle = EditorEngine->OnPreviewShaderPlatformChanged().AddRaw(this, &FNiagaraEditorModule::OnPreviewPlatformChanged);
 
 		// Ensure we have the right feature level set as the editor may already be in one before we get here
 		UNiagaraScript::SetPreviewFeatureLevel(GEditor->DefaultWorldFeatureLevel);
@@ -1960,10 +1958,12 @@ void FNiagaraEditorModule::OnDeviceProfileManagerUpdated()
 	FNiagaraPlatformSet::InvalidateCachedData();
 }
 
-void FNiagaraEditorModule::OnPreviewPlatformChanged()
+void FNiagaraEditorModule::OnPreviewPlatformChanged(EShaderPlatform NewShaderPlatform)
 {
 	FNiagaraPlatformSet::InvalidateCachedData();
 
+	UNiagaraScript::SetPreviewFeatureLevel(GetMaxSupportedFeatureLevel(NewShaderPlatform));
+
 	for (TObjectIterator<UNiagaraSystem> It; It; ++It)
 	{
 		UNiagaraSystem* System = *It;
diff --git a/Engine/Plugins/FX/Niagara/Source/NiagaraEditor/Private/Widgets/SNiagaraSystemViewport.cpp b/Engine/Plugins/FX/Niagara/Source/NiagaraEditor/Private/Widgets/SNiagaraSystemViewport.cpp
--- a/Engine/Plugins/FX/Niagara/Source/NiagaraEditor/Private/Widgets/SNiagaraSystemViewport.cpp
+++ b/Engine/Plugins/FX/Niagara/Source/NiagaraEditor/Private/Widgets/SNiagaraSystemViewport.cpp
@@ -884,15 +884,15 @@ void SNiagaraSystemViewport::Construct(const FArguments& InArgs, TSharedRef<FNia
 
 	Client->EngineShowFlags.SetGrid(Settings->IsShowGridInViewport());
 
-	// Register for preview feature level changes
+	// Register for preview platform changes
 	if (UEditorEngine* EditorEngine = Cast<UEditorEngine>(GEngine))
 	{
-		OnPreviewFeatureLevelChangedHandle = EditorEngine->OnPreviewFeatureLevelChanged().AddLambda(
-			[WeakWorld=TWeakObjectPtr<UWorld>(Client->GetWorld())](ERHIFeatureLevel::Type NewFeatureLevel)
+		OnPreviewFeatureLevelChangedHandle = EditorEngine->OnPreviewShaderPlatformChanged().AddLambda(
+			[WeakWorld=TWeakObjectPtr<UWorld>(Client->GetWorld())](EShaderPlatform NewShaderPlatform)
 			{
 				if (UWorld* World = WeakWorld.Get())
 				{
-					World->ChangeFeatureLevel(NewFeatureLevel);
+					World->ShaderPlatformChanged(NewShaderPlatform);
 				}
 			}
 		);
@@ -907,7 +907,7 @@ SNiagaraSystemViewport::~SNiagaraSystemViewport()
 	{
 		if (UEditorEngine* EditorEngine = Cast<UEditorEngine>(GEngine))
 		{
-			EditorEngine->OnPreviewFeatureLevelChanged().Remove(OnPreviewFeatureLevelChangedHandle);
+			EditorEngine->OnPreviewShaderPlatformChanged().Remove(OnPreviewFeatureLevelChangedHandle);
 		}
 	}
 
diff --git a/Engine/Plugins/FX/Niagara/Source/NiagaraEditor/Public/NiagaraEditorModule.h b/Engine/Plugins/FX/Niagara/Source/NiagaraEditor/Public/NiagaraEditorModule.h
--- a/Engine/Plugins/FX/Niagara/Source/NiagaraEditor/Public/NiagaraEditorModule.h
+++ b/Engine/Plugins/FX/Niagara/Source/NiagaraEditor/Public/NiagaraEditorModule.h
@@ -3,6 +3,7 @@
 #pragma once
 
 #include "CoreMinimal.h"
+#include "RHIShaderPlatform.h"
 #include "AssetToolsModule.h"
 #include "AssetRegistry/AssetRegistryModule.h"
 #include "Modules/ModuleInterface.h"
@@ -432,7 +433,7 @@ private:
 	void OnExecParticleInvoked(const TCHAR* InStr);
 	void OnPostEngineInit();
 	void OnDeviceProfileManagerUpdated();
-	void OnPreviewPlatformChanged();
+	void OnPreviewPlatformChanged(EShaderPlatform NewShaderPlatform);
 	void OnPreExit();
 	void PostGarbageCollect();
 
diff --git a/Engine/Plugins/PCG/Source/PCG/Private/Elements/ControlFlow/PCGPlatformSwitchBase.cpp b/Engine/Plugins/PCG/Source/PCG/Private/Elements/ControlFlow/PCGPlatformSwitchBase.cpp
--- a/Engine/Plugins/PCG/Source/PCG/Private/Elements/ControlFlow/PCGPlatformSwitchBase.cpp
+++ b/Engine/Plugins/PCG/Source/PCG/Private/Elements/ControlFlow/PCGPlatformSwitchBase.cpp
@@ -50,12 +50,7 @@ void UPCGPlatformSwitchSettingsBase::SetupPreviewPlatformChangeEvents()
 	{
 		if (!PreviewPlatformChangedHandle.IsValid())
 		{
-			PreviewPlatformChangedHandle = GEditor->OnPreviewPlatformChanged().AddUObject(this, &UPCGPlatformSwitchSettingsBase::OnPreviewPlatformChanged);
-		}
-
-		if (!PreviewFeatureLevelChangedHandle.IsValid())
-		{
-			PreviewFeatureLevelChangedHandle = GEditor->OnPreviewFeatureLevelChanged().AddUObject(this, &UPCGPlatformSwitchSettingsBase::OnPreviewFeatureLevelChanged);
+			PreviewPlatformChangedHandle = GEditor->OnPreviewShaderPlatformChanged().AddUObject(this, &UPCGPlatformSwitchSettingsBase::OnPreviewPlatformChanged);
 		}
 	}
 }
@@ -66,23 +61,12 @@ void UPCGPlatformSwitchSettingsBase::TearDownPreviewPlatformChangeEvents()
 	{
 		if (PreviewPlatformChangedHandle.IsValid())
 		{
-			GEditor->OnPreviewPlatformChanged().Remove(PreviewPlatformChangedHandle);
-		}
-
-		if (PreviewFeatureLevelChangedHandle.IsValid())
-		{
-			GEditor->OnPreviewFeatureLevelChanged().Remove(PreviewFeatureLevelChangedHandle);
+			GEditor->OnPreviewShaderPlatformChanged().Remove(PreviewPlatformChangedHandle);
 		}
 	}
 }
 
-void UPCGPlatformSwitchSettingsBase::OnPreviewPlatformChanged()
-{
-	// We don't broadcast structural from here. The PCGSubsystem will refresh all runtime gen components when preview platform changes.
-	OnSettingsChangedDelegate.Broadcast(this, EPCGChangeType::Cosmetic);
-}
-
-void UPCGPlatformSwitchSettingsBase::OnPreviewFeatureLevelChanged(ERHIFeatureLevel::Type)
+void UPCGPlatformSwitchSettingsBase::OnPreviewPlatformChanged(EShaderPlatform /*NewShaderPlatform*/)
 {
 	// We don't broadcast structural from here. The PCGSubsystem will refresh all runtime gen components when preview platform changes.
 	OnSettingsChangedDelegate.Broadcast(this, EPCGChangeType::Cosmetic);
diff --git a/Engine/Plugins/PCG/Source/PCG/Private/Subsystems/PCGSubsystem.cpp b/Engine/Plugins/PCG/Source/PCG/Private/Subsystems/PCGSubsystem.cpp
--- a/Engine/Plugins/PCG/Source/PCG/Private/Subsystems/PCGSubsystem.cpp
+++ b/Engine/Plugins/PCG/Source/PCG/Private/Subsystems/PCGSubsystem.cpp
@@ -226,8 +226,7 @@ void UPCGSubsystem::Deinitialize()
 
 	if (GEditor)
 	{
-		GEditor->OnPreviewPlatformChanged().RemoveAll(this);
-		GEditor->OnPreviewFeatureLevelChanged().RemoveAll(this);
+		GEditor->OnPreviewShaderPlatformChanged().RemoveAll(this);
 	}
 
 	if (UWorld* World = GetWorld(); World && !World->IsGameWorld())
@@ -275,8 +274,7 @@ void UPCGSubsystem::PostInitialize()
 
 	if (GEditor)
 	{
-		GEditor->OnPreviewPlatformChanged().AddUObject(this, &UPCGSubsystem::OnPreviewPlatformChanged);
-		GEditor->OnPreviewFeatureLevelChanged().AddUObject(this, &UPCGSubsystem::OnPreviewFeatureLevelChanged);
+		GEditor->OnPreviewShaderPlatformChanged().AddUObject(this, &UPCGSubsystem::OnPreviewPlatformChanged);
 	}
 
 	if (UWorld* World = GetWorld(); World && !World->IsGameWorld())
@@ -1977,14 +1975,7 @@ void UPCGSubsystem::CreateMissingPartitionActors()
 	}
 }
 
-void UPCGSubsystem::OnPreviewPlatformChanged()
-{
-	FlushCache();
-
-	RefreshAllRuntimeGenExecutionSources(EPCGChangeType::GenerationGrid);
-}
-
-void UPCGSubsystem::OnPreviewFeatureLevelChanged(ERHIFeatureLevel::Type)
+void UPCGSubsystem::OnPreviewPlatformChanged(EShaderPlatform /*NewShaderPlatform*/)
 {
 	FlushCache();
 
diff --git a/Engine/Plugins/PCG/Source/PCG/Public/Elements/ControlFlow/PCGPlatformSwitchBase.h b/Engine/Plugins/PCG/Source/PCG/Public/Elements/ControlFlow/PCGPlatformSwitchBase.h
--- a/Engine/Plugins/PCG/Source/PCG/Public/Elements/ControlFlow/PCGPlatformSwitchBase.h
+++ b/Engine/Plugins/PCG/Source/PCG/Public/Elements/ControlFlow/PCGPlatformSwitchBase.h
@@ -3,11 +3,10 @@
 #pragma once
 
 #include "PCGControlFlow.h"
+#include "RHIShaderPlatform.h"
 
 #include "PCGPlatformSwitchBase.generated.h"
 
-namespace ERHIFeatureLevel { enum Type : int; }
-
 UCLASS(MinimalAPI, BlueprintType, Abstract, ClassGroup = (Procedural))
 class UPCGPlatformSwitchSettingsBase : public UPCGControlFlowSettings
 {
@@ -38,8 +37,7 @@ protected:
 #if WITH_EDITOR
 	void SetupPreviewPlatformChangeEvents();
 	void TearDownPreviewPlatformChangeEvents();
-	void OnPreviewPlatformChanged();
-	void OnPreviewFeatureLevelChanged(ERHIFeatureLevel::Type);
+	void OnPreviewPlatformChanged(EShaderPlatform NewShaderPlatform);
 #endif
 
 private:
diff --git a/Engine/Plugins/PCG/Source/PCG/Public/Subsystems/PCGSubsystem.h b/Engine/Plugins/PCG/Source/PCG/Public/Subsystems/PCGSubsystem.h
--- a/Engine/Plugins/PCG/Source/PCG/Public/Subsystems/PCGSubsystem.h
+++ b/Engine/Plugins/PCG/Source/PCG/Public/Subsystems/PCGSubsystem.h
@@ -3,6 +3,7 @@
 #pragma once
 
 #include "Subsystems/IPCGBaseSubsystem.h"
+#include "RHIShaderPlatform.h"
 #include "Subsystems/WorldSubsystem.h"
 
 #include "PCGCommon.h"
@@ -393,8 +394,7 @@ public:
 
 	PCG_API void CreateMissingPartitionActors();
 
-	void OnPreviewPlatformChanged();
-	void OnPreviewFeatureLevelChanged(ERHIFeatureLevel::Type);
+	void OnPreviewPlatformChanged(EShaderPlatform NewShaderPlatform);
 
 private:
 	void OnPCGGraphCancelled(IPCGGraphExecutionSource* InExecutionSource);
diff --git a/Engine/Plugins/PCG/Source/PCGEditor/Private/Widgets/AssetEditorViewport/SPCGEditorViewport.cpp b/Engine/Plugins/PCG/Source/PCGEditor/Private/Widgets/AssetEditorViewport/SPCGEditorViewport.cpp
--- a/Engine/Plugins/PCG/Source/PCGEditor/Private/Widgets/AssetEditorViewport/SPCGEditorViewport.cpp
+++ b/Engine/Plugins/PCG/Source/PCGEditor/Private/Widgets/AssetEditorViewport/SPCGEditorViewport.cpp
@@ -13,6 +13,7 @@
 #include "PreviewProfileController.h"
 #include "ToolMenus.h"
 #include "Engine/World.h"
+#include "SceneInterface.h"
 #include "ViewportToolbar/UnrealEdViewportToolbar.h"
 #include "ViewportToolbar/UnrealEdViewportToolbarContext.h"
 
@@ -46,9 +47,9 @@ FPCGEditorViewportClient::FPCGEditorViewportClient(FEditorModeTools* InModeTools
 {
 	check(PreviewScene);
 
-	if (UWorld* World = PreviewScene->GetWorld())
+	if (UWorld* World = PreviewScene->GetWorld(); World != nullptr && GWorld != nullptr && GWorld->Scene != nullptr)
 	{
-		World->ChangeFeatureLevel(GWorld->GetFeatureLevel());
+		World->ShaderPlatformChanged(GWorld->Scene->GetShaderPlatform());
 	}
 
 	bUsesDrawHelper = true;
diff --git a/Engine/Plugins/Runtime/HairStrands/Source/HairStrandsEditor/Private/HairStrandsEditor.cpp b/Engine/Plugins/Runtime/HairStrands/Source/HairStrandsEditor/Private/HairStrandsEditor.cpp
--- a/Engine/Plugins/Runtime/HairStrands/Source/HairStrandsEditor/Private/HairStrandsEditor.cpp
+++ b/Engine/Plugins/Runtime/HairStrands/Source/HairStrandsEditor/Private/HairStrandsEditor.cpp
@@ -1,6 +1,7 @@
 // Copyright Epic Games, Inc. All Rights Reserved.
 
 #include "HairStrandsEditor.h"
+#include "DataDrivenShaderPlatformInfo.h"
 #include "HairStrandsCore.h"
 
 #include "GroomActions.h"
@@ -230,8 +231,7 @@ void FGroomEditor::OnPostEngineInit()
 	if (GEditor)
 	{
 		UEditorEngine* EditorEngine = CastChecked<UEditorEngine>(GEngine);
-		PreviewPlatformChangedHandle = EditorEngine->OnPreviewPlatformChanged().AddRaw(this, &FGroomEditor::OnPreviewPlatformChanged);
-		PreviewFeatureLevelChangedHandle = EditorEngine->OnPreviewFeatureLevelChanged().AddRaw(this, &FGroomEditor::OnPreviewFeatureLevelChanged);
+		PreviewPlatformChangedHandle = EditorEngine->OnPreviewShaderPlatformChanged().AddRaw(this, &FGroomEditor::OnPreviewPlatformChanged);
 	}
 }
 
@@ -239,46 +239,21 @@ void FGroomEditor::OnPreExit()
 {
 	if (GEditor)
 	{
-		CastChecked<UEditorEngine>(GEngine)->OnPreviewPlatformChanged().Remove(PreviewPlatformChangedHandle);
-		CastChecked<UEditorEngine>(GEngine)->OnPreviewFeatureLevelChanged().Remove(PreviewFeatureLevelChangedHandle);
+		CastChecked<UEditorEngine>(GEngine)->OnPreviewShaderPlatformChanged().Remove(PreviewPlatformChangedHandle);
 	}
 }
 
-void FGroomEditor::OnPreviewPlatformChanged()
+void FGroomEditor::OnPreviewPlatformChanged(EShaderPlatform NewShaderPlatform)
 {
 #if WITH_EDITOR
-	UEditorEngine* EditorEngine = CastChecked<UEditorEngine>(GEngine);
-	ERHIFeatureLevel::Type ActiveFeatureLevel = EditorEngine->GetDefaultWorldFeatureLevel();
-	if (EditorEngine->IsFeatureLevelPreviewActive())
-	{
-		ActiveFeatureLevel = EditorEngine->GetActiveFeatureLevelPreviewType();
-	}
+	ERHIFeatureLevel::Type ActiveFeatureLevel = GetMaxSupportedFeatureLevel(NewShaderPlatform);
 
 	for (TObjectIterator<UGroomComponent> It; It; ++It)
 	{
 		if (UGroomComponent* Component = *It)
 		{
 			Component->HandlePlatformPreviewChanged(ActiveFeatureLevel);
-		}
-	}
-#endif
-}
-
-void FGroomEditor::OnPreviewFeatureLevelChanged(ERHIFeatureLevel::Type InPreviewFeatureLevel)
-{
-#if WITH_EDITOR
-	UEditorEngine* EditorEngine = CastChecked<UEditorEngine>(GEngine);
-	ERHIFeatureLevel::Type ActiveFeatureLevel = EditorEngine->GetDefaultWorldFeatureLevel();
-	if (EditorEngine->IsFeatureLevelPreviewActive())
-	{
-		ActiveFeatureLevel = EditorEngine->GetActiveFeatureLevelPreviewType();
-	}
-
-	for (TObjectIterator<UGroomComponent> It; It; ++It)
-	{
-		if (UGroomComponent* Component = *It)
-		{
-			Component->HandleFeatureLevelChanged(InPreviewFeatureLevel);
+			Component->HandleFeatureLevelChanged(ActiveFeatureLevel);
 		}
 	}
 #endif
diff --git a/Engine/Plugins/Runtime/HairStrands/Source/HairStrandsEditor/Public/HairStrandsEditor.h b/Engine/Plugins/Runtime/HairStrands/Source/HairStrandsEditor/Public/HairStrandsEditor.h
--- a/Engine/Plugins/Runtime/HairStrands/Source/HairStrandsEditor/Public/HairStrandsEditor.h
+++ b/Engine/Plugins/Runtime/HairStrands/Source/HairStrandsEditor/Public/HairStrandsEditor.h
@@ -8,6 +8,7 @@
 #include "Styling/SlateStyle.h"
 #include "HAL/LowLevelMemTracker.h"
 #include "RHIFwd.h"
+#include "RHIShaderPlatform.h"
 
 #define UE_API HAIRSTRANDSEDITOR_API
 
@@ -29,9 +30,8 @@ public:
 	virtual bool SupportsDynamicReloading() override { return false; }
 
 	UE_API void OnPostEngineInit();
-	UE_API void OnPreviewPlatformChanged();
+	UE_API void OnPreviewPlatformChanged(EShaderPlatform NewShaderPlatform);
 	UE_API void OnPreExit();
-	UE_API void OnPreviewFeatureLevelChanged(ERHIFeatureLevel::Type InPreviewFeatureLevel);
 
 	static inline FGroomEditor& Get()
 	{
diff --git a/Engine/Plugins/Runtime/nDisplay/Source/DisplayClusterConfigurator/Private/Views/Viewport/DisplayClusterConfiguratorSCSEditorViewport.cpp b/Engine/Plugins/Runtime/nDisplay/Source/DisplayClusterConfigurator/Private/Views/Viewport/DisplayClusterConfiguratorSCSEditorViewport.cpp
--- a/Engine/Plugins/Runtime/nDisplay/Source/DisplayClusterConfigurator/Private/Views/Viewport/DisplayClusterConfiguratorSCSEditorViewport.cpp
+++ b/Engine/Plugins/Runtime/nDisplay/Source/DisplayClusterConfigurator/Private/Views/Viewport/DisplayClusterConfiguratorSCSEditorViewport.cpp
@@ -17,6 +17,7 @@
 #include "Camera/CameraComponent.h"
 
 #include "EngineUtils.h"
+#include "SceneInterface.h"
 #include "BlueprintEditor.h"
 #include "SSubobjectEditor.h"
 #include "EditorViewportCommands.h"
@@ -485,22 +486,22 @@ void SDisplayClusterConfiguratorSCSEditorViewport::Construct(const FArguments& I
 	if (ViewportClient.IsValid())
 	{
 		UWorld* World = ViewportClient->GetPreviewScene()->GetWorld();
-		if (World != nullptr)
+		if (World != nullptr && GWorld != nullptr && GWorld->Scene != nullptr)
 		{
-			World->ChangeFeatureLevel(GWorld->GetFeatureLevel());
+			World->ShaderPlatformChanged(GWorld->Scene->GetShaderPlatform());
 		}
 	}
 
-	// Use a delegate to inform the attached world of feature level changes.
+	// Use a delegate to inform the attached world of preview platform changes.
 	UEditorEngine* Editor = (UEditorEngine*)GEngine;
-	PreviewFeatureLevelChangedHandle = Editor->OnPreviewFeatureLevelChanged().AddLambda([this](ERHIFeatureLevel::Type NewFeatureLevel)
+	PreviewFeatureLevelChangedHandle = Editor->OnPreviewShaderPlatformChanged().AddLambda([this](EShaderPlatform NewShaderPlatform)
 	{
 		if (ViewportClient.IsValid())
 		{
 			UWorld* World = ViewportClient->GetPreviewScene()->GetWorld();
 			if (World != nullptr)
 			{
-				World->ChangeFeatureLevel(NewFeatureLevel);
+				World->ShaderPlatformChanged(NewShaderPlatform);
 
 				// Refresh the preview scene. Don't change the camera.
 				RequestRefresh(false);
@@ -515,7 +516,7 @@ void SDisplayClusterConfiguratorSCSEditorViewport::Construct(const FArguments& I
 SDisplayClusterConfiguratorSCSEditorViewport::~SDisplayClusterConfiguratorSCSEditorViewport()
 {
 	UEditorEngine* Editor = (UEditorEngine*)GEngine;
-	Editor->OnPreviewFeatureLevelChanged().Remove(PreviewFeatureLevelChangedHandle);
+	Editor->OnPreviewShaderPlatformChanged().Remove(PreviewFeatureLevelChangedHandle);
 
 	if (ViewportClient.IsValid())
 	{
diff --git a/Engine/Plugins/TextureGraph/Source/TextureGraphEditor/Private/STG_EditorViewport.cpp b/Engine/Plugins/TextureGraph/Source/TextureGraphEditor/Private/STG_EditorViewport.cpp
--- a/Engine/Plugins/TextureGraph/Source/TextureGraphEditor/Private/STG_EditorViewport.cpp
+++ b/Engine/Plugins/TextureGraph/Source/TextureGraphEditor/Private/STG_EditorViewport.cpp
@@ -16,6 +16,7 @@
 #include "Engine/SkeletalMesh.h"
 #include "Engine/StaticMesh.h"
 #include "Engine/TextureCube.h"
+#include "SceneInterface.h"
 #include "Framework/Layout/ScrollyZoomy.h"
 #include "ISettingsModule.h"
 #include "ITG_Editor.h"
@@ -236,15 +237,15 @@ void STG_EditorViewport::Construct(const FArguments& InArgs)
 
 	// restore last used feature level
 	UWorld* World = PreviewScene->GetWorld();
-	if (World != nullptr)
+	if (World != nullptr && GWorld != nullptr && GWorld->Scene != nullptr)
 	{
-		World->ChangeFeatureLevel(GWorld->GetFeatureLevel());
+		World->ShaderPlatformChanged(GWorld->Scene->GetShaderPlatform());
 	}
 
 	UEditorEngine* Editor = CastChecked<UEditorEngine>(GEngine);
-	PreviewFeatureLevelChangedHandle = Editor->OnPreviewFeatureLevelChanged().AddLambda([this](ERHIFeatureLevel::Type NewFeatureLevel)
+	PreviewFeatureLevelChangedHandle = Editor->OnPreviewShaderPlatformChanged().AddLambda([this](EShaderPlatform NewShaderPlatform)
 		{
-			PreviewScene->GetWorld()->ChangeFeatureLevel(NewFeatureLevel);
+			PreviewScene->GetWorld()->ShaderPlatformChanged(NewShaderPlatform);
 		});
 
 	CurrentViewMode = VMI_Lit;
@@ -287,7 +288,7 @@ STG_EditorViewport::~STG_EditorViewport()
 	{
 		if (UEditorEngine* EditorEngine = Cast<UEditorEngine>(GEngine))
 		{
-			EditorEngine->OnPreviewFeatureLevelChanged().Remove(PreviewFeatureLevelChangedHandle);
+			EditorEngine->OnPreviewShaderPlatformChanged().Remove(PreviewFeatureLevelChangedHandle);
 		}
 	}
 	FCoreUObjectDelegates::OnObjectPropertyChanged.RemoveAll(this);
diff --git a/Engine/Plugins/VirtualProduction/DynamicMaterial/Source/DynamicMaterialEditor/Private/UI/Widgets/Editor/SDMMaterialPreview.cpp b/Engine/Plugins/VirtualProduction/DynamicMaterial/Source/DynamicMaterialEditor/Private/UI/Widgets/Editor/SDMMaterialPreview.cpp
--- a/Engine/Plugins/VirtualProduction/DynamicMaterial/Source/DynamicMaterialEditor/Private/UI/Widgets/Editor/SDMMaterialPreview.cpp
+++ b/Engine/Plugins/VirtualProduction/DynamicMaterial/Source/DynamicMaterialEditor/Private/UI/Widgets/Editor/SDMMaterialPreview.cpp
@@ -14,6 +14,7 @@
 #include "Engine/Engine.h"
 #include "Engine/PostProcessVolume.h"
 #include "Engine/World.h"
+#include "SceneInterface.h"
 #include "Framework/MultiBox/MultiBoxBuilder.h"
 #include "Material/DynamicMaterialInstance.h"
 #include "MaterialEditor/MaterialEditorMeshComponent.h"
@@ -59,7 +60,7 @@ SDMMaterialPreview::~SDMMaterialPreview()
 
 	if (UEditorEngine* Editor = CastChecked<UEditorEngine>(GEngine))
 	{
-		Editor->OnPreviewFeatureLevelChanged().RemoveAll(this);
+		Editor->OnPreviewShaderPlatformChanged().RemoveAll(this);
 	}
 
 	if (UDynamicMaterialEditorSettings* EditorSettings = UDynamicMaterialEditorSettings::Get())
@@ -91,9 +92,9 @@ void SDMMaterialPreview::Construct(const FArguments& InArgs, const TSharedRef<SD
 	// restore last used feature level
 	UWorld* PreviewWorld = PreviewScene->GetWorld();
 
-	if (PreviewWorld != nullptr)
+	if (PreviewWorld != nullptr && GWorld != nullptr && GWorld->Scene != nullptr)
 	{
-		PreviewWorld->ChangeFeatureLevel(GWorld->GetFeatureLevel());
+		PreviewWorld->ShaderPlatformChanged(GWorld->Scene->GetShaderPlatform());
 	}
 
 	UDynamicMaterialEditorSettings* EditorSettings = UDynamicMaterialEditorSettings::Get();
@@ -110,7 +111,7 @@ void SDMMaterialPreview::Construct(const FArguments& InArgs, const TSharedRef<SD
 
 	if (UEditorEngine* Editor = Cast<UEditorEngine>(GEngine))
 	{
-		Editor->OnPreviewFeatureLevelChanged().AddSP(this, &SDMMaterialPreview::OnFeatureLevelChanged);
+		Editor->OnPreviewShaderPlatformChanged().AddSP(this, &SDMMaterialPreview::OnFeatureLevelChanged);
 	}
 
 	if (UDynamicMaterialInstance* MaterialInstance = InMaterialModelBase->GetDynamicMaterialInstance())
@@ -434,13 +435,13 @@ void SDMMaterialPreview::OnPropertyChanged(UObject* InObjectBeingModified, FProp
 {
 }
 
-void SDMMaterialPreview::OnFeatureLevelChanged(ERHIFeatureLevel::Type InNewFeatureLevel)
+void SDMMaterialPreview::OnFeatureLevelChanged(EShaderPlatform NewShaderPlatform)
 {
 	if (UWorld* World = PreviewScene->GetWorld())
 	{
-		World->ChangeFeatureLevel(InNewFeatureLevel);
+		World->ShaderPlatformChanged(NewShaderPlatform);
 	}
-	
+
 }
 
 TSharedRef<SWidget> SDMMaterialPreview::GenerateToolbarMenu()
diff --git a/Engine/Plugins/VirtualProduction/DynamicMaterial/Source/DynamicMaterialEditor/Private/UI/Widgets/Editor/SDMMaterialPreview.h b/Engine/Plugins/VirtualProduction/DynamicMaterial/Source/DynamicMaterialEditor/Private/UI/Widgets/Editor/SDMMaterialPreview.h
--- a/Engine/Plugins/VirtualProduction/DynamicMaterial/Source/DynamicMaterialEditor/Private/UI/Widgets/Editor/SDMMaterialPreview.h
+++ b/Engine/Plugins/VirtualProduction/DynamicMaterial/Source/DynamicMaterialEditor/Private/UI/Widgets/Editor/SDMMaterialPreview.h
@@ -3,6 +3,7 @@
 #pragma once
 
 #include "SEditorViewport.h"
+#include "RHIShaderPlatform.h"
 #include "UObject/GCObject.h"
 
 #include "Delegates/IDelegateInstance.h"
@@ -109,7 +110,7 @@ protected:
 
 	void OnPropertyChanged(UObject* InObjectBeingModified, FPropertyChangedEvent& InPropertyChangedEvent);
 
-	void OnFeatureLevelChanged(ERHIFeatureLevel::Type InNewFeatureLevel);
+	void OnFeatureLevelChanged(EShaderPlatform NewShaderPlatform);
 
 	TSharedRef<SWidget> GenerateToolbarMenu();
 
diff --git a/Engine/Source/Editor/AdvancedPreviewScene/Private/AdvancedPreviewScene.cpp b/Engine/Source/Editor/AdvancedPreviewScene/Private/AdvancedPreviewScene.cpp
--- a/Engine/Source/Editor/AdvancedPreviewScene/Private/AdvancedPreviewScene.cpp
+++ b/Engine/Source/Editor/AdvancedPreviewScene/Private/AdvancedPreviewScene.cpp
@@ -2,6 +2,7 @@
 
 #include "AdvancedPreviewScene.h"
 #include "UnrealClient.h"
+#include "SceneInterface.h"
 
 #include "Components/PostProcessComponent.h"
 #include "Engine/Texture.h"
@@ -102,9 +103,9 @@ FAdvancedPreviewScene::FAdvancedPreviewScene(ConstructionValues CVS, float InFlo
 	BindCommands();
 
 	// since advance preview scenes are used in conjunction with material/mesh editors we should match their feature level with the global one
-	if (GIsEditor && GEditor != nullptr)
+	if (GIsEditor && GEditor != nullptr && GWorld != nullptr && GWorld->Scene != nullptr)
 	{
-		PreviewWorld->ChangeFeatureLevel(GEditor->DefaultWorldFeatureLevel);
+		PreviewWorld->ShaderPlatformChanged(GWorld->Scene->GetShaderPlatform());
 	}
 
 	PreviousRotation = Profile.LightingRigRotation;
diff --git a/Engine/Source/Editor/Kismet/Private/SSCSEditorViewport.cpp b/Engine/Source/Editor/Kismet/Private/SSCSEditorViewport.cpp
--- a/Engine/Source/Editor/Kismet/Private/SSCSEditorViewport.cpp
+++ b/Engine/Source/Editor/Kismet/Private/SSCSEditorViewport.cpp
@@ -15,6 +15,7 @@
 #include "Engine/Engine.h"
 #include "Engine/EngineBaseTypes.h"
 #include "Engine/World.h"
+#include "SceneInterface.h"
 #include "Framework/Application/SlateApplication.h"
 #include "Framework/Commands/UIAction.h"
 #include "Framework/Commands/UICommandList.h"
@@ -108,22 +109,22 @@ void SSCSEditorViewport::Construct(const FArguments& InArgs)
 	if (ViewportClient.IsValid())
 	{
 		UWorld* World = ViewportClient->GetPreviewScene()->GetWorld();
-		if (World != nullptr)
+		if (World != nullptr && GWorld != nullptr && GWorld->Scene != nullptr)
 		{
-			World->ChangeFeatureLevel(GWorld->GetFeatureLevel());
+			World->ShaderPlatformChanged(GWorld->Scene->GetShaderPlatform());
 		}
 	}
 
-	// Use a delegate to inform the attached world of feature level changes.
+	// Use a delegate to inform the attached world of preview platform changes.
 	UEditorEngine* Editor = (UEditorEngine*)GEngine;
-	PreviewFeatureLevelChangedHandle = Editor->OnPreviewFeatureLevelChanged().AddLambda([this](ERHIFeatureLevel::Type NewFeatureLevel)
+	PreviewFeatureLevelChangedHandle = Editor->OnPreviewShaderPlatformChanged().AddLambda([this](EShaderPlatform NewShaderPlatform)
 		{
 			if(ViewportClient.IsValid())
 			{
 				UWorld* World = ViewportClient->GetPreviewScene()->GetWorld();
 				if (World != nullptr)
 				{
-					World->ChangeFeatureLevel(NewFeatureLevel);
+					World->ShaderPlatformChanged(NewShaderPlatform);
 
 					// Refresh the preview scene. Don't change the camera.
 					RequestRefresh(false);
@@ -138,7 +139,7 @@ void SSCSEditorViewport::Construct(const FArguments& InArgs)
 SSCSEditorViewport::~SSCSEditorViewport()
 {
 	UEditorEngine* Editor = (UEditorEngine*)GEngine;
-	Editor->OnPreviewFeatureLevelChanged().Remove(PreviewFeatureLevelChangedHandle);
+	Editor->OnPreviewShaderPlatformChanged().Remove(PreviewFeatureLevelChangedHandle);
 
 	if(ViewportClient.IsValid())
 	{
diff --git a/Engine/Source/Editor/LevelEditor/Private/SLevelEditor.cpp b/Engine/Source/Editor/LevelEditor/Private/SLevelEditor.cpp
--- a/Engine/Source/Editor/LevelEditor/Private/SLevelEditor.cpp
+++ b/Engine/Source/Editor/LevelEditor/Private/SLevelEditor.cpp
@@ -1,6 +1,7 @@
 // Copyright Epic Games, Inc. All Rights Reserved.
 
 #include "SLevelEditor.h"
+#include "DataDrivenShaderPlatformInfo.h"
 #include "Components/StaticMeshComponent.h"
 #include "Framework/MultiBox/MultiBoxExtender.h"
 #include "Materials/MaterialInterface.h"
@@ -259,27 +260,22 @@ void SLevelEditor::Construct( const SLevelEditor::FArguments& InArgs)
 		GEditor->GetEditorWorldContext(true).AddRef(World);
 
 		// Set the initial preview feature level.
-		World->ChangeFeatureLevel(GEditor->GetActiveFeatureLevelPreviewType());
+		World->ShaderPlatformChanged(GEditor->GetActiveShaderPlatform());
 
 		LevelActorDeletedHandle = GEditor->OnLevelActorDeleted().AddSP(this, &SLevelEditor::OnLevelActorDeleted);
 		LevelActorOuterChangedHandle = GEditor->OnLevelActorOuterChanged().AddSP(this, &SLevelEditor::OnLevelActorOuterChanged);
 	}
 
-	// Patch into the OnPreviewFeatureLevelChanged() delegate to swap out the current feature level with a user selection.
-	PreviewFeatureLevelChangedHandle = GEditor->OnPreviewFeatureLevelChanged().AddLambda([this](ERHIFeatureLevel::Type NewFeatureLevel)
+	PreviewPlatformChangedHandle = GEditor->OnPreviewShaderPlatformChanged().AddLambda([this](EShaderPlatform NewShaderPlatform)
 		{
 			// Do one recapture if atleast one ReflectionComponent is dirty
 			// BuildReflectionCapturesOnly_Execute in LevelEditorActions relies on this happening on toggle between SM5->ES31. If you remove this, update that code!
+			ERHIFeatureLevel::Type NewFeatureLevel = GetMaxSupportedFeatureLevel(NewShaderPlatform);
 			if (World->NumUnbuiltReflectionCaptures >= 1 && NewFeatureLevel == ERHIFeatureLevel::ES3_1 && GEditor != nullptr)
 			{
 				GEditor->BuildReflectionCaptures();
 			}
-			World->ChangeFeatureLevel(NewFeatureLevel, true, true);
-		});
-
-	PreviewPlatformChangedHandle = GEditor->OnPreviewPlatformChanged().AddLambda([this]()
-		{
-			World->ShaderPlatformChanged();
+			World->ShaderPlatformChanged(NewShaderPlatform, true);
 		});
 
 	FEditorDelegates::MapChange.AddRaw(this, &SLevelEditor::HandleEditorMapChange);
@@ -427,8 +423,7 @@ SLevelEditor::~SLevelEditor()
 
 	if (GEngine)
 	{
-		CastChecked<UEditorEngine>(GEngine)->OnPreviewFeatureLevelChanged().Remove(PreviewFeatureLevelChangedHandle);
-		CastChecked<UEditorEngine>(GEngine)->OnPreviewPlatformChanged().Remove(PreviewPlatformChangedHandle);
+		CastChecked<UEditorEngine>(GEngine)->OnPreviewShaderPlatformChanged().Remove(PreviewPlatformChangedHandle);
 	}
 
 	if (GEditor)
diff --git a/Engine/Source/Editor/LevelEditor/Private/SLevelViewport.cpp b/Engine/Source/Editor/LevelEditor/Private/SLevelViewport.cpp
--- a/Engine/Source/Editor/LevelEditor/Private/SLevelViewport.cpp
+++ b/Engine/Source/Editor/LevelEditor/Private/SLevelViewport.cpp
@@ -1446,7 +1446,7 @@ void SLevelViewport::OnMapChanged( UWorld* World, EMapChangeType MapChangeType )
 		}
 		World->EditorViews[LevelViewportClient->ViewportType].CamUpdated = false;
 
-		World->ChangeFeatureLevel(GEditor->GetActiveFeatureLevelPreviewType());
+		World->ShaderPlatformChanged(GEditor->GetActiveShaderPlatform());
 
 		RedrawViewport(true);
 	}
@@ -5029,12 +5029,12 @@ void SLevelViewport::StartPlayInEditorSession(UGameViewportClient* InPlayClient,
 
 	GEngine->BroadcastLevelActorListChanged();
 
-	// register for preview feature level change
+	// register for preview platform change
 	UEditorEngine* Editor = CastChecked<UEditorEngine>(GEngine);
-	
-	PIEPreviewFeatureLevelChangedHandle = Editor->OnPreviewFeatureLevelChanged().AddLambda([InPlayClient](ERHIFeatureLevel::Type NewFeatureLevel)
+
+	PIEPreviewFeatureLevelChangedHandle = Editor->OnPreviewShaderPlatformChanged().AddLambda([InPlayClient](EShaderPlatform NewShaderPlatform)
 		{
-			InPlayClient->GetWorld()->ChangeFeatureLevel(NewFeatureLevel);
+			InPlayClient->GetWorld()->ShaderPlatformChanged(NewShaderPlatform, true);
 		});
 
 }
@@ -5263,7 +5263,7 @@ void SLevelViewport::EndPlayInEditorSession()
 	// Remove preview feature level delegate if set
 	if (PIEPreviewFeatureLevelChangedHandle.IsValid())
 	{
-		CastChecked<UEditorEngine>(GEngine)->OnPreviewFeatureLevelChanged().Remove(PIEPreviewFeatureLevelChangedHandle);
+		CastChecked<UEditorEngine>(GEngine)->OnPreviewShaderPlatformChanged().Remove(PIEPreviewFeatureLevelChangedHandle);
 		PIEPreviewFeatureLevelChangedHandle.Reset();
 	}
 }
diff --git a/Engine/Source/Editor/MaterialEditor/Private/SMaterialEditorViewport.cpp b/Engine/Source/Editor/MaterialEditor/Private/SMaterialEditorViewport.cpp
--- a/Engine/Source/Editor/MaterialEditor/Private/SMaterialEditorViewport.cpp
+++ b/Engine/Source/Editor/MaterialEditor/Private/SMaterialEditorViewport.cpp
@@ -20,6 +20,7 @@
 #include "Materials/MaterialExpressionUserSceneTexture.h"
 #include "Materials/MaterialInstanceConstant.h"
 #include "Engine/StaticMesh.h"
+#include "SceneInterface.h"
 #include "Engine/Selection.h"
 #include "Editor.h"
 #include "Dialogs/Dialogs.h"
@@ -260,15 +261,15 @@ void SMaterialEditor3DPreviewViewport::Construct(const FArguments& InArgs)
 
 	// restore last used feature level
 	UWorld* PreviewWorld = AdvancedPreviewScene->GetWorld();
-	if (PreviewWorld != nullptr)
+	if (PreviewWorld != nullptr && GWorld != nullptr && GWorld->Scene != nullptr)
 	{
-		PreviewWorld->ChangeFeatureLevel(GWorld->GetFeatureLevel());
+		PreviewWorld->ShaderPlatformChanged(GWorld->Scene->GetShaderPlatform());
 	}	
 
 	UEditorEngine* Editor = CastChecked<UEditorEngine>(GEngine);
-	PreviewFeatureLevelChangedHandle = Editor->OnPreviewFeatureLevelChanged().AddLambda([this](ERHIFeatureLevel::Type NewFeatureLevel)
+	PreviewFeatureLevelChangedHandle = Editor->OnPreviewShaderPlatformChanged().AddLambda([this](EShaderPlatform NewShaderPlatform)
 		{
-			AdvancedPreviewScene->GetWorld()->ChangeFeatureLevel(NewFeatureLevel);
+			AdvancedPreviewScene->GetWorld()->ShaderPlatformChanged(NewShaderPlatform);
 		});
 
 
@@ -300,7 +301,7 @@ void SMaterialEditor3DPreviewViewport::Construct(const FArguments& InArgs)
 
 SMaterialEditor3DPreviewViewport::~SMaterialEditor3DPreviewViewport()
 {
-	CastChecked<UEditorEngine>(GEngine)->OnPreviewFeatureLevelChanged().Remove(PreviewFeatureLevelChangedHandle);
+	CastChecked<UEditorEngine>(GEngine)->OnPreviewShaderPlatformChanged().Remove(PreviewFeatureLevelChangedHandle);
 	
 	if (PreviewMeshComponent != nullptr)
 	{
diff --git a/Engine/Source/Editor/Persona/Private/SAnimationEditorViewport.cpp b/Engine/Source/Editor/Persona/Private/SAnimationEditorViewport.cpp
--- a/Engine/Source/Editor/Persona/Private/SAnimationEditorViewport.cpp
+++ b/Engine/Source/Editor/Persona/Private/SAnimationEditorViewport.cpp
@@ -14,6 +14,7 @@
 #include "Animation/AnimMontage.h"
 #include "Preferences/PersonaOptions.h"
 #include "Rendering/SkeletalMeshRenderData.h"
+#include "SceneInterface.h"
 
 #include "SAnimationScrubPanel.h"
 #include "SAnimMontageScrubPanel.h"
@@ -85,22 +86,22 @@ void SAnimationEditorViewport::Construct(const FArguments& InArgs, const FAnimat
 	if (ScenePtr.IsValid())
 	{
 		UWorld* World = ScenePtr->GetWorld();
-		if (World != nullptr)
+		if (World != nullptr && GWorld != nullptr && GWorld->Scene != nullptr)
 		{
-			World->ChangeFeatureLevel(GWorld->GetFeatureLevel());
+			World->ShaderPlatformChanged(GWorld->Scene->GetShaderPlatform());
 		}
 	}
 
 	UEditorEngine* Editor = (UEditorEngine*)GEngine;
-	PreviewFeatureLevelChangedHandle = Editor->OnPreviewFeatureLevelChanged().AddLambda(
-		[PreviewScenePtrWeak = PreviewScenePtr](ERHIFeatureLevel::Type NewFeatureLevel)
+	PreviewFeatureLevelChangedHandle = Editor->OnPreviewShaderPlatformChanged().AddLambda(
+		[PreviewScenePtrWeak = PreviewScenePtr](EShaderPlatform NewShaderPlatform)
 		{
 			if (TSharedPtr<IPersonaPreviewScene> ScenePtr = PreviewScenePtrWeak.Pin())
 			{
 				UWorld* World = ScenePtr->GetWorld();
 				if (World != nullptr)
 				{
-					World->ChangeFeatureLevel(NewFeatureLevel);
+					World->ShaderPlatformChanged(NewShaderPlatform);
 				}
 			}
 		});
@@ -109,7 +110,7 @@ void SAnimationEditorViewport::Construct(const FArguments& InArgs, const FAnimat
 SAnimationEditorViewport::~SAnimationEditorViewport()
 {
 	UEditorEngine* Editor = (UEditorEngine*)GEngine;
-	Editor->OnPreviewFeatureLevelChanged().Remove(PreviewFeatureLevelChangedHandle);
+	Editor->OnPreviewShaderPlatformChanged().Remove(PreviewFeatureLevelChangedHandle);
 }
 
 void SAnimationEditorViewport::PopulateViewportOverlays(TSharedRef<SOverlay> Overlay)
diff --git a/Engine/Source/Editor/StaticMeshEditor/Private/SStaticMeshEditorViewport.cpp b/Engine/Source/Editor/StaticMeshEditor/Private/SStaticMeshEditorViewport.cpp
--- a/Engine/Source/Editor/StaticMeshEditor/Private/SStaticMeshEditorViewport.cpp
+++ b/Engine/Source/Editor/StaticMeshEditor/Private/SStaticMeshEditorViewport.cpp
@@ -10,6 +10,7 @@
 #include "Components/StaticMeshComponent.h"
 #include "Styling/AppStyle.h"
 #include "Engine/StaticMesh.h"
+#include "SceneInterface.h"
 #include "IStaticMeshEditor.h"
 #include "StaticMeshEditorActions.h"
 #include "Slate/SceneViewport.h"
@@ -67,15 +68,15 @@ void SStaticMeshEditorViewport::Construct(const FArguments& InArgs)
 
 	// restore last used feature level
 	UWorld* World = PreviewScene->GetWorld();
-	if (World != nullptr)
+	if (World != nullptr && GWorld != nullptr && GWorld->Scene != nullptr)
 	{
-		World->ChangeFeatureLevel(GWorld->GetFeatureLevel());
+		World->ShaderPlatformChanged(GWorld->Scene->GetShaderPlatform());
 	}
 
 	UEditorEngine* Editor = CastChecked<UEditorEngine>(GEngine);
-	PreviewFeatureLevelChangedHandle = Editor->OnPreviewFeatureLevelChanged().AddLambda([this](ERHIFeatureLevel::Type NewFeatureLevel)
+	PreviewFeatureLevelChangedHandle = Editor->OnPreviewShaderPlatformChanged().AddLambda([this](EShaderPlatform NewShaderPlatform)
 		{
-			PreviewScene->GetWorld()->ChangeFeatureLevel(NewFeatureLevel);
+			PreviewScene->GetWorld()->ShaderPlatformChanged(NewShaderPlatform);
 		});
 
 	CurrentViewMode = VMI_Lit;
@@ -146,7 +147,7 @@ SStaticMeshEditorViewport::SStaticMeshEditorViewport()
 
 SStaticMeshEditorViewport::~SStaticMeshEditorViewport()
 {
-	CastChecked<UEditorEngine>(GEngine)->OnPreviewFeatureLevelChanged().Remove(PreviewFeatureLevelChangedHandle);
+	CastChecked<UEditorEngine>(GEngine)->OnPreviewShaderPlatformChanged().Remove(PreviewFeatureLevelChangedHandle);
 
 	FCoreUObjectDelegates::OnObjectPropertyChanged.RemoveAll(this);
 	if (EditorViewportClient.IsValid())
diff --git a/Engine/Source/Editor/UnrealEd/Classes/Editor/EditorEngine.h b/Engine/Source/Editor/UnrealEd/Classes/Editor/EditorEngine.h
--- a/Engine/Source/Editor/UnrealEd/Classes/Editor/EditorEngine.h
+++ b/Engine/Source/Editor/UnrealEd/Classes/Editor/EditorEngine.h
@@ -651,12 +651,22 @@ public:
 	EShaderPlatform CachedEditorShaderPlatform;
 
 	/** A delegate that is called when the preview feature level changes. Primarily used to switch a viewport's feature level. */
+	UE_DEPRECATED(5.8, "FPreviewFeatureLevelChanged is deprecated, use FPreviewShaderPlatformChanged instead")
 	DECLARE_MULTICAST_DELEGATE_OneParam(FPreviewFeatureLevelChanged, ERHIFeatureLevel::Type);
+PRAGMA_DISABLE_DEPRECATION_WARNINGS
 	FPreviewFeatureLevelChanged PreviewFeatureLevelChanged;
+PRAGMA_ENABLE_DEPRECATION_WARNINGS
 
 	/** A delegate that is called when the preview platform changes. */
+	UE_DEPRECATED(5.8, "FPreviewPlatformChanged is deprecated, use FPreviewShaderPlatformChanged instead")
 	DECLARE_MULTICAST_DELEGATE(FPreviewPlatformChanged);
+PRAGMA_DISABLE_DEPRECATION_WARNINGS
 	FPreviewPlatformChanged PreviewPlatformChanged;
+PRAGMA_ENABLE_DEPRECATION_WARNINGS
+
+	/** A delegate that is called when the preview platform changes. */
+	DECLARE_MULTICAST_DELEGATE_OneParam(FPreviewShaderPlatformChanged, EShaderPlatform);
+	FPreviewShaderPlatformChanged PreviewShaderPlatformChanged;
 
 	/** An array of delegates that can force disable throttling cpu usage if any of them return true. */
 	DECLARE_DELEGATE_RetVal(bool, FShouldDisableCPUThrottling);
@@ -3254,11 +3264,17 @@ public:
 	 */
 	UNREALED_API ERHIFeatureLevel::Type GetActiveFeatureLevelPreviewType() const;
 
-	/** Return the delegate that is called when the preview feature level changes */
+PRAGMA_DISABLE_DEPRECATION_WARNINGS
+/** Return the delegate that is called when the preview feature level changes */
+	UE_DEPRECATED(5.8, "OnPreviewFeatureLevelChanged is deprecated, use OnPreviewShaderPlatformChanged instead")
 	FPreviewFeatureLevelChanged& OnPreviewFeatureLevelChanged() { return PreviewFeatureLevelChanged; }
-
 	/** Return the delegate that is called when the preview platform changes */
+	UE_DEPRECATED(5.8, "OnPreviewPlatformChanged is deprecated, use OnPreviewShaderPlatformChanged instead")
 	FPreviewPlatformChanged& OnPreviewPlatformChanged() { return PreviewPlatformChanged; }
+PRAGMA_ENABLE_DEPRECATION_WARNINGS
+
+	/** Return the delegate that is called when the preview shader platform changes */
+	FPreviewShaderPlatformChanged& OnPreviewShaderPlatformChanged() { return PreviewShaderPlatformChanged; }
 
 	/** Return the delegate that is called when the bugitgo command is called */
 	FPostBugItGoCalled& OnPostBugItGoCalled() { return PostBugItGoCalled; }
diff --git a/Engine/Source/Editor/UnrealEd/Private/EditorEngine.cpp b/Engine/Source/Editor/UnrealEd/Private/EditorEngine.cpp
--- a/Engine/Source/Editor/UnrealEd/Private/EditorEngine.cpp
+++ b/Engine/Source/Editor/UnrealEd/Private/EditorEngine.cpp
@@ -6655,6 +6655,14 @@ static void PreviewPlatformChangePost()
 		// after Register, TP->GetTextureLODSettings is equal to the CDP way
 	}
 
+	for (TObjectIterator<UWorld> It; It; ++It)
+	{
+		if (It->Scene)
+		{
+			It->Scene->UpdateEarlyZPassMode();
+		}
+	}
+
 	{
 		// TextureLODSettings and platform cvars may have changed
 		//	need to notify all textures
@@ -6774,7 +6782,10 @@ void UEditorEngine::SetPreviewPlatform(const FPreviewPlatformInfo& NewPreviewPla
 		GShaderCompilingManager->ProcessAsyncResults(false, true);
 
 		DefaultWorldFeatureLevel = EffectiveFeatureLevel;
+		PreviewShaderPlatformChanged.Broadcast(ShaderPlatform);
+PRAGMA_DISABLE_DEPRECATION_WARNINGS
 		PreviewFeatureLevelChanged.Broadcast(EffectiveFeatureLevel);
+PRAGMA_ENABLE_DEPRECATION_WARNINGS
 	}
 
 	Scalability::ChangeScalabilityPreviewPlatform(PreviewPlatform.GetEffectivePreviewPlatformName(), GetActiveShaderPlatform());
@@ -6800,8 +6811,10 @@ void UEditorEngine::SetPreviewPlatform(const FPreviewPlatformInfo& NewPreviewPla
 	Scalability::ApplyCachedQualityLevels();
 
 	PreviewPlatformChangePost();
-	
+
+PRAGMA_DISABLE_DEPRECATION_WARNINGS	
 	PreviewPlatformChanged.Broadcast();
+PRAGMA_ENABLE_DEPRECATION_WARNINGS
 
 	if (bSaveSettings)
 	{
@@ -6824,7 +6837,9 @@ void UEditorEngine::ToggleFeatureLevelPreview()
 	ERHIFeatureLevel::Type NewPreviewFeatureLevel = PreviewPlatform.GetEffectivePreviewFeatureLevel();
 
 	DefaultWorldFeatureLevel = NewPreviewFeatureLevel;
+PRAGMA_DISABLE_DEPRECATION_WARNINGS
 	PreviewFeatureLevelChanged.Broadcast(NewPreviewFeatureLevel);
+PRAGMA_ENABLE_DEPRECATION_WARNINGS
 
 	Scalability::ChangeScalabilityPreviewPlatform(PreviewPlatform.GetEffectivePreviewPlatformName(), GetActiveShaderPlatform());
 
@@ -6854,6 +6869,8 @@ void UEditorEngine::ToggleFeatureLevelPreview()
 		UDeviceProfileManager::Get().RestorePreviewDeviceProfile();
 	}
 
+	PreviewShaderPlatformChanged.Broadcast(GShaderPlatformForFeatureLevel[NewPreviewFeatureLevel]);
+
 	// Update any PerPlatformConfig class defaults or instances
 	for (FThreadSafeObjectIterator ObjIterator(UObject::StaticClass(), RF_NoFlags); ObjIterator; ++ObjIterator)
 	{
@@ -6868,7 +6885,9 @@ void UEditorEngine::ToggleFeatureLevelPreview()
 
 	PreviewPlatformChangePost();
 	
+PRAGMA_DISABLE_DEPRECATION_WARNINGS	
 	PreviewPlatformChanged.Broadcast();
+PRAGMA_ENABLE_DEPRECATION_WARNINGS
 
 	UStaticMesh::OnLodStrippingQualityLevelChanged(nullptr);
 
diff --git a/Engine/Source/Editor/UnrealEd/Private/EditorServer.cpp b/Engine/Source/Editor/UnrealEd/Private/EditorServer.cpp
--- a/Engine/Source/Editor/UnrealEd/Private/EditorServer.cpp
+++ b/Engine/Source/Editor/UnrealEd/Private/EditorServer.cpp
@@ -2670,7 +2670,7 @@ bool UEditorEngine::Map_Load(const TCHAR* Str, FOutputDevice& Ar)
 					// We do not handle reinitializing. If the World was initialized during Load, we require that it is still fully initialized
 					// This should have been guaranteed by UWorld::KeepInitializedDuringLoadTag
 					check(World->IsInitialized());
-					World->ChangeFeatureLevel(FeatureLevel);
+					World->ShaderPlatformChanged(GShaderPlatformForFeatureLevel[FeatureLevel]);
 				}
 				else
 				{
diff --git a/Engine/Source/Editor/UnrealEd/Private/PlayLevel.cpp b/Engine/Source/Editor/UnrealEd/Private/PlayLevel.cpp
--- a/Engine/Source/Editor/UnrealEd/Private/PlayLevel.cpp
+++ b/Engine/Source/Editor/UnrealEd/Private/PlayLevel.cpp
@@ -3435,7 +3435,7 @@ TSharedRef<SPIEViewport> UEditorEngine::GeneratePIEViewportWindow(const FRequest
 				}
 			}
 
-			static void OnPIEWindowClosed(const TSharedRef<SWindow>& WindowBeingClosed, TWeakPtr<SViewport> PIEViewportWidget, TWeakObjectPtr<UEditorEngine> OwningEditorEngine, int32 ViewportIndex, bool bRestoreRootWindow, FDelegateHandle PreviewFeatureLevelChangedHandle)
+			static void OnPIEWindowClosed(const TSharedRef<SWindow>& WindowBeingClosed, TWeakPtr<SViewport> PIEViewportWidget, TWeakObjectPtr<UEditorEngine> OwningEditorEngine, int32 ViewportIndex, bool bRestoreRootWindow, FDelegateHandle PreviewPlatformChangedHandle)
 			{
 				// Save off the window position
 				const FVector2D PIEWindowPos = WindowBeingClosed->GetLocalToScreenTransform().GetTranslation();
@@ -3447,9 +3447,9 @@ TSharedRef<SPIEViewport> UEditorEngine::GeneratePIEViewportWindow(const FRequest
 				{
 					OwningEditorEngine->StoreWindowSizeAndPositionForInstanceIndex(ViewportIndex, WindowSize, WindowPosition);
 
-					if (PreviewFeatureLevelChangedHandle.IsValid())
+					if (PreviewPlatformChangedHandle.IsValid())
 					{
-						OwningEditorEngine->OnPreviewFeatureLevelChanged().Remove(PreviewFeatureLevelChangedHandle);
+						OwningEditorEngine->OnPreviewShaderPlatformChanged().Remove(PreviewPlatformChangedHandle);
 					}
 				}
 
@@ -3468,19 +3468,19 @@ TSharedRef<SPIEViewport> UEditorEngine::GeneratePIEViewportWindow(const FRequest
 			}
 		};
 
-		FDelegateHandle PreviewFeatureLevelChangedHandle = OnPreviewFeatureLevelChanged().AddLambda([PieWindow, bHasCustomWindow, InViewportClient, InNetMode, InSessionParams, InWorldContext, bVRPreview](ERHIFeatureLevel::Type NewFeatureLevel)
+		FDelegateHandle PreviewPlatformChangedHandle = OnPreviewShaderPlatformChanged().AddLambda([PieWindow, bHasCustomWindow, InViewportClient, InNetMode, InSessionParams, InWorldContext, bVRPreview](EShaderPlatform NewShaderPlatform)
 			{
 				if (!bHasCustomWindow)
 				{
 					FText ViewportName = GeneratePIEViewportWindowTitle(InNetMode, InSessionParams, InWorldContext.PIEInstance, InWorldContext.PIEFixedTickSeconds, bVRPreview);
 					PieWindow->SetTitle(ViewportName);
 				}
-				InViewportClient->GetWorld()->ChangeFeatureLevel(NewFeatureLevel);
+				InViewportClient->GetWorld()->ShaderPlatformChanged(NewShaderPlatform);
 			});
 
 		PieWindow->SetRequestDestroyWindowOverride(FRequestDestroyWindowOverride::CreateStatic(&FLocal::RequestDestroyPIEWindowOverride, TWeakObjectPtr<UEditorEngine>(this)));
 		PieWindow->SetOnWindowClosed(FOnWindowClosed::CreateStatic(&FLocal::OnPIEWindowClosed, TWeakPtr<SViewport>(PieViewportWidget), TWeakObjectPtr<UEditorEngine>(this),
-			InViewportIndex, bShouldMinimizeRootWindowForVRPreview || bShouldMinimizeRootWindowForNonVRPreview, PreviewFeatureLevelChangedHandle));
+			InViewportIndex, bShouldMinimizeRootWindowForVRPreview || bShouldMinimizeRootWindowForNonVRPreview, PreviewPlatformChangedHandle));
 	}
 
 	// Create a new viewport that the viewport widget will use to render the game
diff --git a/Engine/Source/Editor/VirtualTexturingEditor/Private/RuntimeVirtualTextureBuildStreamingMips.cpp b/Engine/Source/Editor/VirtualTexturingEditor/Private/RuntimeVirtualTextureBuildStreamingMips.cpp
--- a/Engine/Source/Editor/VirtualTexturingEditor/Private/RuntimeVirtualTextureBuildStreamingMips.cpp
+++ b/Engine/Source/Editor/VirtualTexturingEditor/Private/RuntimeVirtualTextureBuildStreamingMips.cpp
@@ -435,7 +435,7 @@ namespace RuntimeVirtualTexture
 						CompileGlobalShaderMap(AltFeatureLevel);
 					}
 
-					InParams.World->ChangeFeatureLevel(AltFeatureLevel);
+					InParams.World->ShaderPlatformChanged(GShaderPlatformForFeatureLevel[AltFeatureLevel]);
 
 					// Make sure all assets are finished compiling. Recreate render state after shader compilation complete
 					{
@@ -480,7 +480,7 @@ namespace RuntimeVirtualTexture
 				{
 					Task.EnterProgressFrame();
 					UMaterialInterface::SetGlobalRequiredFeatureLevel(CurFeatureLevel, /*bShouldCompile = */ false);
-					InParams.World->ChangeFeatureLevel(CurFeatureLevel);
+					InParams.World->ShaderPlatformChanged(GShaderPlatformForFeatureLevel[CurFeatureLevel]);
 				}
 			}
 		}
diff --git a/Engine/Source/Runtime/Engine/Classes/Engine/World.h b/Engine/Source/Runtime/Engine/Classes/Engine/World.h
--- a/Engine/Source/Runtime/Engine/Classes/Engine/World.h
+++ b/Engine/Source/Runtime/Engine/Classes/Engine/World.h
@@ -28,6 +28,7 @@
 #include "AudioDeviceHandle.h"
 #include "Subsystems/WorldSubsystem.h"
 #include "Subsystems/SubsystemCollection.h"
+#include "RHIShaderPlatform.h"
 #include "CollisionProfile.h"
 #include "RHIFeatureLevel.h"
 #include "Engine/WorldInitializationValues.h"
@@ -1538,9 +1539,10 @@ public:
 #if WITH_EDITOR
 
 	/** Change the feature level that this world is current rendering with */
+	UE_DEPRECATED(5.8, "ChangeFeatureLevel is deprecated, use ShaderPlatformChanged instead")
 	UE_API void ChangeFeatureLevel(ERHIFeatureLevel::Type InFeatureLevel, bool bShowSlowProgressDialog = true, bool bForceUpdate = false);
 
-	UE_API void ShaderPlatformChanged();
+	UE_API void ShaderPlatformChanged(EShaderPlatform InShaderPlatform = EShaderPlatform::SP_NumPlatforms, bool bShowSlowProgressDialog = true);
 
 	UE_API void RecreateScene(ERHIFeatureLevel::Type InFeatureLevel, bool bBroadcastChange = true);
 
diff --git a/Engine/Source/Runtime/Engine/Private/GameInstance.cpp b/Engine/Source/Runtime/Engine/Private/GameInstance.cpp
--- a/Engine/Source/Runtime/Engine/Private/GameInstance.cpp
+++ b/Engine/Source/Runtime/Engine/Private/GameInstance.cpp
@@ -10,6 +10,7 @@
 #include "Misc/Paths.h"
 #include "Engine/Console.h"
 #include "Engine/GameViewportClient.h"
+#include "SceneInterface.h"
 #include "Misc/ScopedSlowTask.h"
 #include "Engine/GameEngine.h"
 #include "GameFramework/GameModeBase.h"
@@ -325,9 +326,9 @@ FGameInstancePIEResult UGameInstance::InitializeForPlayInEditor(int32 PIEInstanc
 			// Attempt to load the target world asset
 			FSoftObjectPath TargetWorld = FSoftObjectPath(Params.OverrideMapURL);
 			UWorld* WorldToDuplicate = Cast<UWorld>(TargetWorld.TryLoad());
-			if (WorldToDuplicate)
+			if (WorldToDuplicate && EditorWorld->Scene != nullptr)
 			{
-				WorldToDuplicate->ChangeFeatureLevel(EditorWorld->GetFeatureLevel(), false);
+				WorldToDuplicate->ShaderPlatformChanged(EditorWorld->Scene->GetShaderPlatform(), false);
 				NewWorld = EditorEngine->CreatePIEWorldByDuplication(*WorldContext, WorldToDuplicate, PIEMapName);
 			}
 		}
diff --git a/Engine/Source/Runtime/Engine/Private/SceneInterface.cpp b/Engine/Source/Runtime/Engine/Private/SceneInterface.cpp
--- a/Engine/Source/Runtime/Engine/Private/SceneInterface.cpp
+++ b/Engine/Source/Runtime/Engine/Private/SceneInterface.cpp
@@ -21,6 +21,7 @@ FSceneInterface::IPrimitiveTransformUpdater* FSceneInterface::IPrimitiveTransfor
 
 FSceneInterface::FSceneInterface(ERHIFeatureLevel::Type InFeatureLevel)
 	: FeatureLevel(InFeatureLevel)
+	, ShaderPlatform(GShaderPlatformForFeatureLevel[InFeatureLevel])
 {
 }
 
@@ -31,10 +32,6 @@ void FSceneInterface::UpdateAllPrimitiveSceneInfos(FRHICommandListImmediate& RHI
 	GraphBuilder.Execute();
 }
 
-EShaderPlatform FSceneInterface::GetShaderPlatform() const
-{
-	return GShaderPlatformForFeatureLevel[GetFeatureLevel()];
-}
 
 EShadingPath FSceneInterface::GetShadingPath(ERHIFeatureLevel::Type InFeatureLevel)
 {
diff --git a/Engine/Source/Runtime/Engine/Private/UnrealEngine.cpp b/Engine/Source/Runtime/Engine/Private/UnrealEngine.cpp
--- a/Engine/Source/Runtime/Engine/Private/UnrealEngine.cpp
+++ b/Engine/Source/Runtime/Engine/Private/UnrealEngine.cpp
@@ -16381,7 +16381,7 @@ bool UEngine::LoadMap( FWorldContext& WorldContext, FURL URL, class UPendingNetG
 			// PIE worlds should use the same feature level as the editor
 			if (WorldContext.PIEWorldFeatureLevel != ERHIFeatureLevel::Num && NewWorld->GetFeatureLevel() != WorldContext.PIEWorldFeatureLevel)
 			{
-				NewWorld->ChangeFeatureLevel(WorldContext.PIEWorldFeatureLevel);
+				NewWorld->ShaderPlatformChanged(GShaderPlatformForFeatureLevel[WorldContext.PIEWorldFeatureLevel]);
 			}
 #endif
 		}
diff --git a/Engine/Source/Runtime/Engine/Private/World.cpp b/Engine/Source/Runtime/Engine/Private/World.cpp
--- a/Engine/Source/Runtime/Engine/Private/World.cpp
+++ b/Engine/Source/Runtime/Engine/Private/World.cpp
@@ -8944,7 +8944,7 @@ UWorld* FSeamlessTravelHandler::Tick()
 			// PIE worlds should use the same feature level as the editor
 			if (CurrentContext.PIEWorldFeatureLevel != ERHIFeatureLevel::Num && LoadedWorld->GetFeatureLevel() != CurrentContext.PIEWorldFeatureLevel)
 			{
-				LoadedWorld->ChangeFeatureLevel(CurrentContext.PIEWorldFeatureLevel);
+				LoadedWorld->ShaderPlatformChanged(GShaderPlatformForFeatureLevel[CurrentContext.PIEWorldFeatureLevel]);
 			}
 #endif
 
@@ -10038,12 +10038,28 @@ void UWorld::DuplicateRequestedLevels(const FName MapName)
 }
 
 #if WITH_EDITOR
-void UWorld::ChangeFeatureLevel(ERHIFeatureLevel::Type InFeatureLevel, bool bShowSlowProgressDialog, bool bForceUpdate)
+PRAGMA_DISABLE_DEPRECATION_WARNINGS
+void UWorld::ChangeFeatureLevel(ERHIFeatureLevel::Type InFeatureLevel, bool bShowSlowProgressDialog, bool /*bForceUpdate*/)
 {
-	if (InFeatureLevel != GetFeatureLevel() || bForceUpdate)
+	ShaderPlatformChanged(GShaderPlatformForFeatureLevel[InFeatureLevel], bShowSlowProgressDialog);
+}
+PRAGMA_ENABLE_DEPRECATION_WARNINGS
+
+void UWorld::ShaderPlatformChanged(EShaderPlatform InShaderPlatform, bool bShowSlowProgressDialog)
+{
+	if (InShaderPlatform == EShaderPlatform::SP_NumPlatforms)
+	{
+		UE_LOG(LogWorld, Warning, TEXT("ShaderPlatformChanged called without a valid EShaderPlatform, defaulting to GMaxRHIShaderPlatform (%s). Please update the call site to pass an explicit EShaderPlatform."), *LexToString(GMaxRHIShaderPlatform));
+		InShaderPlatform = GMaxRHIShaderPlatform;
+	}
+
+	if (Scene && Scene->GetShaderPlatform() != InShaderPlatform)
 	{
-		UE_LOG(LogWorld, Log, TEXT("Changing Feature Level (Enum) from %s to %s%s"), *LexToString(GetFeatureLevel()), *LexToString(InFeatureLevel), bForceUpdate ? TEXT(" (forced)") : TEXT(""));
-		FScopedSlowTask SlowTask(100.f, NSLOCTEXT("Engine", "ChangingPreviewRenderingLevelMessage", "Changing Preview Rendering Level"), bShowSlowProgressDialog);
+		ERHIFeatureLevel::Type InFeatureLevel = GetMaxSupportedFeatureLevel(InShaderPlatform);
+		FText FriendlyNamePreviewFrom = FDataDrivenShaderPlatformInfo::GetFriendlyName(Scene->GetShaderPlatform());
+		FText FriendlyNamePreviewTo = FDataDrivenShaderPlatformInfo::GetFriendlyName(InShaderPlatform);
+		UE_LOG(LogWorld, Log, TEXT("Changing Preview Shader Platform from %s to %s"), *FriendlyNamePreviewFrom.ToString(), *FriendlyNamePreviewTo.ToString());
+		FScopedSlowTask SlowTask(100.f, NSLOCTEXT("Engine", "ChangingPreviewRenderingLevelMessage", "Changing Preview Shader Platform"), bShowSlowProgressDialog);
 		SlowTask.MakeDialog();
 		{
 			SlowTask.EnterProgressFrame(10.0f);
@@ -10078,14 +10094,6 @@ void UWorld::ChangeFeatureLevel(ERHIFeatureLevel::Type InFeatureLevel, bool bSho
 	}
 }
 
-void UWorld::ShaderPlatformChanged()
-{
-	if (Scene)
-	{
-		Scene->UpdateEarlyZPassMode();
-	}
-}
-
 void UWorld::RecreateScene(ERHIFeatureLevel::Type InFeatureLevel, bool bBroadcastChange)
 {
 	if (Scene)
diff --git a/Engine/Source/Runtime/Engine/Public/SceneInterface.h b/Engine/Source/Runtime/Engine/Public/SceneInterface.h
--- a/Engine/Source/Runtime/Engine/Public/SceneInterface.h
+++ b/Engine/Source/Runtime/Engine/Public/SceneInterface.h
@@ -729,7 +729,7 @@ public:
 
 	ERHIFeatureLevel::Type GetFeatureLevel() const { return FeatureLevel; }
 
-	ENGINE_API EShaderPlatform GetShaderPlatform() const;
+	EShaderPlatform GetShaderPlatform() const { return ShaderPlatform; }
 
 	static ENGINE_API EShadingPath GetShadingPath(ERHIFeatureLevel::Type InFeatureLevel);
 
@@ -818,5 +818,8 @@ protected:
 	/** This scene's feature level */
 	ERHIFeatureLevel::Type FeatureLevel;
 
+	/** The shader platform this scene was created with */
+	EShaderPlatform ShaderPlatform;
+
 	friend class FSceneViewStateReference;
 };
diff --git a/Engine/Source/Runtime/Renderer/Private/HeterogeneousVolumes/HeterogeneousVolumesVoxelGridPipeline.cpp b/Engine/Source/Runtime/Renderer/Private/HeterogeneousVolumes/HeterogeneousVolumesVoxelGridPipeline.cpp
--- a/Engine/Source/Runtime/Renderer/Private/HeterogeneousVolumes/HeterogeneousVolumesVoxelGridPipeline.cpp
+++ b/Engine/Source/Runtime/Renderer/Private/HeterogeneousVolumes/HeterogeneousVolumesVoxelGridPipeline.cpp
@@ -1331,7 +1331,7 @@ void BuildMajorantVoxelGrid(
 			TEXT("HeterogeneousVolumes.MajorantVoxelGridBuffer")
 		);
 
-		const FGlobalShaderMap* GlobalShaderMap = GetGlobalShaderMap(Scene->GetFeatureLevel());
+		const FGlobalShaderMap* GlobalShaderMap = GetGlobalShaderMap(Scene->GetShaderPlatform());
 
 		// Build Majorant Grid
 		{
@@ -1921,7 +1921,7 @@ void MarkTopLevelGrid(
 	GroupCount.Y = FMath::DivideAndRoundUp(TopLevelGridResolution.Y, FAllocateBottomLevelGrid::GetThreadGroupSize3D());
 	GroupCount.Z = FMath::DivideAndRoundUp(TopLevelGridResolution.Z, FAllocateBottomLevelGrid::GetThreadGroupSize3D());
 
-	const FGlobalShaderMap* GlobalShaderMap = GetGlobalShaderMap(Scene->GetFeatureLevel());
+	const FGlobalShaderMap* GlobalShaderMap = GetGlobalShaderMap(Scene->GetShaderPlatform());
 
 	FAllocateBottomLevelGrid::FPermutationDomain PermutationVector;
 	PermutationVector.Set<FAllocateBottomLevelGrid::FEnableIndirectionGrid>(HeterogeneousVolumes::EnableIndirectionGrid());
@@ -1978,7 +1978,7 @@ void GenerateRasterTiles(
 		GroupCount.Y = FMath::DivideAndRoundUp(TopLevelGridResolution.Y, FTopLevelGridCalculateVoxelSize::GetThreadGroupSize3D());
 		GroupCount.Z = FMath::DivideAndRoundUp(TopLevelGridResolution.Z, FTopLevelGridCalculateVoxelSize::GetThreadGroupSize3D());
 
-		const FGlobalShaderMap* GlobalShaderMap = GetGlobalShaderMap(Scene->GetFeatureLevel());
+		const FGlobalShaderMap* GlobalShaderMap = GetGlobalShaderMap(Scene->GetShaderPlatform());
 
 		FGenerateRasterTiles::FPermutationDomain PermutationVector;
 		PermutationVector.Set<FGenerateRasterTiles::FEnableIndirectionGrid>(bEnableIndirectionGrid);
@@ -2592,7 +2592,7 @@ void RasterizeVolumesIntoOrthoVoxelGrid(
 		IndirectArgsPassParameters->RWRasterizeBottomLevelGridIndirectArgsBuffer = GraphBuilder.CreateUAV(RasterizeBottomLevelGridIndirectArgsBuffer, PF_R32_UINT);
 	}
 
-	const FGlobalShaderMap* GlobalShaderMap = GetGlobalShaderMap(Scene->GetFeatureLevel());
+	const FGlobalShaderMap* GlobalShaderMap = GetGlobalShaderMap(Scene->GetShaderPlatform());
 
 	FIntVector GroupCount(1, 1, 1);
 	TShaderRef<FSetRasterizeBottomLevelGridIndirectArgs> ComputeShader = GlobalShaderMap->GetShader<FSetRasterizeBottomLevelGridIndirectArgs>();
diff --git a/Engine/Source/Runtime/Renderer/Private/PathTracing.cpp b/Engine/Source/Runtime/Renderer/Private/PathTracing.cpp
--- a/Engine/Source/Runtime/Renderer/Private/PathTracing.cpp
+++ b/Engine/Source/Runtime/Renderer/Private/PathTracing.cpp
@@ -2163,7 +2163,7 @@ void PreparePathTracingCloudMaterial(FRDGBuilder& GraphBuilder, FScene* Scene, T
 	TShaderRef<FPathTracingVolumetricCloudMaterialVisualize> CallableShaderVisualize;
 	if (bCloudAccelerationMapVisualize)
 	{
-		CallableShaderVisualize = GetGlobalShaderMap(Scene->GetFeatureLevel())->GetShader<FPathTracingVolumetricCloudMaterialVisualize>();
+		CallableShaderVisualize = GetGlobalShaderMap(Scene->GetShaderPlatform())->GetShader<FPathTracingVolumetricCloudMaterialVisualize>();
 		if (!CallableShaderVisualize.IsValid())
 		{
 			// asking for visualization, but visualization shader is not ready
@@ -2262,7 +2262,7 @@ void BindCallableShader(
 )
 {
 	//auto CallableShader = MaterialShaderMap->GetShader<CallableShaderType>();
-	auto CallableShader = GetGlobalShaderMap(Scene->GetFeatureLevel())->GetShader<CallableShaderType>();
+	auto CallableShader = GetGlobalShaderMap(Scene->GetShaderPlatform())->GetShader<CallableShaderType>();
 	if (!CallableShader.IsValid())
 	{
 		return;
diff --git a/Engine/Source/Runtime/Renderer/Private/RendererScene.cpp b/Engine/Source/Runtime/Renderer/Private/RendererScene.cpp
--- a/Engine/Source/Runtime/Renderer/Private/RendererScene.cpp
+++ b/Engine/Source/Runtime/Renderer/Private/RendererScene.cpp
@@ -1191,11 +1191,11 @@ FScene::FScene(UWorld* InWorld, bool bInRequiresHitProxies, bool bInIsEditorScen
 	World->Scene = this;
 
 	FeatureLevel = World->GetFeatureLevel();
+	ShaderPlatform = GShaderPlatformForFeatureLevel[FeatureLevel];
 
 	checkf((uint32)FeatureLevel < (uint32)ERHIFeatureLevel::Num, TEXT("World provided an invalid feature level (%d) to FScene."), FeatureLevel);
-	checkf(GShaderPlatformForFeatureLevel[FeatureLevel] != SP_NumPlatforms, TEXT("Invalid feature level %s for platform (max feature level %s)"), *LexToString(FeatureLevel), *LexToString(GMaxRHIFeatureLevel));
+	checkf(ShaderPlatform != SP_NumPlatforms, TEXT("Invalid feature level %s for platform (max feature level %s)"), *LexToString(FeatureLevel), *LexToString(GMaxRHIFeatureLevel));
 
-	const EShaderPlatform ShaderPlatform = GetFeatureLevelShaderPlatform(InFeatureLevel);
 	for (int32 NanitePass = 0; NanitePass < ENaniteMeshPass::Num; ++NanitePass)
 	{
 		const bool bBasePass = (ENaniteMeshPass::Type)NanitePass == ENaniteMeshPass::BasePass;	// Only basepass for now
@@ -2312,7 +2312,6 @@ void FScene::AddLightSceneInfo_RenderThread(FLightSceneInfo* LightSceneInfo)
 		LightSceneInfo->Proxy->IESAtlasId = IESAtlas::AddTexture(IESTexture);
 	}
 
-	const EShaderPlatform ShaderPlatform = GetShaderPlatform();
 	if (IsMobilePlatform(ShaderPlatform) && LightSceneInfo->Proxy->CastsDynamicShadow())
 	{
 		if (IsStaticLightingAllowed() && MobileUsesShadowMaskTexture(ShaderPlatform))
diff --git a/Engine/Source/Runtime/Renderer/Private/SceneCulling/SceneCulling.cpp b/Engine/Source/Runtime/Renderer/Private/SceneCulling/SceneCulling.cpp
--- a/Engine/Source/Runtime/Renderer/Private/SceneCulling/SceneCulling.cpp
+++ b/Engine/Source/Runtime/Renderer/Private/SceneCulling/SceneCulling.cpp
@@ -2328,7 +2328,7 @@ public:
 					PassParameters->OutExplicitChunkBoundsBuffer = GraphBuilder.CreateUAV(ExplicitChunkBoundsBufferRDG);
 					PassParameters->OutExplicitChunkCellIdsBuffer = GraphBuilder.CreateUAV(ExplicitChunkCellIdsRDG);
 
-					auto ComputeShader = GetGlobalShaderMap(SceneCulling.Scene.GetFeatureLevel())->GetShader<FComputeExplicitChunkBounds_CS>();
+					auto ComputeShader = GetGlobalShaderMap(SceneCulling.Scene.GetShaderPlatform())->GetShader<FComputeExplicitChunkBounds_CS>();
 
 					FComputeShaderUtils::AddPass(
 						GraphBuilder,
diff --git a/Engine/Source/Runtime/Renderer/Private/Skinning/SkinningSceneExtension.cpp b/Engine/Source/Runtime/Renderer/Private/Skinning/SkinningSceneExtension.cpp
--- a/Engine/Source/Runtime/Renderer/Private/Skinning/SkinningSceneExtension.cpp
+++ b/Engine/Source/Runtime/Renderer/Private/Skinning/SkinningSceneExtension.cpp
@@ -1621,7 +1621,7 @@ void FSkinningSceneExtension::FRenderer::UpdateViewData(FRDGBuilder& GraphBuilde
 		PassParameters->IndirectArgs = InstanceWorkArgsRDG;
 		PassParameters->InstanceWorkGroups = GraphBuilder.CreateSRV(InstanceWorkGroupsRDG);
 
-		auto ComputeShader = GetGlobalShaderMap(SceneData->Scene.GetFeatureLevel())->GetShader<FNaniteSkinningUpdateViewDataCS>();
+		auto ComputeShader = GetGlobalShaderMap(SceneData->Scene.GetShaderPlatform())->GetShader<FNaniteSkinningUpdateViewDataCS>();
 
 		FComputeShaderUtils::AddPass(
 			GraphBuilder,
diff --git a/Engine/Source/Runtime/Renderer/Private/SplineMeshSceneResources.cpp b/Engine/Source/Runtime/Renderer/Private/SplineMeshSceneResources.cpp
--- a/Engine/Source/Runtime/Renderer/Private/SplineMeshSceneResources.cpp
+++ b/Engine/Source/Runtime/Renderer/Private/SplineMeshSceneResources.cpp
@@ -545,7 +545,7 @@ void FSplineMeshSceneUpdater::AddUpdatePass(
 	PassParameters->TextureHeight = Extent.Y;
 	PassParameters->TextureHeightInv = InvExtent.Y;
 
-	auto ComputeShader = GetGlobalShaderMap(Scene.GetFeatureLevel())->GetShader<FSplineMeshTextureFillCS>();
+	auto ComputeShader = GetGlobalShaderMap(Scene.GetShaderPlatform())->GetShader<FSplineMeshTextureFillCS>();
 
 	const uint32 NumThreadGroups = NumUpdateRequests ? NumUpdateRequests : SceneData->SlotAllocator.GetMaxSize();
 	FComputeShaderUtils::AddPass(
diff --git a/Engine/Source/Runtime/Renderer/Private/VirtualShadowMaps/VirtualShadowMapArray.cpp b/Engine/Source/Runtime/Renderer/Private/VirtualShadowMaps/VirtualShadowMapArray.cpp
--- a/Engine/Source/Runtime/Renderer/Private/VirtualShadowMaps/VirtualShadowMapArray.cpp
+++ b/Engine/Source/Runtime/Renderer/Private/VirtualShadowMaps/VirtualShadowMapArray.cpp
@@ -1930,7 +1930,7 @@ void FVirtualShadowMapArray::PostRender(FRDGBuilder& GraphBuilder)
 		PassParameters->VirtualShadowMap = GetUncachedUniformBuffer(GraphBuilder);
 		PassParameters->OutPhysicalPageMetaData = GraphBuilder.CreateUAV(PhysicalPageMetaDataRDG);
 		PassParameters->DirtyPageFlagsInOut = GraphBuilder.CreateUAV(DirtyPageFlagsRDG);
-		auto ComputeShader = GetGlobalShaderMap(Scene.GetFeatureLevel())->GetShader<FUpdateAndClearDirtyFlagsCS>();
+		auto ComputeShader = GetGlobalShaderMap(Scene.GetShaderPlatform())->GetShader<FUpdateAndClearDirtyFlagsCS>();
 
 		FComputeShaderUtils::AddPass(
 			GraphBuilder,
@@ -1967,7 +1967,7 @@ void FVirtualShadowMapArray::PostRender(FRDGBuilder& GraphBuilder)
 			FSelectPagesToPostProcessCS::FPermutationDomain PermutationVector;
 			SetStatsArgsAndPermutation<FSelectPagesToPostProcessCS>(ShouldGenerateStats(), StatsBufferUAV, PassParameters, PermutationVector);
 
-			auto ComputeShader = GetGlobalShaderMap(Scene.GetFeatureLevel())->GetShader<FSelectPagesToPostProcessCS>(PermutationVector);
+			auto ComputeShader = GetGlobalShaderMap(Scene.GetShaderPlatform())->GetShader<FSelectPagesToPostProcessCS>(PermutationVector);
 
 			FComputeShaderUtils::AddPass(
 				GraphBuilder,
@@ -1987,7 +1987,7 @@ void FVirtualShadowMapArray::PostRender(FRDGBuilder& GraphBuilder)
 			PassParameters->OutPhysicalPagePool = GraphBuilder.CreateUAV(PhysicalPagePoolRDG);
 			PassParameters->IndirectArgs = PostProcessPagesIndirectArgsRDG;
 			PassParameters->PhysicalPagesToMerge = GraphBuilder.CreateSRV(PhysicalPagesToMergeRDG);
-			auto ComputeShader = GetGlobalShaderMap(Scene.GetFeatureLevel())->GetShader<FMergeStaticPhysicalPagesIndirectCS>();
+			auto ComputeShader = GetGlobalShaderMap(Scene.GetShaderPlatform())->GetShader<FMergeStaticPhysicalPagesIndirectCS>();
 
 			FComputeShaderUtils::AddPass(
 				GraphBuilder,
@@ -2024,7 +2024,7 @@ void FVirtualShadowMapArray::PostRender(FRDGBuilder& GraphBuilder)
 					PassParameters->PhysicalPagesWithSeparableFilterable = GraphBuilder.CreateSRV(PhysicalPagesWithSeparableFilterableRDG);
 					PassParameters->IndirectArgs = PostProcessPagesIndirectArgsRDG;
 			
-					auto ComputeShader = GetGlobalShaderMap(Scene.GetFeatureLevel())->GetShader<FUpdateFilterablePagesSeparableX_CS>();
+					auto ComputeShader = GetGlobalShaderMap(Scene.GetShaderPlatform())->GetShader<FUpdateFilterablePagesSeparableX_CS>();
 					FComputeShaderUtils::AddPass(
 						GraphBuilder,
 						RDG_EVENT_NAME("FUpdateFilterablePagesSeparableX_CS"),
@@ -2078,7 +2078,7 @@ void FVirtualShadowMapArray::PostRender(FRDGBuilder& GraphBuilder)
 						PassParameters->PrevVirtualShadowMapIds = GraphBuilder.CreateSRV(GSystemTextures.GetDefaultStructuredBuffer(GraphBuilder, 4, 0u));
 					}
 
-					auto ComputeShader = GetGlobalShaderMap(Scene.GetFeatureLevel())->GetShader<FUpdateFilterablePagesSeparableY_CS>();
+					auto ComputeShader = GetGlobalShaderMap(Scene.GetShaderPlatform())->GetShader<FUpdateFilterablePagesSeparableY_CS>();
 					FComputeShaderUtils::AddPass(
 						GraphBuilder,
 						RDG_EVENT_NAME("FUpdateFilterablePagesSeparableY_CS"),
@@ -2101,7 +2101,7 @@ void FVirtualShadowMapArray::PostRender(FRDGBuilder& GraphBuilder)
 				PassParameters->PhysicalPagesWithSeparableFilterable = GraphBuilder.CreateSRV(PhysicalPagesWithSeparableFilterableRDG);
 				PassParameters->IndirectArgs = PostProcessPagesIndirectArgsRDG;
 
-				auto ComputeShader = GetGlobalShaderMap(Scene.GetFeatureLevel())->GetShader<FGatherFilterableBorderTexels_CS>();
+				auto ComputeShader = GetGlobalShaderMap(Scene.GetShaderPlatform())->GetShader<FGatherFilterableBorderTexels_CS>();
 
 				FComputeShaderUtils::AddPass(
 					GraphBuilder,
@@ -2126,7 +2126,7 @@ void FVirtualShadowMapArray::PostRender(FRDGBuilder& GraphBuilder)
 				PassParameters->PhysicalPagesWithFilterable = GraphBuilder.CreateSRV(PhysicalPagesWithFilterableRDG);
 				PassParameters->IndirectArgs = PostProcessPagesIndirectArgsRDG;
 
-				auto ComputeShader = GetGlobalShaderMap(Scene.GetFeatureLevel())->GetShader<FMergeDownFilterablePagesCS>();
+				auto ComputeShader = GetGlobalShaderMap(Scene.GetShaderPlatform())->GetShader<FMergeDownFilterablePagesCS>();
 
 				FComputeShaderUtils::AddPass(
 					GraphBuilder,
@@ -2339,7 +2339,7 @@ void FVirtualShadowMapArray::UpdateVisualizeLight(
 
 void FVirtualShadowMapArray::AppendPhysicalPageList(FRDGBuilder& GraphBuilder, bool bEmptyToAvailable)
 {
-	auto ComputeShader = GetGlobalShaderMap(Scene.GetFeatureLevel())->GetShader<FAppendPhysicalPageListsCS>();
+	auto ComputeShader = GetGlobalShaderMap(Scene.GetShaderPlatform())->GetShader<FAppendPhysicalPageListsCS>();
 	
 	FAppendPhysicalPageListsCS::FParameters* PassParameters = GraphBuilder.AllocParameters<FAppendPhysicalPageListsCS::FParameters>();
 	PassParameters->VirtualShadowMap		= GetUncachedUniformBuffer(GraphBuilder);
@@ -2463,7 +2463,7 @@ void FVirtualShadowMapArray::UpdatePhysicalPageAddresses(FRDGBuilder& GraphBuild
 
 	SetStatsArgsAndPermutation<FUpdatePhysicalPageAddresses>(ShouldGenerateStats(), StatsBufferUAV, PassParameters, PermutationVector);
 
-	auto ComputeShader = GetGlobalShaderMap(Scene.GetFeatureLevel())->GetShader<FUpdatePhysicalPageAddresses>(PermutationVector);
+	auto ComputeShader = GetGlobalShaderMap(Scene.GetShaderPlatform())->GetShader<FUpdatePhysicalPageAddresses>(PermutationVector);
 
 	FComputeShaderUtils::AddPass(
 		GraphBuilder,
@@ -2680,7 +2680,7 @@ void FVirtualShadowMapArray::BeginMarkPages(
 		
 		FClearPageTableCS::FPermutationDomain PermutationVector;
 		PermutationVector.Set<FClearPageTableCS::FNumMipLevelsDim>(Dest->Desc.NumMips);
-		auto ComputeShader = GetGlobalShaderMap(Scene.GetFeatureLevel())->GetShader<FClearPageTableCS>(PermutationVector);
+		auto ComputeShader = GetGlobalShaderMap(Scene.GetShaderPlatform())->GetShader<FClearPageTableCS>(PermutationVector);
 
 		if (bDirectionalOnly)
 		{
@@ -2795,7 +2795,7 @@ void FVirtualShadowMapArray::BeginMarkPages(
 		PassParameters->NumPageRectsToClear = NumPageRectsToClear;
 		PassParameters->OutPhysicalPageLists = GraphBuilder.CreateUAV(PhysicalPageListsRDG);
 
-		auto ComputeShader = GetGlobalShaderMap(Scene.GetFeatureLevel())->GetShader<FInitPageRectBoundsCS>();
+		auto ComputeShader = GetGlobalShaderMap(Scene.GetShaderPlatform())->GetShader<FInitPageRectBoundsCS>();
 		ClearUnusedGraphResources(ComputeShader, PassParameters);
 
 		FComputeShaderUtils::AddPass(
@@ -2883,7 +2883,7 @@ void FVirtualShadowMapArray::BeginMarkPages(
 
 		PassParameters->OutThrottleBuffer = GraphBuilder.CreateUAV(ThrottleBufferRDG);
 
-		auto ComputeShader = GetGlobalShaderMap(Scene.GetFeatureLevel())->GetShader<FProcessPrevFramePerfDataCS>();
+		auto ComputeShader = GetGlobalShaderMap(Scene.GetShaderPlatform())->GetShader<FProcessPrevFramePerfDataCS>();
 		ClearUnusedGraphResources(ComputeShader, PassParameters);
 
 		FComputeShaderUtils::AddPass(
@@ -2909,7 +2909,7 @@ void FVirtualShadowMapArray::BeginMarkPages(
 
 		PassParameters->InOutProjectionData = GraphBuilder.CreateUAV(ProjectionDataRDG);
 
-		auto ComputeShader = GetGlobalShaderMap(Scene.GetFeatureLevel())->GetShader<FUpdateThrottleParametersCS>();
+		auto ComputeShader = GetGlobalShaderMap(Scene.GetShaderPlatform())->GetShader<FUpdateThrottleParametersCS>();
 		ClearUnusedGraphResources(ComputeShader, PassParameters);
 
 		FComputeShaderUtils::AddPass(
@@ -3025,7 +3025,7 @@ void FVirtualShadowMapArray::BeginMarkPages(
 		PassParameters->bMarkCoarsePagesLocal = CVarMarkCoarsePagesLocal.GetValueOnRenderThread() != 0 ? 1 : 0;
 		PassParameters->bIncludeNonNaniteGeometry = CVarCoarsePagesIncludeNonNanite.GetValueOnRenderThread();
 
-		auto ComputeShader = GetGlobalShaderMap(Scene.GetFeatureLevel())->GetShader<FMarkCoarsePagesCS>();
+		auto ComputeShader = GetGlobalShaderMap(Scene.GetShaderPlatform())->GetShader<FMarkCoarsePagesCS>();
 
 		FComputeShaderUtils::AddPass(
 			GraphBuilder,
@@ -3417,7 +3417,7 @@ void FVirtualShadowMapArray::BuildPageAllocations(
 			FSelectPagesToInitializeCS::FPermutationDomain PermutationVector;
 			SetStatsArgsAndPermutation<FSelectPagesToInitializeCS>(ShouldGenerateStats(), StatsBufferUAV, PassParameters, PermutationVector);
 
-			auto ComputeShader = GetGlobalShaderMap(Scene.GetFeatureLevel())->GetShader<FSelectPagesToInitializeCS>(PermutationVector);
+			auto ComputeShader = GetGlobalShaderMap(Scene.GetShaderPlatform())->GetShader<FSelectPagesToInitializeCS>(PermutationVector);
 
 			FComputeShaderUtils::AddPass(
 				GraphBuilder,
@@ -3436,7 +3436,7 @@ void FVirtualShadowMapArray::BuildPageAllocations(
 			PassParameters->OutPhysicalPagePool = GraphBuilder.CreateUAV(PhysicalPagePoolRDG);
 			PassParameters->IndirectArgs = InitializePagesIndirectArgsRDG;
 			PassParameters->PhysicalPagesToInitialize = GraphBuilder.CreateSRV(PhysicalPagesToInitializeRDG);
-			auto ComputeShader = GetGlobalShaderMap(Scene.GetFeatureLevel())->GetShader<FInitializePhysicalPagesIndirectCS>();
+			auto ComputeShader = GetGlobalShaderMap(Scene.GetShaderPlatform())->GetShader<FInitializePhysicalPagesIndirectCS>();
 
 			FComputeShaderUtils::AddPass(
 				GraphBuilder,
@@ -3473,7 +3473,7 @@ void FVirtualShadowMapArray::BuildPageAllocations(
 		PassParameters->StatusMessageId = CacheManager->GetStatusFeedbackMessageId();
 		PassParameters->VirtualShadowMap = GetUncachedUniformBuffer(GraphBuilder);
 
-		auto ComputeShader = GetGlobalShaderMap(Scene.GetFeatureLevel())->GetShader<FVirtualSmFeedbackStatusCS>();
+		auto ComputeShader = GetGlobalShaderMap(Scene.GetShaderPlatform())->GetShader<FVirtualSmFeedbackStatusCS>();
 
 		FComputeShaderUtils::AddPass(
 			GraphBuilder,
@@ -3581,7 +3581,7 @@ void FVirtualShadowMapArray::RenderDebugInfo(FRDGBuilder& GraphBuilder, TArrayVi
 
 		PassParameters->OutVisualize = GraphBuilder.CreateUAV(DebugVisualizationOutput[ViewIndex]);
 
-		auto ComputeShader = GetGlobalShaderMap(Scene.GetFeatureLevel())->GetShader<FDebugVisualizeVirtualSmCS>();
+		auto ComputeShader = GetGlobalShaderMap(Scene.GetShaderPlatform())->GetShader<FDebugVisualizeVirtualSmCS>();
 
 		FComputeShaderUtils::AddPass(
 			GraphBuilder,
@@ -4836,7 +4836,7 @@ void FVirtualShadowMapArray::UpdateHZB(FRDGBuilder& GraphBuilder)
 		PassParameters->bForceFullHZBUpdate = CVarShadowsVirtualForceFullHZBUpdate.GetValueOnRenderThread();
 		FSelectPagesForHZBAndUpdateDirtyFlagsCS::FPermutationDomain PermutationVector;
 		SetStatsArgsAndPermutation<FSelectPagesForHZBAndUpdateDirtyFlagsCS>(ShouldGenerateStats(), StatsBufferUAV, PassParameters, PermutationVector);
-		auto ComputeShader = GetGlobalShaderMap(Scene.GetFeatureLevel())->GetShader<FSelectPagesForHZBAndUpdateDirtyFlagsCS>(PermutationVector);
+		auto ComputeShader = GetGlobalShaderMap(Scene.GetShaderPlatform())->GetShader<FSelectPagesForHZBAndUpdateDirtyFlagsCS>(PermutationVector);
 
 		FComputeShaderUtils::AddPass(
 			GraphBuilder,
@@ -4864,7 +4864,7 @@ void FVirtualShadowMapArray::UpdateHZB(FRDGBuilder& GraphBuilder)
 
 		PassParameters->IndirectArgs = PagesForHZBIndirectArgsRDG;
 		PassParameters->PhysicalPagesForHZB = GraphBuilder.CreateSRV(PhysicalPagesForHZBRDG);
-		auto ComputeShader = GetGlobalShaderMap(Scene.GetFeatureLevel())->GetShader<FVirtualSmBuildHZBPerPageCS>();
+		auto ComputeShader = GetGlobalShaderMap(Scene.GetShaderPlatform())->GetShader<FVirtualSmBuildHZBPerPageCS>();
 
 		FComputeShaderUtils::AddPass(
 			GraphBuilder,
@@ -4892,7 +4892,7 @@ void FVirtualShadowMapArray::UpdateHZB(FRDGBuilder& GraphBuilder)
 
 		PassParameters->IndirectArgs = PagesForHZBIndirectArgsRDG;
 		PassParameters->PhysicalPagesForHZB = GraphBuilder.CreateSRV(PhysicalPagesForHZBRDG);
-		auto ComputeShader = GetGlobalShaderMap(Scene.GetFeatureLevel())->GetShader<FVirtualSmBBuildHZBPerPageTopCS>();
+		auto ComputeShader = GetGlobalShaderMap(Scene.GetShaderPlatform())->GetShader<FVirtualSmBBuildHZBPerPageTopCS>();
 
 		FComputeShaderUtils::AddPass(
 			GraphBuilder,
diff --git a/Engine/Source/Runtime/Renderer/Private/VirtualShadowMaps/VirtualShadowMapCacheManager.cpp b/Engine/Source/Runtime/Renderer/Private/VirtualShadowMaps/VirtualShadowMapCacheManager.cpp
--- a/Engine/Source/Runtime/Renderer/Private/VirtualShadowMaps/VirtualShadowMapCacheManager.cpp
+++ b/Engine/Source/Runtime/Renderer/Private/VirtualShadowMaps/VirtualShadowMapCacheManager.cpp
@@ -1542,7 +1542,7 @@ void FVirtualShadowMapArrayCacheManager::ExtractStats(FRDGBuilder& GraphBuilder,
 			}
 		}
 
-		auto ComputeShader = GetGlobalShaderMap(Scene.GetFeatureLevel())->GetShader<FVirtualSmCopyStatsCS>();
+		auto ComputeShader = GetGlobalShaderMap(Scene.GetShaderPlatform())->GetShader<FVirtualSmCopyStatsCS>();
 
 		FComputeShaderUtils::AddPass(
 			GraphBuilder,
@@ -1965,7 +1965,7 @@ void FVirtualShadowMapArrayCacheManager::ProcessInvalidations(
 	FInvalidateInstancePagesLoadBalancerCS::FPermutationDomain PermutationVector;
 	PermutationVector.Set<FInvalidateInstancePagesLoadBalancerCS::FUseHzbDim>(PassParameters->InvalidatePagesParameters.HZBTextureArray != nullptr);
 	
-	auto ComputeShader = GetGlobalShaderMap(Scene.GetFeatureLevel())->GetShader<FInvalidateInstancePagesLoadBalancerCS>(PermutationVector);
+	auto ComputeShader = GetGlobalShaderMap(Scene.GetShaderPlatform())->GetShader<FInvalidateInstancePagesLoadBalancerCS>(PermutationVector);
 
 	//UE_LOGF(LogRenderer, Display, "VirtualShadowMapCacheManager: Invalidation %lu instances!", Instances.GetTotalNumInstances());
 
@@ -2142,7 +2142,7 @@ void FVirtualShadowMapInvalidationSceneUpdater::PostGPUSceneUpdate(FRDGBuilder&
 		if (!LoadBalancer.IsEmpty())
 		{
 			FLoadBalancer::FGPUData LoadBalancerGpuData = LoadBalancer.Upload(GraphBuilder);
-			auto ComputeShader = GetGlobalShaderMap(CacheManager.Scene.GetFeatureLevel())->GetShader<FVSMResetInstanceStateCS>();
+			auto ComputeShader = GetGlobalShaderMap(CacheManager.Scene.GetShaderPlatform())->GetShader<FVSMResetInstanceStateCS>();
 
 			FGPUSceneResourceParameters GPUSceneParams = CacheManager.Scene.GPUScene.GetShaderParameters(GraphBuilder);
 
