Optimized ASTP to use the single pass mode for all skeletons that have unified bone indices and <= 256 bones in the skeleton.

#Synced-CL 52010655

diff --git a/Engine/Shaders/Private/Skinning/AnimSequenceTransformProvider.usf b/Engine/Shaders/Private/Skinning/AnimSequenceTransformProvider.usf
--- a/Engine/Shaders/Private/Skinning/AnimSequenceTransformProvider.usf
+++ b/Engine/Shaders/Private/Skinning/AnimSequenceTransformProvider.usf
@@ -12,6 +12,8 @@ ByteAddressBuffer RetargetingDataBuffer;
 ByteAddressBuffer SequenceTransformBuffer;
 ByteAddressBuffer EvaluateBlockHeaderBuffer;
 RWByteAddressBuffer LocalSpaceTransformBufferUAV;
+uint NumTotalHeaders;
+uint HeaderBaseOffset;
 
 #define RETARGET_MODE_ANIMATION 0
 #define RETARGET_MODE_SKELETON 1
@@ -131,18 +133,17 @@ FBoneTransformWithUniformScale EvaluateAndRetargetBone(
 [numthreads(BONES_PER_GROUP, 1, 1)]
 void EvaluateAnimSequenceCS(uint3 GroupId : SV_GroupID, uint GroupIndex : SV_GroupIndex)
 {
-	const uint BoneBlock      = GroupId.x;
-	const uint LocalBoneIndex = GroupIndex;
+	const uint InstanceIndex = GroupId.y;
+	const uint HeaderSize    = uint(sizeof(FEvaluateAnimSequenceBlockHeader));
 
-	const FEvaluateAnimSequenceBlockHeader BlockHeader = EvaluateBlockHeaderBuffer.Load<FEvaluateAnimSequenceBlockHeader>(BoneBlock * uint(sizeof(FEvaluateAnimSequenceBlockHeader)));
+	const FEvaluateAnimSequenceBlockHeader BlockHeader = EvaluateBlockHeaderBuffer.Load<FEvaluateAnimSequenceBlockHeader>((HeaderBaseOffset + InstanceIndex) * HeaderSize);
 
-	if (LocalBoneIndex >= BlockHeader.BlockTransformCount)
+	const uint BoneIndex = GroupId.x * BONES_PER_GROUP + GroupIndex;
+	if (BoneIndex >= BlockHeader.NumCompactBones)
 	{
 		return;
 	}
 
-	const uint BoneIndex = LocalBoneIndex + (BlockHeader.BlockLocalIndex * BONES_PER_GROUP);
-
 	FBoneTransformWithUniformScale MeshLocalAnimTransform = EvaluateAndRetargetBone(
 		BoneIndex,
 		BlockHeader.SequenceTransformBufferOffset,
@@ -159,18 +160,17 @@ void EvaluateAnimSequenceCS(uint3 GroupId : SV_GroupID, uint GroupIndex : SV_Gro
 [numthreads(BONES_PER_GROUP, 1, 1)]
 void EvaluateAnimSequenceBlendCS(uint3 GroupId : SV_GroupID, uint GroupIndex : SV_GroupIndex)
 {
-	const uint BoneBlock          = GroupId.x;
-	const uint LocalBoneIndex = GroupIndex;
+	const uint InstanceIndex = GroupId.y;
+	const uint HeaderSize    = uint(sizeof(FEvaluateAnimSequenceBlendBlockHeader));
 
-	const FEvaluateAnimSequenceBlendBlockHeader BlockHeader = EvaluateBlockHeaderBuffer.Load<FEvaluateAnimSequenceBlendBlockHeader>(BoneBlock * uint(sizeof(FEvaluateAnimSequenceBlendBlockHeader)));
+	const FEvaluateAnimSequenceBlendBlockHeader BlockHeader = EvaluateBlockHeaderBuffer.Load<FEvaluateAnimSequenceBlendBlockHeader>((HeaderBaseOffset + InstanceIndex) * HeaderSize);
 
-	if (LocalBoneIndex >= BlockHeader.BlockTransformCount)
+	const uint BoneIndex = GroupId.x * BONES_PER_GROUP + GroupIndex;
+	if (BoneIndex >= BlockHeader.NumCompactBones)
 	{
 		return;
 	}
 
-	const uint BoneIndex = LocalBoneIndex + (BlockHeader.BlockLocalIndex * BONES_PER_GROUP);
-
 	FBoneTransformWithUniformScale SourceMeshLocalAnimTransform = EvaluateAndRetargetBone(
 		BoneIndex,
 		BlockHeader.SourceSequenceTransformBufferOffset,
@@ -218,12 +218,11 @@ FBoneTransformWithUniformScale LoadLocalSpaceTransform(uint TransformBufferOffse
 [numthreads(BONES_PER_GROUP, 1, 1)]
 void ConcatenateLocalToGlobalCS(uint3 GroupId : SV_GroupID, uint GroupIndex : SV_GroupIndex)
 {
-	const uint BlockIndex         = GroupId.x;
-	const uint LocalLeafBoneIndex = GroupIndex;
-	const uint PrefetchBoneIndex  = GroupIndex;
-	const uint PathElementSize    = 4;
+	const uint InstanceIndex     = GroupId.y;
+	const uint PrefetchBoneIndex = GroupIndex;
+	const uint PathElementSize   = 4;
 
-	FConcatenateBlockHeader BlockHeader = ConcatenateBlockHeaderBuffer.Load<FConcatenateBlockHeader>(BlockIndex * uint(sizeof(FConcatenateBlockHeader)));
+	FConcatenateBlockHeader BlockHeader = ConcatenateBlockHeaderBuffer.Load<FConcatenateBlockHeader>((HeaderBaseOffset + InstanceIndex) * uint(sizeof(FConcatenateBlockHeader)));
 
 	if (PrefetchBoneIndex < BlockHeader.NumPrefetchBones)
 	{
@@ -231,13 +230,12 @@ void ConcatenateLocalToGlobalCS(uint3 GroupId : SV_GroupID, uint GroupIndex : SV
 	}
 	GroupMemoryBarrierWithGroupSync();
 
-	if (LocalLeafBoneIndex >= BlockHeader.BlockLeafBoneCount)
+	const uint LeafBoneIndex = GroupId.x * BONES_PER_GROUP + GroupIndex;
+	if (LeafBoneIndex >= BlockHeader.NumLeafBones)
 	{
 		return;
 	}
 
-	const uint LeafBoneIndex = LocalLeafBoneIndex + (BlockHeader.BlockLocalIndex * BONES_PER_GROUP);
-
 	// Load relative offset from offset table, then add base to get absolute offset
 	const uint RelativePathHeaderOffset = LeafBonePathBuffer.Load<uint>(BlockHeader.LeafBonePathBufferOffset + LeafBoneIndex * PathElementSize);
 	const uint PathHeaderOffset = BlockHeader.LeafBonePathBufferOffset + RelativePathHeaderOffset;
@@ -296,18 +294,17 @@ void ScatterGlobalSpaceTransform(uint DstTransformBufferOffset1, uint DstTransfo
 [numthreads(BONES_PER_GROUP, 1, 1)]
 void ScatterTransformsCS(uint3 GroupId : SV_GroupID, uint GroupIndex : SV_GroupIndex)
 {
-	const uint BlockHeaderIndex  = GroupId.x;
-	const uint LocalTransformIndex = GroupIndex;
+	const uint InstanceIndex = GroupId.y;
+	const uint HeaderSize    = uint(sizeof(FScatterTransformsBlockHeader));
 
-	const FScatterTransformsBlockHeader BlockHeader = ScatterBlockHeaderBuffer.Load<FScatterTransformsBlockHeader>(BlockHeaderIndex * uint(sizeof(FScatterTransformsBlockHeader)));
+	const FScatterTransformsBlockHeader BlockHeader = ScatterBlockHeaderBuffer.Load<FScatterTransformsBlockHeader>((HeaderBaseOffset + InstanceIndex) * HeaderSize);
 
-	if (LocalTransformIndex >= BlockHeader.BlockTransformCount)
+	const uint TransformIndex = GroupId.x * BONES_PER_GROUP + GroupIndex;
+	if (TransformIndex >= BlockHeader.MaxBoneTransformCount)
 	{
 		return;
 	}
 
-	const uint TransformIndex = LocalTransformIndex + (BlockHeader.BlockLocalIndex * BONES_PER_GROUP);
-
 	const uint CompactBoneIndex = BoneMapBuffer.Load<uint>(BlockHeader.BoneMapBufferOffset + TransformIndex * uint(sizeof(uint)));
 
 	if (CompactBoneIndex != INVALID_BONE_INDEX)
@@ -328,18 +325,19 @@ FBoneHierarchyData LoadBoneHierarchyData(ByteAddressBuffer Buffer, uint BaseOffs
 
 FBoneTransformWithUniformScale ConcatenateToGlobal(
 	uint BoneIndex,
+	uint SharedTransformOffset,
 	uint BoneMapBufferOffset,
 	uint InvGlobalRefPoseTransformBufferOffset,
 	out uint OutTransformIndex)
 {
-	FBoneTransformWithUniformScale GlobalTransform = SharedTransform[BoneIndex];
+	FBoneTransformWithUniformScale GlobalTransform = SharedTransform[SharedTransformOffset + BoneIndex];
 	
 	FBoneHierarchyData HierarchyData = LoadBoneHierarchyData(BoneMapBuffer, BoneMapBufferOffset, BoneIndex);
 	OutTransformIndex = HierarchyData.TransformIndex;
 
 	while (HierarchyData.ParentIndex != INVALID_BONE_INDEX)
 	{
-		GlobalTransform = MultiplyBoneTransformWithUniformScale(GlobalTransform, SharedTransform[HierarchyData.ParentIndex]);
+		GlobalTransform = MultiplyBoneTransformWithUniformScale(GlobalTransform, SharedTransform[SharedTransformOffset + HierarchyData.ParentIndex]);
 		HierarchyData = LoadBoneHierarchyData(BoneMapBuffer, BoneMapBufferOffset, HierarchyData.ParentIndex);
 	}
 
@@ -350,14 +348,26 @@ FBoneTransformWithUniformScale ConcatenateToGlobal(
 [numthreads(BONES_PER_GROUP, 1, 1)]
 void EvaluateConcatenateScatterCS(uint3 GroupId : SV_GroupID, uint GroupIndex : SV_GroupIndex)
 {
-	const uint BlockIndex = GroupId.x;
-	const uint BoneIndex = GroupIndex;
+	const uint HeaderSize      = uint(sizeof(FEvaluateConcatenateScatterBlockHeader));
+	const uint SlotsPerInstance = EvaluateBlockHeaderBuffer.Load<FEvaluateConcatenateScatterBlockHeader>(HeaderBaseOffset * HeaderSize).SlotsPerInstance;
+	const uint InstancesPerGroup = BONES_PER_GROUP / SlotsPerInstance;
+
+	const uint InstanceInGroup       = GroupIndex / SlotsPerInstance;
+	const uint BoneIndex             = GroupIndex % SlotsPerInstance;
+	const uint HeaderIndex           = GroupId.x * InstancesPerGroup + InstanceInGroup;
+	const uint SharedTransformOffset = InstanceInGroup * SlotsPerInstance;
+
+	if (HeaderIndex >= NumTotalHeaders)
+	{
+		return;
+	}
 
-	FEvaluateConcatenateScatterBlockHeader BlockHeader = EvaluateBlockHeaderBuffer.Load<FEvaluateConcatenateScatterBlockHeader>(BlockIndex * uint(sizeof(FEvaluateConcatenateScatterBlockHeader)));
+	const FEvaluateConcatenateScatterBlockHeader BlockHeader =
+		EvaluateBlockHeaderBuffer.Load<FEvaluateConcatenateScatterBlockHeader>((HeaderBaseOffset + HeaderIndex) * HeaderSize);
 
 	if (BoneIndex < BlockHeader.NumCompactBones)
 	{
-		SharedTransform[BoneIndex] = EvaluateAndRetargetBone(
+		SharedTransform[SharedTransformOffset + BoneIndex] = EvaluateAndRetargetBone(
 			BoneIndex,
 			BlockHeader.SequenceTransformBufferOffset,
 			BlockHeader.RetargetingTransformBufferOffset,
@@ -373,7 +383,7 @@ void EvaluateConcatenateScatterCS(uint3 GroupId : SV_GroupID, uint GroupIndex :
 	if (BoneIndex < BlockHeader.NumCompactBones)
 	{
 		uint TransformIndex;
-		FBoneTransformWithUniformScale GlobalTransform = ConcatenateToGlobal(BoneIndex, BlockHeader.BoneMapBufferOffset, BlockHeader.InvGlobalRefPoseTransformBufferOffset, TransformIndex);
+		FBoneTransformWithUniformScale GlobalTransform = ConcatenateToGlobal(BoneIndex, SharedTransformOffset, BlockHeader.BoneMapBufferOffset, BlockHeader.InvGlobalRefPoseTransformBufferOffset, TransformIndex);
 		ScatterGlobalSpaceTransform(BlockHeader.DstTransformBufferOffset1, BlockHeader.DstTransformBufferOffset2, TransformIndex, GlobalTransform);
 	}
 }
@@ -381,10 +391,22 @@ void EvaluateConcatenateScatterCS(uint3 GroupId : SV_GroupID, uint GroupIndex :
 [numthreads(BONES_PER_GROUP, 1, 1)]
 void EvaluateBlendConcatenateScatterCS(uint3 GroupId : SV_GroupID, uint GroupIndex : SV_GroupIndex)
 {
-	const uint BlockIndex = GroupId.x;
-	const uint BoneIndex = GroupIndex;
+	const uint HeaderSize      = uint(sizeof(FEvaluateBlendConcatenateScatterBlockHeader));
+	const uint SlotsPerInstance = EvaluateBlockHeaderBuffer.Load<FEvaluateBlendConcatenateScatterBlockHeader>(HeaderBaseOffset * HeaderSize).SlotsPerInstance;
+	const uint InstancesPerGroup = BONES_PER_GROUP / SlotsPerInstance;
+
+	const uint InstanceInGroup       = GroupIndex / SlotsPerInstance;
+	const uint BoneIndex             = GroupIndex % SlotsPerInstance;
+	const uint HeaderIndex           = GroupId.x * InstancesPerGroup + InstanceInGroup;
+	const uint SharedTransformOffset = InstanceInGroup * SlotsPerInstance;
+
+	if (HeaderIndex >= NumTotalHeaders)
+	{
+		return;
+	}
 
-	FEvaluateBlendConcatenateScatterBlockHeader BlockHeader = EvaluateBlockHeaderBuffer.Load<FEvaluateBlendConcatenateScatterBlockHeader>(BlockIndex * uint(sizeof(FEvaluateBlendConcatenateScatterBlockHeader)));
+	const FEvaluateBlendConcatenateScatterBlockHeader BlockHeader =
+		EvaluateBlockHeaderBuffer.Load<FEvaluateBlendConcatenateScatterBlockHeader>((HeaderBaseOffset + HeaderIndex) * HeaderSize);
 
 	if (BoneIndex < BlockHeader.NumCompactBones)
 	{
@@ -408,7 +430,7 @@ void EvaluateBlendConcatenateScatterCS(uint3 GroupId : SV_GroupID, uint GroupInd
 			BlockHeader.TargetAlpha,
 			BlockHeader.TargetNumAnimBones);
 
-		SharedTransform[BoneIndex] = BlendBoneTransformsWithUniformScale(SourceTransform, TargetTransform, BlockHeader.BlendWeight);
+		SharedTransform[SharedTransformOffset + BoneIndex] = BlendBoneTransformsWithUniformScale(SourceTransform, TargetTransform, BlockHeader.BlendWeight);
 	}
 
 	GroupMemoryBarrierWithGroupSync();
@@ -416,7 +438,7 @@ void EvaluateBlendConcatenateScatterCS(uint3 GroupId : SV_GroupID, uint GroupInd
 	if (BoneIndex < BlockHeader.NumCompactBones)
 	{
 		uint TransformIndex;
-		FBoneTransformWithUniformScale GlobalTransform = ConcatenateToGlobal(BoneIndex, BlockHeader.BoneMapBufferOffset, BlockHeader.InvGlobalRefPoseTransformBufferOffset, TransformIndex);
+		FBoneTransformWithUniformScale GlobalTransform = ConcatenateToGlobal(BoneIndex, SharedTransformOffset, BlockHeader.BoneMapBufferOffset, BlockHeader.InvGlobalRefPoseTransformBufferOffset, TransformIndex);
 		ScatterGlobalSpaceTransform(BlockHeader.DstTransformBufferOffset1, BlockHeader.DstTransformBufferOffset2, TransformIndex, GlobalTransform);
 	}
 }
\ No newline at end of file
diff --git a/Engine/Shaders/Shared/SkinningDefinitions.h b/Engine/Shaders/Shared/SkinningDefinitions.h
--- a/Engine/Shaders/Shared/SkinningDefinitions.h
+++ b/Engine/Shaders/Shared/SkinningDefinitions.h
@@ -55,13 +55,13 @@ struct FEvaluateAnimSequenceBlockHeader
 	uint RetargetingDataBufferOffset;
 	uint DstTransformBufferOffset;
 	uint NumAnimBones;
-	uint BlockLocalIndex      : 16;
-	uint BlockTransformCount  : 16;
+	uint NumCompactBones : 16;
+	uint Padding         : 16;
 };
 SKINNING_STATIC_ASSERT(sizeof(FEvaluateAnimSequenceBlockHeader) == 36, "FEvaluateAnimSequenceBlockHeader must be 36 bytes");
 
 struct FEvaluateAnimSequenceBlendBlockHeader
- {
+{
 	float SourceAlpha;
 	float TargetAlpha;
 	float BlendWeight;
@@ -78,8 +78,8 @@ struct FEvaluateAnimSequenceBlendBlockHeader
 	uint DstTransformBufferOffset;
 	uint SourceNumAnimBones   : 16;
 	uint TargetNumAnimBones   : 16;
-	uint BlockLocalIndex      : 16;
-	uint BlockTransformCount  : 16;
+	uint NumCompactBones      : 16;
+	uint Padding              : 16;
  };
 SKINNING_STATIC_ASSERT(sizeof(FEvaluateAnimSequenceBlendBlockHeader) == 64, "FEvaluateAnimSequenceBlendBlockHeader must be 64 bytes");
 
@@ -95,10 +95,12 @@ struct FEvaluateConcatenateScatterBlockHeader
 	uint BoneMapBufferOffset;
 	uint DstTransformBufferOffset1;
 	uint DstTransformBufferOffset2;
-	uint NumAnimBones    : 16;
-	uint NumCompactBones : 16;
+	uint NumAnimBones       : 16;
+	uint NumCompactBones    : 16;
+	uint SlotsPerInstance   : 16;
+	uint Padding            : 16;
 };
-SKINNING_STATIC_ASSERT(sizeof(FEvaluateConcatenateScatterBlockHeader) == 44, "FEvaluateConcatenateScatterBlockHeader must be 44 bytes");
+SKINNING_STATIC_ASSERT(sizeof(FEvaluateConcatenateScatterBlockHeader) == 48, "FEvaluateConcatenateScatterBlockHeader must be 48 bytes");
 
 struct FEvaluateBlendConcatenateScatterBlockHeader
 {
@@ -122,7 +124,7 @@ struct FEvaluateBlendConcatenateScatterBlockHeader
 	uint SourceNumAnimBones : 16;
 	uint TargetNumAnimBones : 16;
 	uint NumCompactBones    : 16;
-	uint Unused             : 16;
+	uint SlotsPerInstance   : 16;
 };
 SKINNING_STATIC_ASSERT(sizeof(FEvaluateBlendConcatenateScatterBlockHeader) == 76, "FEvaluateBlendConcatenateScatterBlockHeader must be 76 bytes");
 
@@ -153,8 +155,8 @@ struct FConcatenateBlockHeader
 	uint TransformBufferOffset;
 	uint LeafBonePathBufferOffset : 24;
 	uint NumPrefetchBones   : 8;
-	uint BlockLocalIndex    : 16;
-	uint BlockLeafBoneCount : 16;
+	uint NumLeafBones       : 16;
+	uint Padding            : 16;
 };
 SKINNING_STATIC_ASSERT(sizeof(FConcatenateBlockHeader) == 12, "FConcatenateBlockHeader must be 12 bytes");
 
@@ -165,8 +167,8 @@ struct FScatterTransformsBlockHeader
 	uint DstTransformBufferOffset1;
 	uint DstTransformBufferOffset2;
 	uint BoneMapBufferOffset;
-	uint BlockLocalIndex     : 16;
-	uint BlockTransformCount : 16;
+	uint MaxBoneTransformCount : 16;
+	uint Padding               : 16;
 };
 SKINNING_STATIC_ASSERT(sizeof(FScatterTransformsBlockHeader) == 24, "FScatterTransformsBlockHeader must be 24 bytes");
 
diff --git a/Engine/Source/Runtime/Renderer/Private/Skinning/AnimSequenceTransformProvider.cpp b/Engine/Source/Runtime/Renderer/Private/Skinning/AnimSequenceTransformProvider.cpp
--- a/Engine/Source/Runtime/Renderer/Private/Skinning/AnimSequenceTransformProvider.cpp
+++ b/Engine/Source/Runtime/Renderer/Private/Skinning/AnimSequenceTransformProvider.cpp
@@ -183,6 +183,7 @@ public:
 	}
 
 	BEGIN_SHADER_PARAMETER_STRUCT(FParameters, )
+		SHADER_PARAMETER(uint32, HeaderBaseOffset)
 		SHADER_PARAMETER_RDG_BUFFER_SRV(ByteAddressBuffer, RetargetingDataBuffer)
 		SHADER_PARAMETER_RDG_BUFFER_SRV(ByteAddressBuffer, SequenceTransformBuffer)
 		SHADER_PARAMETER_RDG_BUFFER_SRV(ByteAddressBuffer, EvaluateBlockHeaderBuffer)
@@ -214,6 +215,7 @@ public:
 	}
 
 	BEGIN_SHADER_PARAMETER_STRUCT(FParameters, )
+		SHADER_PARAMETER(uint32, HeaderBaseOffset)
 		SHADER_PARAMETER_RDG_BUFFER_SRV(ByteAddressBuffer, RetargetingDataBuffer)
 		SHADER_PARAMETER_RDG_BUFFER_SRV(ByteAddressBuffer, SequenceTransformBuffer)
 		SHADER_PARAMETER_RDG_BUFFER_SRV(ByteAddressBuffer, EvaluateBlockHeaderBuffer)
@@ -245,6 +247,7 @@ public:
 	}
 
 	BEGIN_SHADER_PARAMETER_STRUCT(FParameters, )
+		SHADER_PARAMETER(uint32, HeaderBaseOffset)
 		SHADER_PARAMETER_RDG_BUFFER_SRV(ByteAddressBuffer, LocalSpaceTransformBuffer)
 		SHADER_PARAMETER_RDG_BUFFER_SRV(ByteAddressBuffer, ConcatenateBlockHeaderBuffer)
 		SHADER_PARAMETER_RDG_BUFFER_SRV(ByteAddressBuffer, LeafBonePathBuffer)
@@ -276,6 +279,7 @@ public:
 	}
 
 	BEGIN_SHADER_PARAMETER_STRUCT(FParameters, )
+		SHADER_PARAMETER(uint32, HeaderBaseOffset)
 		SHADER_PARAMETER_RDG_BUFFER_SRV(ByteAddressBuffer, GlobalSpaceTransformBuffer)
 		SHADER_PARAMETER_RDG_BUFFER_SRV(ByteAddressBuffer, SequenceTransformBuffer)
 		SHADER_PARAMETER_RDG_BUFFER_SRV(ByteAddressBuffer, BoneMapBuffer)
@@ -291,7 +295,8 @@ public:
 	DECLARE_GLOBAL_SHADER(FEvaluateConcatenateScatterCS);
 	SHADER_USE_PARAMETER_STRUCT(FEvaluateConcatenateScatterCS, FGlobalShader);
 
-	static constexpr uint32 BonesPerGroup = 32u;
+	class FBonesPerGroupDim : SHADER_PERMUTATION_SPARSE_INT("BONES_PER_GROUP", 64, 128, 256);
+	using FPermutationDomain = TShaderPermutationDomain<FBonesPerGroupDim>;
 
 	static bool ShouldCompilePermutation(const FGlobalShaderPermutationParameters& Parameters)
 	{
@@ -301,13 +306,14 @@ public:
 	static void ModifyCompilationEnvironment(const FGlobalShaderPermutationParameters& Parameters, FShaderCompilerEnvironment& OutEnvironment)
 	{
 		FGlobalShader::ModifyCompilationEnvironment(Parameters, OutEnvironment);
-		OutEnvironment.SetDefine(TEXT("BONES_PER_GROUP"), BonesPerGroup);
 
 		OutEnvironment.CompilerFlags.Add(CFLAG_WarningsAsErrors);
 		OutEnvironment.CompilerFlags.Add(CFLAG_HLSL2021);
 	}
 
 	BEGIN_SHADER_PARAMETER_STRUCT(FParameters, )
+		SHADER_PARAMETER(uint32, NumTotalHeaders)
+		SHADER_PARAMETER(uint32, HeaderBaseOffset)
 		SHADER_PARAMETER_RDG_BUFFER_SRV(ByteAddressBuffer, RetargetingDataBuffer)
 		SHADER_PARAMETER_RDG_BUFFER_SRV(ByteAddressBuffer, SequenceTransformBuffer)
 		SHADER_PARAMETER_RDG_BUFFER_SRV(ByteAddressBuffer, EvaluateBlockHeaderBuffer)
@@ -323,7 +329,8 @@ public:
 	DECLARE_GLOBAL_SHADER(FEvaluateBlendConcatenateScatterCS);
 	SHADER_USE_PARAMETER_STRUCT(FEvaluateBlendConcatenateScatterCS, FGlobalShader);
 
-	static constexpr uint32 BonesPerGroup = FEvaluateConcatenateScatterCS::BonesPerGroup;
+	class FBonesPerGroupDim : SHADER_PERMUTATION_SPARSE_INT("BONES_PER_GROUP", 64, 128, 256);
+	using FPermutationDomain = TShaderPermutationDomain<FBonesPerGroupDim>;
 
 	static bool ShouldCompilePermutation(const FGlobalShaderPermutationParameters& Parameters)
 	{
@@ -333,13 +340,14 @@ public:
 	static void ModifyCompilationEnvironment(const FGlobalShaderPermutationParameters& Parameters, FShaderCompilerEnvironment& OutEnvironment)
 	{
 		FGlobalShader::ModifyCompilationEnvironment(Parameters, OutEnvironment);
-		OutEnvironment.SetDefine(TEXT("BONES_PER_GROUP"), BonesPerGroup);
 
 		OutEnvironment.CompilerFlags.Add(CFLAG_WarningsAsErrors);
 		OutEnvironment.CompilerFlags.Add(CFLAG_HLSL2021);
 	}
 
 	BEGIN_SHADER_PARAMETER_STRUCT(FParameters, )
+		SHADER_PARAMETER(uint32, NumTotalHeaders)
+		SHADER_PARAMETER(uint32, HeaderBaseOffset)
 		SHADER_PARAMETER_RDG_BUFFER_SRV(ByteAddressBuffer, RetargetingDataBuffer)
 		SHADER_PARAMETER_RDG_BUFFER_SRV(ByteAddressBuffer, SequenceTransformBuffer)
 		SHADER_PARAMETER_RDG_BUFFER_SRV(ByteAddressBuffer, EvaluateBlockHeaderBuffer)
@@ -713,7 +721,8 @@ void FAnimSequenceTransformProvider::RegisterProxy(FAnimSequenceTransformProvide
 		BuildBoneTables(LOD.BoneTables, SkinnedAsset, RenderData->LODRenderData[LODIndex]);
 
 		// Single pass mode skips building leaf bone paths in favor of doing everything in one compute dispatch.
-		const bool bSinglePass = LOD.BoneTables.NumCompactBones <= FEvaluateConcatenateScatterCS::BonesPerGroup && (RenderData->HasUnifiedBoneMap() || ExtensionProxy->GetBoneTransformStorageMode() == EBoneTransformStorageMode::Direct);
+		ComputeSinglePassGroupParams(LOD.BoneTables.NumCompactBones, LOD.GroupSize, LOD.InstancesPerGroup, LOD.SlotsPerInstance);
+		const bool bSinglePass = LOD.GroupSize > 0 && (RenderData->HasUnifiedBoneMap() || ExtensionProxy->GetBoneTransformStorageMode() == EBoneTransformStorageMode::Direct);
 		if (!bSinglePass)
 		{
 			BuildLeafBonePathData(LOD.LeafBonePathData, LOD.BoneTables);
@@ -873,28 +882,7 @@ void FAnimSequenceTransformProvider::UnregisterProxy(FAnimSequenceTransformProvi
 	delete ProxyUserData;
 }
 
-template <typename T>
-class TBufferUploadData
-{
-public:
-	TBufferUploadData() = default;
 
-	TBufferUploadData(T* InElements, int32 InMaxElements)
-		: Elements(InElements)
-		, MaxElements(InMaxElements)
-	{}
-
-	T& Add_GetRef()
-	{
-		check(NumElements < MaxElements);
-		return Elements[NumElements++];
-	}
-
-private:
-	T* Elements = nullptr;
-	int32 MaxElements = 0;
-	int32 NumElements = 0;
-};
 
 template <typename T>
 FRDGBuffer* CreateByteAddressBuffer(FRDGBuilder& GraphBuilder, uint32 MaxNumElements, const TCHAR* Name)
@@ -903,12 +891,12 @@ FRDGBuffer* CreateByteAddressBuffer(FRDGBuilder& GraphBuilder, uint32 MaxNumElem
 }
 
 template <typename T>
-TBufferUploadData<T> Lock(FRDGBuilder& GraphBuilder, FRDGBuffer* Buffer)
+TArrayView<T> Lock(FRDGBuilder& GraphBuilder, FRDGBuffer* Buffer)
 {
 	if (Buffer)
 	{
 		GraphBuilder.ConvertToExternalBuffer(Buffer);
-		return TBufferUploadData<T>(reinterpret_cast<T*>(GraphBuilder.RHICmdList.LockBuffer(Buffer->GetRHIUnchecked(), 0, Buffer->Desc.GetSize(), RLM_WriteOnly)), Buffer->Desc.GetSize() / sizeof(T));
+		return TArrayView<T>(reinterpret_cast<T*>(GraphBuilder.RHICmdList.LockBuffer(Buffer->GetRHIUnchecked(), 0, Buffer->Desc.GetSize(), RLM_WriteOnly)), Buffer->Desc.GetSize() / sizeof(T));
 	}
 	return {};
 }
@@ -921,11 +909,81 @@ void Unlock(FRDGBuilder& GraphBuilder, FRDGBuffer* Buffer)
 	}
 }
 
+struct FSinglePassDispatchGroup
+{
+	// Identity (used for TSet equality/hash)
+	uint32 NumCompactBones         = 0;
+
+	// Group parameters
+	uint32 GroupSize               = 0;
+	uint32 InstancesPerGroup       = 1;
+	uint32 SlotsPerInstance        = 0;
+
+	// Header counts (accumulated during counting phase)
+	uint32 NumEvaluateHeaders      = 0;
+	uint32 NumBlendEvaluateHeaders = 0;
+
+	// Offsets into shared header buffers (assigned during buffer creation phase)
+	uint32 EvaluateHeaderBaseOffset      = 0;
+	uint32 BlendEvaluateHeaderBaseOffset = 0;
+
+	// Write cursors (used during fill phase)
+	uint32 NextEvaluateWriteIndex      = 0;
+	uint32 NextBlendEvaluateWriteIndex = 0;
+
+	friend bool operator==(const FSinglePassDispatchGroup& A, const FSinglePassDispatchGroup& B)
+	{
+		return A.NumCompactBones == B.NumCompactBones;
+	}
+
+	friend uint32 GetTypeHash(const FSinglePassDispatchGroup& Group)
+	{
+		return GetTypeHash(Group.NumCompactBones);
+	}
+};
+
+struct FMultiPassDispatchGroup
+{
+	// Identity (used for TSet equality/hash)
+	uint32 NumCompactBones         = 0;
+	uint32 NumLeafBones            = 0;
+	uint32 MaxBoneTransformCount   = 0;
+
+	// Header counts (accumulated during counting phase)
+	uint32 NumEvaluateHeaders      = 0;
+	uint32 NumBlendEvaluateHeaders = 0;
+
+	// Offsets into shared header buffers (assigned during buffer creation phase)
+	uint32 EvaluateHeaderBaseOffset      = 0;
+	uint32 BlendEvaluateHeaderBaseOffset = 0;
+	uint32 ConcatenateHeaderBaseOffset   = 0;
+	uint32 ScatterHeaderBaseOffset       = 0;
+
+	// Write cursors (used during fill phase)
+	uint32 NextEvaluateWriteIndex      = 0;
+	uint32 NextBlendEvaluateWriteIndex = 0;
+	uint32 NextConcatenateWriteIndex   = 0;
+	uint32 NextScatterWriteIndex       = 0;
+
+	friend bool operator==(const FMultiPassDispatchGroup& A, const FMultiPassDispatchGroup& B)
+	{
+		return A.NumCompactBones == B.NumCompactBones
+			&& A.NumLeafBones == B.NumLeafBones
+			&& A.MaxBoneTransformCount == B.MaxBoneTransformCount;
+	}
+
+	friend uint32 GetTypeHash(const FMultiPassDispatchGroup& Group)
+	{
+		return HashCombine(
+			Group.NumCompactBones,
+			HashCombine(Group.NumLeafBones, Group.MaxBoneTransformCount));
+	}
+};
+
 void FAnimSequenceTransformProvider::ProvideGPUTransforms(FSkinningTransformProvider::FProviderContext& Context)
 {
 	TRACE_CPUPROFILER_EVENT_SCOPE(FAnimSequenceTransformProvider::ProvideGPUTransforms);
 	FRDGBuilder& GraphBuilder = Context.GraphBuilder;
-	static const uint32 BonesPerGroup = FEvaluateAnimSequenceCS::BonesPerGroup;
 
 	SequenceTransformAllocator.Consolidate();
 	RetargetingDataAllocator.Consolidate();
@@ -978,12 +1036,9 @@ void FAnimSequenceTransformProvider::ProvideGPUTransforms(FSkinningTransformProv
 
 	uint32 NumTotalSubTracks = 0;
 	uint32 NumTotalIntermediateTransforms = 0;
-	uint32 NumTotalEvaluateConcatenateScatterBlockHeaders = 0;
-	uint32 NumTotalEvaluateBlendConcatenateScatterBlockHeaders = 0;
-	uint32 NumTotalEvaluateBlockHeaders = 0;
-	uint32 NumTotalEvaluateBlendBlockHeaders = 0;
-	uint32 NumTotalScatterBlockHeaders = 0;
-	uint32 NumTotalConcatenateBlockHeaders = 0;
+
+	auto& SinglePassDispatchGroups = *GraphBuilder.AllocObject<TSet<FSinglePassDispatchGroup>>();
+	auto& MultiPassDispatchGroups  = *GraphBuilder.AllocObject<TSet<FMultiPassDispatchGroup>>();
 
 	for (const FSkinningTransformProvider::FProviderIndirection Indirection : Context.Indirections)
 	{
@@ -994,16 +1049,37 @@ void FAnimSequenceTransformProvider::ProvideGPUTransforms(FSkinningTransformProv
 		FAnimSequenceTrackRenderPool& Tracks = ProxyUserData.Tracks;
 		ProxyUserData.TrackHistories.SetNum(Tracks.GetNumTracks());
 
-		const uint32 NumEvaluateConcatenateBlendBlockHeaders = LOD.bSinglePass ? 1 : 0;
-		const uint32 NumEvaluateBlockHeaders                 = LOD.bSinglePass ? 0 : FMath::DivideAndRoundUp<uint32>(LOD.BoneTables.NumCompactBones, BonesPerGroup);
-		const uint32 NumScatterBlockHeaders                  = LOD.bSinglePass ? 0 : FMath::DivideAndRoundUp<uint32>(ExtensionProxy->GetMaxBoneTransformCount(), BonesPerGroup);
-		const uint32 NumConcatenateBlockHeaders              = LOD.bSinglePass ? 0 : FMath::DivideAndRoundUp<uint32>(LOD.LeafBonePathData.NumLeafBones, BonesPerGroup);
-		const uint32 CurrentTransformSlot                    = Indirection.CurrentTransformOffset < Indirection.PreviousTransformOffset ? 0 : 1;
+		const uint32 CurrentTransformSlot = Indirection.CurrentTransformOffset < Indirection.PreviousTransformOffset ? 0 : 1;
+
+		// Resolve group once per proxy, not per track
+		uint32* NumDispatchGroupEvaluateHeaders = nullptr;
+		uint32* NumDispatchGroupBlendEvaluateHeaders = nullptr;
+
+		if (LOD.bSinglePass)
+		{
+			FSinglePassDispatchGroup LookupKey;
+			LookupKey.NumCompactBones   = LOD.BoneTables.NumCompactBones;
+			LookupKey.GroupSize         = LOD.GroupSize;
+			LookupKey.InstancesPerGroup = LOD.InstancesPerGroup;
+			LookupKey.SlotsPerInstance  = LOD.SlotsPerInstance;
+
+			FSinglePassDispatchGroup& DispatchGroup = SinglePassDispatchGroups.FindOrAdd(LookupKey);
+			NumDispatchGroupEvaluateHeaders      = &DispatchGroup.NumEvaluateHeaders;
+			NumDispatchGroupBlendEvaluateHeaders = &DispatchGroup.NumBlendEvaluateHeaders;
+		}
+		else
+		{
+			FMultiPassDispatchGroup LookupKey;
+			LookupKey.NumCompactBones       = LOD.BoneTables.NumCompactBones;
+			LookupKey.NumLeafBones          = LOD.LeafBonePathData.NumLeafBones;
+			LookupKey.MaxBoneTransformCount = ExtensionProxy->GetMaxBoneTransformCount();
+
+			FMultiPassDispatchGroup& DispatchGroup = MultiPassDispatchGroups.FindOrAdd(LookupKey);
+			NumDispatchGroupEvaluateHeaders      = &DispatchGroup.NumEvaluateHeaders;
+			NumDispatchGroupBlendEvaluateHeaders = &DispatchGroup.NumBlendEvaluateHeaders;
+		}
 
 		uint32 NumSubTracks = 0;
-		uint32 NumEvaluateSubTracks = 0;
-		uint32 NumEvaluateBlendSubTracks = 0;
-		uint32 NumIntermediateTransforms = 0;
 
 		Tracks.EnumerateActiveTracks([&] (int32 TrackIndex)
 		{
@@ -1029,20 +1105,21 @@ void FAnimSequenceTransformProvider::ProvideGPUTransforms(FSkinningTransformProv
 			const int32 NumLocalSubTracks = UpdateResult.bNeedsCurrent + (UpdateRequest.GetUpdateMode() == EPreviousBoneTransformUpdateMode::UpdatePrevious ? 1 : 0);
 
 			NumSubTracks                   += NumLocalSubTracks;
-			NumEvaluateSubTracks           += BlendWeight <= 0.0f ? NumLocalSubTracks : 0;
-			NumEvaluateBlendSubTracks      += BlendWeight >  0.0f ? NumLocalSubTracks : 0;
 			NumTotalIntermediateTransforms += LOD.bSinglePass ? 0 : NumLocalSubTracks * LOD.BoneTables.NumCompactBones;
 
+			if (BlendWeight > 0.0f)
+			{
+				*NumDispatchGroupBlendEvaluateHeaders += NumLocalSubTracks;
+			}
+			else
+			{
+				*NumDispatchGroupEvaluateHeaders += NumLocalSubTracks;
+			}
+
 			TrackUpdateRequests.Emplace(UpdateRequest);
 		});
 
-		NumTotalSubTracks                                   += NumSubTracks;
-		NumTotalEvaluateConcatenateScatterBlockHeaders      += NumEvaluateSubTracks * NumEvaluateConcatenateBlendBlockHeaders;
-		NumTotalEvaluateBlendConcatenateScatterBlockHeaders += NumEvaluateBlendSubTracks * NumEvaluateConcatenateBlendBlockHeaders;
-		NumTotalEvaluateBlockHeaders                        += NumEvaluateSubTracks * NumEvaluateBlockHeaders;
-		NumTotalEvaluateBlendBlockHeaders                   += NumEvaluateBlendSubTracks * NumEvaluateBlockHeaders;
-		NumTotalScatterBlockHeaders                         += NumSubTracks * NumScatterBlockHeaders;
-		NumTotalConcatenateBlockHeaders                     += NumSubTracks * NumConcatenateBlockHeaders;
+		NumTotalSubTracks += NumSubTracks;
 
 		if (!LOD.bUploaded && NumSubTracks > 0)
 		{
@@ -1130,20 +1207,76 @@ void FAnimSequenceTransformProvider::ProvideGPUTransforms(FSkinningTransformProv
 
 	int32& NumTotalFramesRequested = *GraphBuilder.AllocPOD<int32>();
 	NumTotalFramesRequested = 0;
-	
-	FRDGBuffer* EvaluateConcatenateScatterBlockHeaderBuffer      = CreateByteAddressBuffer<UE::HLSL::FEvaluateConcatenateScatterBlockHeader>(GraphBuilder, NumTotalEvaluateConcatenateScatterBlockHeaders, TEXT("AnimSequenceTransformProvider.EvaluateConcatenateScatterBlockHeaders"));
-	FRDGBuffer* EvaluateBlendConcatenateScatterBlockHeaderBuffer = CreateByteAddressBuffer<UE::HLSL::FEvaluateBlendConcatenateScatterBlockHeader>(GraphBuilder, NumTotalEvaluateBlendConcatenateScatterBlockHeaders, TEXT("AnimSequenceTransformProvider.EvaluateBlendConcatenateScatterBlockHeaders"));
-	FRDGBuffer* EvaluateBlockHeaderBuffer                        = CreateByteAddressBuffer<UE::HLSL::FEvaluateAnimSequenceBlockHeader>(GraphBuilder, NumTotalEvaluateBlockHeaders, TEXT("AnimSequenceTransformProvider.EvaluateBlockHeaders"));
-	FRDGBuffer* EvaluateBlendBlockHeaderBuffer                   = CreateByteAddressBuffer<UE::HLSL::FEvaluateAnimSequenceBlendBlockHeader>(GraphBuilder, NumTotalEvaluateBlendBlockHeaders, TEXT("AnimSequenceTransformProvider.EvaluateBlendBlockHeaders"));
-	FRDGBuffer* ScatterBlockHeaderBuffer                         = CreateByteAddressBuffer<UE::HLSL::FScatterTransformsBlockHeader>(GraphBuilder, NumTotalScatterBlockHeaders, TEXT("AnimSequenceTransformProvider.ScatterBlockHeaders"));
-	FRDGBuffer* ConcatenateBlockHeaderBuffer                     = CreateByteAddressBuffer<UE::HLSL::FConcatenateBlockHeader>(GraphBuilder, NumTotalConcatenateBlockHeaders, TEXT("AnimSequenceTransformProvider.ConcatenateBlockHeaders"));
-	
-	TBufferUploadData<UE::HLSL::FEvaluateConcatenateScatterBlockHeader> EvaluateConcatenateScatterBlockHeaders           = Lock<UE::HLSL::FEvaluateConcatenateScatterBlockHeader>(GraphBuilder, EvaluateConcatenateScatterBlockHeaderBuffer);
-	TBufferUploadData<UE::HLSL::FEvaluateBlendConcatenateScatterBlockHeader> EvaluateBlendConcatenateScatterBlockHeaders = Lock<UE::HLSL::FEvaluateBlendConcatenateScatterBlockHeader>(GraphBuilder, EvaluateBlendConcatenateScatterBlockHeaderBuffer);
-	TBufferUploadData<UE::HLSL::FEvaluateAnimSequenceBlockHeader> EvaluateBlockHeaders                                   = Lock<UE::HLSL::FEvaluateAnimSequenceBlockHeader>(GraphBuilder, EvaluateBlockHeaderBuffer);
-	TBufferUploadData<UE::HLSL::FEvaluateAnimSequenceBlendBlockHeader> EvaluateBlendBlockHeaders                         = Lock<UE::HLSL::FEvaluateAnimSequenceBlendBlockHeader>(GraphBuilder, EvaluateBlendBlockHeaderBuffer);
-	TBufferUploadData<UE::HLSL::FScatterTransformsBlockHeader> ScatterBlockHeaders                                       = Lock<UE::HLSL::FScatterTransformsBlockHeader>(GraphBuilder, ScatterBlockHeaderBuffer);
-	TBufferUploadData<UE::HLSL::FConcatenateBlockHeader> ConcatenateBlockHeaders                                         = Lock<UE::HLSL::FConcatenateBlockHeader>(GraphBuilder, ConcatenateBlockHeaderBuffer);
+
+	// Accumulate totals and assign per-group base offsets into shared buffers
+	uint32 NumTotalSinglePassEvaluateHeaders      = 0;
+	uint32 NumTotalSinglePassBlendEvaluateHeaders = 0;
+	uint32 NumTotalMultiPassEvaluateHeaders       = 0;
+	uint32 NumTotalMultiPassBlendEvaluateHeaders  = 0;
+	uint32 NumTotalConcatenateHeaders             = 0;
+	uint32 NumTotalScatterHeaders                 = 0;
+
+	for (FSinglePassDispatchGroup& Group : SinglePassDispatchGroups)
+	{
+		Group.EvaluateHeaderBaseOffset      = NumTotalSinglePassEvaluateHeaders;
+		Group.BlendEvaluateHeaderBaseOffset = NumTotalSinglePassBlendEvaluateHeaders;
+		NumTotalSinglePassEvaluateHeaders      += Group.NumEvaluateHeaders;
+		NumTotalSinglePassBlendEvaluateHeaders += Group.NumBlendEvaluateHeaders;
+	}
+
+	for (FMultiPassDispatchGroup& Group : MultiPassDispatchGroups)
+	{
+		const uint32 NumGroupTotalHeaders = Group.NumEvaluateHeaders + Group.NumBlendEvaluateHeaders;
+		Group.EvaluateHeaderBaseOffset      = NumTotalMultiPassEvaluateHeaders;
+		Group.BlendEvaluateHeaderBaseOffset = NumTotalMultiPassBlendEvaluateHeaders;
+		Group.ConcatenateHeaderBaseOffset   = NumTotalConcatenateHeaders;
+		Group.ScatterHeaderBaseOffset       = NumTotalScatterHeaders;
+		NumTotalMultiPassEvaluateHeaders      += Group.NumEvaluateHeaders;
+		NumTotalMultiPassBlendEvaluateHeaders += Group.NumBlendEvaluateHeaders;
+		NumTotalConcatenateHeaders            += NumGroupTotalHeaders;
+		NumTotalScatterHeaders                += NumGroupTotalHeaders;
+	}
+
+	// Create shared buffers and lock
+	FRDGBuffer* SinglePassEvaluateHeaderBuffer      = CreateByteAddressBuffer<UE::HLSL::FEvaluateConcatenateScatterBlockHeader>(GraphBuilder, NumTotalSinglePassEvaluateHeaders, TEXT("AnimSequenceTransformProvider.SinglePassEvaluateHeaders"));
+	FRDGBuffer* SinglePassBlendEvaluateHeaderBuffer = CreateByteAddressBuffer<UE::HLSL::FEvaluateBlendConcatenateScatterBlockHeader>(GraphBuilder, NumTotalSinglePassBlendEvaluateHeaders, TEXT("AnimSequenceTransformProvider.SinglePassBlendEvaluateHeaders"));
+	FRDGBuffer* MultiPassEvaluateHeaderBuffer       = CreateByteAddressBuffer<UE::HLSL::FEvaluateAnimSequenceBlockHeader>(GraphBuilder, NumTotalMultiPassEvaluateHeaders, TEXT("AnimSequenceTransformProvider.EvaluateHeaders"));
+	FRDGBuffer* MultiPassBlendEvaluateHeaderBuffer  = CreateByteAddressBuffer<UE::HLSL::FEvaluateAnimSequenceBlendBlockHeader>(GraphBuilder, NumTotalMultiPassBlendEvaluateHeaders, TEXT("AnimSequenceTransformProvider.BlendEvaluateHeaders"));
+	FRDGBuffer* ConcatenateHeaderBuffer             = CreateByteAddressBuffer<UE::HLSL::FConcatenateBlockHeader>(GraphBuilder, NumTotalConcatenateHeaders, TEXT("AnimSequenceTransformProvider.ConcatenateHeaders"));
+	FRDGBuffer* ScatterHeaderBuffer                 = CreateByteAddressBuffer<UE::HLSL::FScatterTransformsBlockHeader>(GraphBuilder, NumTotalScatterHeaders, TEXT("AnimSequenceTransformProvider.ScatterHeaders"));
+
+	TArrayView<UE::HLSL::FEvaluateConcatenateScatterBlockHeader>      SinglePassEvaluateUpload      = Lock<UE::HLSL::FEvaluateConcatenateScatterBlockHeader>(GraphBuilder, SinglePassEvaluateHeaderBuffer);
+	TArrayView<UE::HLSL::FEvaluateBlendConcatenateScatterBlockHeader> SinglePassBlendEvaluateUpload = Lock<UE::HLSL::FEvaluateBlendConcatenateScatterBlockHeader>(GraphBuilder, SinglePassBlendEvaluateHeaderBuffer);
+	TArrayView<UE::HLSL::FEvaluateAnimSequenceBlockHeader>            MultiPassEvaluateUpload       = Lock<UE::HLSL::FEvaluateAnimSequenceBlockHeader>(GraphBuilder, MultiPassEvaluateHeaderBuffer);
+	TArrayView<UE::HLSL::FEvaluateAnimSequenceBlendBlockHeader>       MultiPassBlendEvaluateUpload  = Lock<UE::HLSL::FEvaluateAnimSequenceBlendBlockHeader>(GraphBuilder, MultiPassBlendEvaluateHeaderBuffer);
+	TArrayView<UE::HLSL::FConcatenateBlockHeader>                     ConcatenateUpload             = Lock<UE::HLSL::FConcatenateBlockHeader>(GraphBuilder, ConcatenateHeaderBuffer);
+	TArrayView<UE::HLSL::FScatterTransformsBlockHeader>               ScatterUpload                 = Lock<UE::HLSL::FScatterTransformsBlockHeader>(GraphBuilder, ScatterHeaderBuffer);
+
+	// Assign GroupBufferIndex on each LOD so the fill phase can index directly
+	for (const FSkinningTransformProvider::FProviderIndirection Indirection : Context.Indirections)
+	{
+		FInstancedSkinningSceneExtensionProxy* ExtensionProxy = static_cast<FInstancedSkinningSceneExtensionProxy*>(Context.Proxies[Indirection.Index]);
+		FProxyUserData& ProxyUserData = static_cast<FProxyUserData&>(ExtensionProxy->GetUserDataChecked());
+		const uint32 LODIndex = ExtensionProxy->GetLOD();
+		FProxyUserData::FLOD& LOD = ProxyUserData.LODs[LODIndex];
+
+		if (LOD.bSinglePass)
+		{
+			FSinglePassDispatchGroup LookupKey;
+			LookupKey.NumCompactBones = LOD.BoneTables.NumCompactBones;
+			LOD.DispatchGroup = SinglePassDispatchGroups.Find(LookupKey);
+			check(LOD.DispatchGroup != nullptr);
+		}
+		else
+		{
+			FMultiPassDispatchGroup LookupKey;
+			LookupKey.NumCompactBones       = LOD.BoneTables.NumCompactBones;
+			LookupKey.NumLeafBones          = LOD.LeafBonePathData.NumLeafBones;
+			LookupKey.MaxBoneTransformCount = ExtensionProxy->GetMaxBoneTransformCount();
+			LOD.DispatchGroup = MultiPassDispatchGroups.Find(LookupKey);
+			check(LOD.DispatchGroup != nullptr);
+		}
+	}
 
 	UE::Tasks::FTask ProcessTask = GraphBuilder.AddSetupTask(
 		[
@@ -1154,27 +1287,27 @@ void FAnimSequenceTransformProvider::ProvideGPUTransforms(FSkinningTransformProv
 			, Proxies = Context.Proxies
 			, &NumTotalFramesRequested
 			, TrackUpdateRequests = MoveTemp(TrackUpdateRequests)
-			, EvaluateConcatenateScatterBlockHeaders
-			, EvaluateBlendConcatenateScatterBlockHeaders
-			, EvaluateBlockHeaders
-			, EvaluateBlendBlockHeaders
-			, ScatterBlockHeaders
-			, ConcatenateBlockHeaders
+			, SinglePassEvaluateUpload
+			, SinglePassBlendEvaluateUpload
+			, MultiPassEvaluateUpload
+			, MultiPassBlendEvaluateUpload
+			, ConcatenateUpload
+			, ScatterUpload
 		] () mutable
 	{
 		TRACE_CPUPROFILER_EVENT_SCOPE(FAnimSequenceTransformProvider::EvaluateAnimSequences);
 
+		static const uint32 MultiPassBonesPerGroup = FEvaluateAnimSequenceCS::BonesPerGroup;
 		uint32 IntermediateTransformOffset = 0;
 		uint32 NumActiveTracks = 0;
-		uint32 NumEvaluateConcatenateScatterBlocks = 0;
-		uint32 NumEvaluateBlendConcatenateScatterBlocks = 0;
-		uint32 NumEvaluateBlocks = 0;
-		uint32 NumEvaluateBlendBlocks = 0;
+		uint32 NumSinglePassBlocks = 0;
+		uint32 NumSinglePassBlendBlocks = 0;
+		uint32 NumMultiPassBlocks = 0;
+		uint32 NumMultiPassBlendBlocks = 0;
 		uint32 NumCurrentTransformUpdates = 0;
 		uint32 NumPreviousTransformUpdates = 0;
 
 		const auto EvaluateAnimSequence = [&](
-			const FSkinningTransformProvider::FProviderIndirection& Indirection,
 			const FProxyUserData::FLOD& LOD,
 			const FMapping& SourceMapping,
 			const FMapping& TargetMapping,
@@ -1220,9 +1353,11 @@ void FAnimSequenceTransformProvider::ProvideGPUTransforms(FSkinningTransformProv
 
 			if (LOD.bSinglePass)
 			{
+				FSinglePassDispatchGroup& Group = *static_cast<FSinglePassDispatchGroup*>(LOD.DispatchGroup);
+
 				if (BlendWeight > 0.0f)
 				{
-					UE::HLSL::FEvaluateBlendConcatenateScatterBlockHeader& BlockHeader = EvaluateBlendConcatenateScatterBlockHeaders.Add_GetRef();
+					UE::HLSL::FEvaluateBlendConcatenateScatterBlockHeader& BlockHeader = SinglePassBlendEvaluateUpload[Group.BlendEvaluateHeaderBaseOffset + Group.NextBlendEvaluateWriteIndex++];
 
 					BlockHeader.SourceAlpha                            = SourceAlpha;
 					BlockHeader.TargetAlpha                            = TargetAlpha;
@@ -1244,12 +1379,13 @@ void FAnimSequenceTransformProvider::ProvideGPUTransforms(FSkinningTransformProv
 					BlockHeader.SourceNumAnimBones                     = SourceSequence.NumAnimBones;
 					BlockHeader.TargetNumAnimBones                     = TargetSequence.NumAnimBones;
 					BlockHeader.NumCompactBones                        = LOD.BoneTables.NumCompactBones;
+					BlockHeader.SlotsPerInstance                       = LOD.SlotsPerInstance;
 
-					NumEvaluateBlendConcatenateScatterBlocks++;
+					NumSinglePassBlendBlocks++;
 				}
 				else
 				{
-					UE::HLSL::FEvaluateConcatenateScatterBlockHeader& BlockHeader = EvaluateConcatenateScatterBlockHeaders.Add_GetRef();
+					UE::HLSL::FEvaluateConcatenateScatterBlockHeader& BlockHeader = SinglePassEvaluateUpload[Group.EvaluateHeaderBaseOffset + Group.NextEvaluateWriteIndex++];
 
 					BlockHeader.Alpha                                 = SourceAlpha;
 					BlockHeader.KeyIndex0                             = SourceKeyIndex0;
@@ -1263,104 +1399,82 @@ void FAnimSequenceTransformProvider::ProvideGPUTransforms(FSkinningTransformProv
 					BlockHeader.DstTransformBufferOffset2             = DstTransformBufferOffset2;
 					BlockHeader.NumAnimBones                          = SourceSequence.NumAnimBones;
 					BlockHeader.NumCompactBones                       = LOD.BoneTables.NumCompactBones;
+					BlockHeader.SlotsPerInstance                      = LOD.SlotsPerInstance;
+					BlockHeader.Padding                               = 0;
 
-					NumEvaluateConcatenateScatterBlocks++;
+					NumSinglePassBlocks++;
 				}
 			}
 			else
 			{
-				const auto GetBlockTransformCount = [](uint32 BlockIndex, uint32 TotalCount, uint32 BlockSize) -> uint32
-				{
-					return FMath::Min(BlockSize, TotalCount - BlockIndex * BlockSize);
-				};
+				FMultiPassDispatchGroup& Group = *static_cast<FMultiPassDispatchGroup*>(LOD.DispatchGroup);
 
-				// Evaluation Pass
+				// Evaluate: single header per instance
+				if (BlendWeight > 0.0f)
 				{
-					const uint32 NumBlocks = FMath::DivideAndRoundUp<uint32>(TargetMapping.NumCompactBones, BonesPerGroup);
-
-					if (BlendWeight > 0.0f)
-					{
-						for (uint32 BlockIndex = 0; BlockIndex < NumBlocks; ++BlockIndex)
-						{
-							UE::HLSL::FEvaluateAnimSequenceBlendBlockHeader& BlockHeader = EvaluateBlendBlockHeaders.Add_GetRef();
-
-							BlockHeader.SourceAlpha                            = SourceAlpha;
-							BlockHeader.TargetAlpha                            = TargetAlpha;
-							BlockHeader.BlendWeight                            = BlendWeight;
-							BlockHeader.SourceKeyIndex0                        = SourceKeyIndex0;
-							BlockHeader.SourceKeyIndex1                        = SourceKeyIndex1;
-							BlockHeader.TargetKeyIndex0                        = TargetKeyIndex0;
-							BlockHeader.TargetKeyIndex1                        = TargetKeyIndex1;
-							BlockHeader.SourceSequenceTransformBufferOffset    = SourceSequenceTransformBufferOffset;
-							BlockHeader.TargetSequenceTransformBufferOffset    = TargetSequenceTransformBufferOffset;
-							BlockHeader.TargetRetargetingTransformBufferOffset = TargetRetargetingTransformBufferOffset;
-							BlockHeader.SourceRetargetingTransformBufferOffset = SourceRetargetingTransformBufferOffset;
-							BlockHeader.TargetRetargetingDataBufferOffset      = TargetRetargetingDataBufferOffset;
-							BlockHeader.SourceRetargetingDataBufferOffset      = SourceRetargetingDataBufferOffset;
-							BlockHeader.SourceNumAnimBones                     = SourceSequence.NumAnimBones;
-							BlockHeader.TargetNumAnimBones                     = TargetSequence.NumAnimBones;
-							BlockHeader.DstTransformBufferOffset               = IntermediateTransformOffset;
-							BlockHeader.BlockLocalIndex                        = BlockIndex;
-							BlockHeader.BlockTransformCount                    = GetBlockTransformCount(BlockIndex, TargetMapping.NumCompactBones, BonesPerGroup);
-						}
+					UE::HLSL::FEvaluateAnimSequenceBlendBlockHeader& BlockHeader = MultiPassBlendEvaluateUpload[Group.BlendEvaluateHeaderBaseOffset + Group.NextBlendEvaluateWriteIndex++];
 
-						NumEvaluateBlendBlocks += NumBlocks;
-					}
-					else
-					{
-						for (uint32 BlockIndex = 0; BlockIndex < NumBlocks; ++BlockIndex)
-						{
-							UE::HLSL::FEvaluateAnimSequenceBlockHeader& BlockHeader = EvaluateBlockHeaders.Add_GetRef();
-
-							BlockHeader.SequenceTransformBufferOffset    = SourceSequenceTransformBufferOffset;
-							BlockHeader.RetargetingTransformBufferOffset = SourceRetargetingTransformBufferOffset;
-							BlockHeader.RetargetingDataBufferOffset      = SourceRetargetingDataBufferOffset;
-							BlockHeader.DstTransformBufferOffset         = IntermediateTransformOffset;
-							BlockHeader.NumAnimBones                     = SourceSequence.NumAnimBones;
-							BlockHeader.KeyIndex0                        = SourceKeyIndex0;
-							BlockHeader.KeyIndex1                        = SourceKeyIndex1;
-							BlockHeader.Alpha                            = SourceAlpha;
-							BlockHeader.BlockLocalIndex                  = BlockIndex;
-							BlockHeader.BlockTransformCount              = GetBlockTransformCount(BlockIndex, TargetMapping.NumCompactBones, BonesPerGroup);
-						}
+					BlockHeader.SourceAlpha                            = SourceAlpha;
+					BlockHeader.TargetAlpha                            = TargetAlpha;
+					BlockHeader.BlendWeight                            = BlendWeight;
+					BlockHeader.SourceKeyIndex0                        = SourceKeyIndex0;
+					BlockHeader.SourceKeyIndex1                        = SourceKeyIndex1;
+					BlockHeader.TargetKeyIndex0                        = TargetKeyIndex0;
+					BlockHeader.TargetKeyIndex1                        = TargetKeyIndex1;
+					BlockHeader.SourceSequenceTransformBufferOffset    = SourceSequenceTransformBufferOffset;
+					BlockHeader.TargetSequenceTransformBufferOffset    = TargetSequenceTransformBufferOffset;
+					BlockHeader.TargetRetargetingTransformBufferOffset = TargetRetargetingTransformBufferOffset;
+					BlockHeader.SourceRetargetingTransformBufferOffset = SourceRetargetingTransformBufferOffset;
+					BlockHeader.TargetRetargetingDataBufferOffset      = TargetRetargetingDataBufferOffset;
+					BlockHeader.SourceRetargetingDataBufferOffset      = SourceRetargetingDataBufferOffset;
+					BlockHeader.SourceNumAnimBones                     = SourceSequence.NumAnimBones;
+					BlockHeader.TargetNumAnimBones                     = TargetSequence.NumAnimBones;
+					BlockHeader.DstTransformBufferOffset               = IntermediateTransformOffset;
+					BlockHeader.NumCompactBones                        = TargetMapping.NumCompactBones;
+					BlockHeader.Padding                                = 0;
 
-						NumEvaluateBlocks += NumBlocks;
-					}
+					NumMultiPassBlendBlocks++;
 				}
-
-				// Concatenate Local to Global Pass
+				else
 				{
-					const uint32 NumLeafBones = LOD.LeafBonePathData.NumLeafBones;
-					const uint32 NumBlocks    = FMath::DivideAndRoundUp(NumLeafBones, BonesPerGroup);
+					UE::HLSL::FEvaluateAnimSequenceBlockHeader& BlockHeader = MultiPassEvaluateUpload[Group.EvaluateHeaderBaseOffset + Group.NextEvaluateWriteIndex++];
+
+					BlockHeader.SequenceTransformBufferOffset    = SourceSequenceTransformBufferOffset;
+					BlockHeader.RetargetingTransformBufferOffset = SourceRetargetingTransformBufferOffset;
+					BlockHeader.RetargetingDataBufferOffset      = SourceRetargetingDataBufferOffset;
+					BlockHeader.DstTransformBufferOffset         = IntermediateTransformOffset;
+					BlockHeader.NumAnimBones                     = SourceSequence.NumAnimBones;
+					BlockHeader.KeyIndex0                        = SourceKeyIndex0;
+					BlockHeader.KeyIndex1                        = SourceKeyIndex1;
+					BlockHeader.Alpha                            = SourceAlpha;
+					BlockHeader.NumCompactBones                  = TargetMapping.NumCompactBones;
+					BlockHeader.Padding                          = 0;
+
+					NumMultiPassBlocks++;
+				}
 
-					for (uint32 BlockIndex = 0; BlockIndex < NumBlocks; ++BlockIndex)
-					{
-						UE::HLSL::FConcatenateBlockHeader& BlockHeader = ConcatenateBlockHeaders.Add_GetRef();
+				// Concatenate: single header per instance
+				{
+					UE::HLSL::FConcatenateBlockHeader& ConcatHeader = ConcatenateUpload[Group.ConcatenateHeaderBaseOffset + Group.NextConcatenateWriteIndex++];
 
-						BlockHeader.LeafBonePathBufferOffset = LeafBonePathBufferOffset;
-						BlockHeader.TransformBufferOffset    = IntermediateTransformOffset;
-						BlockHeader.NumPrefetchBones         = FMath::Min<uint16>(LOD.BoneTables.NumCompactBones, BonesPerGroup);
-						BlockHeader.BlockLocalIndex          = BlockIndex;
-						BlockHeader.BlockLeafBoneCount       = GetBlockTransformCount(BlockIndex, NumLeafBones, BonesPerGroup);
-					}
+					ConcatHeader.LeafBonePathBufferOffset = LeafBonePathBufferOffset;
+					ConcatHeader.TransformBufferOffset    = IntermediateTransformOffset;
+					ConcatHeader.NumPrefetchBones         = FMath::Min<uint32>(LOD.BoneTables.NumCompactBones, MultiPassBonesPerGroup);
+					ConcatHeader.NumLeafBones             = LOD.LeafBonePathData.NumLeafBones;
+					ConcatHeader.Padding                  = 0;
 				}
 
-				// Scatter Pass
+				// Scatter: single header per instance
 				{
-					const uint32 NumBlocks = FMath::DivideAndRoundUp(MaxBoneTransformCount, BonesPerGroup);
-
-					for (uint32 BlockIndex = 0; BlockIndex < NumBlocks; ++BlockIndex)
-					{
-						UE::HLSL::FScatterTransformsBlockHeader& BlockHeader = ScatterBlockHeaders.Add_GetRef();
-
-						BlockHeader.InvGlobalRefPoseTransformBufferOffset = InvGlobalRefPoseTransformBufferOffset;
-						BlockHeader.SrcTransformBufferOffset              = IntermediateTransformOffset;
-						BlockHeader.DstTransformBufferOffset1             = DstTransformBufferOffset1;
-						BlockHeader.DstTransformBufferOffset2             = DstTransformBufferOffset2;
-						BlockHeader.BoneMapBufferOffset                   = BoneMapBufferOffset;
-						BlockHeader.BlockLocalIndex                       = BlockIndex;
-						BlockHeader.BlockTransformCount                   = GetBlockTransformCount(BlockIndex, MaxBoneTransformCount, BonesPerGroup);
-					}
+					UE::HLSL::FScatterTransformsBlockHeader& ScatterHeader = ScatterUpload[Group.ScatterHeaderBaseOffset + Group.NextScatterWriteIndex++];
+
+					ScatterHeader.InvGlobalRefPoseTransformBufferOffset = InvGlobalRefPoseTransformBufferOffset;
+					ScatterHeader.SrcTransformBufferOffset              = IntermediateTransformOffset;
+					ScatterHeader.DstTransformBufferOffset1             = DstTransformBufferOffset1;
+					ScatterHeader.DstTransformBufferOffset2             = DstTransformBufferOffset2;
+					ScatterHeader.BoneMapBufferOffset                   = BoneMapBufferOffset;
+					ScatterHeader.MaxBoneTransformCount                 = MaxBoneTransformCount;
+					ScatterHeader.Padding                               = 0;
 				}
 
 				IntermediateTransformOffset += SourceMapping.NumCompactBones * sizeof(FBoneTransformWithScale);
@@ -1372,9 +1486,8 @@ void FAnimSequenceTransformProvider::ProvideGPUTransforms(FSkinningTransformProv
 		for (const FSkinningTransformProvider::FProviderIndirection Indirection : Indirections)
 		{
 			FInstancedSkinningSceneExtensionProxy* ExtensionProxy = static_cast<FInstancedSkinningSceneExtensionProxy*>(Proxies[Indirection.Index]);
-			const uint32 LODIndex = ExtensionProxy->GetLOD();
-
 			FProxyUserData& ProxyUserData = static_cast<FProxyUserData&>(ExtensionProxy->GetUserDataChecked());
+			const uint32 LODIndex = ExtensionProxy->GetLOD();
 			FProxyUserData::FLOD& LOD = ProxyUserData.LODs[LODIndex];
 			FAnimSequenceTrackRenderPool& Tracks = ProxyUserData.Tracks;
 
@@ -1400,7 +1513,6 @@ void FAnimSequenceTransformProvider::ProvideGPUTransforms(FSkinningTransformProv
 					uint32 OutputTransformOffset2 = UpdateRequest.GetUpdateMode() == EPreviousBoneTransformUpdateMode::DuplicateCurrentToPrevious ? TrackPreviousTransformOffset : TrackCurrentTransformOffset;
 
 					EvaluateAnimSequence(
-						Indirection,
 						LOD,
 						SourceTrack.Mapping, TargetTrack.Mapping,
 						SourceTrack.Sequence, TargetTrack.Sequence,
@@ -1415,7 +1527,6 @@ void FAnimSequenceTransformProvider::ProvideGPUTransforms(FSkinningTransformProv
 				if (UpdateRequest.GetUpdateMode() == EPreviousBoneTransformUpdateMode::UpdatePrevious)
 				{
 					EvaluateAnimSequence(
-						Indirection,
 						LOD,
 						SourceTrack.Mapping, TargetTrack.Mapping,
 						SourceTrack.Sequence, TargetTrack.Sequence,
@@ -1433,20 +1544,20 @@ void FAnimSequenceTransformProvider::ProvideGPUTransforms(FSkinningTransformProv
 		}
 
 		INC_DWORD_STAT_BY(STAT_AnimSequenceTransformProvider_ActiveTracks, NumActiveTracks);
-		INC_DWORD_STAT_BY(STAT_AnimSequenceTransformProvider_EvaluateSinglePassBlocks, NumEvaluateConcatenateScatterBlocks);
-		INC_DWORD_STAT_BY(STAT_AnimSequenceTransformProvider_EvaluateSinglePassBlendBlocks, NumEvaluateBlendConcatenateScatterBlocks);
-		INC_DWORD_STAT_BY(STAT_AnimSequenceTransformProvider_EvaluateBlocks, NumEvaluateBlocks);
-		INC_DWORD_STAT_BY(STAT_AnimSequenceTransformProvider_EvaluateBlendBlocks, NumEvaluateBlendBlocks);
+		INC_DWORD_STAT_BY(STAT_AnimSequenceTransformProvider_EvaluateSinglePassBlocks, NumSinglePassBlocks);
+		INC_DWORD_STAT_BY(STAT_AnimSequenceTransformProvider_EvaluateSinglePassBlendBlocks, NumSinglePassBlendBlocks);
+		INC_DWORD_STAT_BY(STAT_AnimSequenceTransformProvider_EvaluateBlocks, NumMultiPassBlocks);
+		INC_DWORD_STAT_BY(STAT_AnimSequenceTransformProvider_EvaluateBlendBlocks, NumMultiPassBlendBlocks);
 		INC_DWORD_STAT_BY(STAT_AnimSequenceTransformProvider_CurrentTransformUpdates, NumCurrentTransformUpdates);
 		INC_DWORD_STAT_BY(STAT_AnimSequenceTransformProvider_PreviousTransformUpdates, NumPreviousTransformUpdates);
 	});
 
-	Unlock(GraphBuilder, EvaluateConcatenateScatterBlockHeaderBuffer);
-	Unlock(GraphBuilder, EvaluateBlendConcatenateScatterBlockHeaderBuffer);
-	Unlock(GraphBuilder, EvaluateBlockHeaderBuffer);
-	Unlock(GraphBuilder, EvaluateBlendBlockHeaderBuffer);
-	Unlock(GraphBuilder, ScatterBlockHeaderBuffer);
-	Unlock(GraphBuilder, ConcatenateBlockHeaderBuffer);
+	Unlock(GraphBuilder, SinglePassEvaluateHeaderBuffer);
+	Unlock(GraphBuilder, SinglePassBlendEvaluateHeaderBuffer);
+	Unlock(GraphBuilder, MultiPassEvaluateHeaderBuffer);
+	Unlock(GraphBuilder, MultiPassBlendEvaluateHeaderBuffer);
+	Unlock(GraphBuilder, ConcatenateHeaderBuffer);
+	Unlock(GraphBuilder, ScatterHeaderBuffer);
 
 	FRDGBuffer* SequenceTransformBuffer = SequenceTransformPersistentBuffer.ResizeBufferIfNeeded(GraphBuilder, SequenceTransformAllocator.GetMaxSize());
 	FRDGBuffer* RetargetingDataBuffer = RetargetingDataPersistentBuffer.ResizeBufferIfNeeded(GraphBuilder, RetargetingDataAllocator.GetMaxSize());
@@ -1758,125 +1869,193 @@ void FAnimSequenceTransformProvider::ProvideGPUTransforms(FSkinningTransformProv
 
 	RDG_EVENT_SCOPE_STAT(GraphBuilder, AnimSequenceTransformProvider, "AnimSequenceTransformProvider");
 
-	if (NumTotalEvaluateConcatenateScatterBlockHeaders > 0)
+	// Single-pass evaluate dispatches
+	if (SinglePassEvaluateHeaderBuffer)
 	{
 		FEvaluateConcatenateScatterCS::FParameters* PassParameters = GraphBuilder.AllocParameters<FEvaluateConcatenateScatterCS::FParameters>();
-		PassParameters->RetargetingDataBuffer = GraphBuilder.CreateSRV(RetargetingDataBuffer);
-		PassParameters->SequenceTransformBuffer = GraphBuilder.CreateSRV(SequenceTransformBuffer);
-		PassParameters->EvaluateBlockHeaderBuffer = GraphBuilder.CreateSRV(EvaluateConcatenateScatterBlockHeaderBuffer);
-		PassParameters->BoneMapBuffer = GraphBuilder.CreateSRV(LeafBonePathBuffer);
+		PassParameters->RetargetingDataBuffer      = GraphBuilder.CreateSRV(RetargetingDataBuffer);
+		PassParameters->SequenceTransformBuffer    = GraphBuilder.CreateSRV(SequenceTransformBuffer);
+		PassParameters->EvaluateBlockHeaderBuffer  = GraphBuilder.CreateSRV(SinglePassEvaluateHeaderBuffer);
+		PassParameters->BoneMapBuffer              = GraphBuilder.CreateSRV(LeafBonePathBuffer);
 		PassParameters->SkinningTransformBufferUAV = GetCompressedBoneTransformUAV(GraphBuilder, Context.TransformBuffer);
 
-		TShaderMapRef<FEvaluateConcatenateScatterCS> ComputeShader(ShaderMap);
-
-		FComputeShaderUtils::AddPass(
-			GraphBuilder,
+		GraphBuilder.AddPass(
 			RDG_EVENT_NAME("EvaluateConcatenateScatter"),
-			ComputeShader,
 			PassParameters,
-			FIntVector(NumTotalEvaluateConcatenateScatterBlockHeaders, 1, 1)
-		);
+			ERDGPassFlags::Compute,
+			[ShaderMap, PassParameters, &SinglePassDispatchGroups](FRDGAsyncTask, FRHIComputeCommandList& RHICmdList)
+		{
+			for (const FSinglePassDispatchGroup& Group : SinglePassDispatchGroups)
+			{
+				if (Group.NumEvaluateHeaders > 0)
+				{
+					FEvaluateConcatenateScatterCS::FPermutationDomain Perm;
+					Perm.Set<FEvaluateConcatenateScatterCS::FBonesPerGroupDim>(Group.GroupSize);
+					TShaderMapRef<FEvaluateConcatenateScatterCS> ComputeShader(ShaderMap, Perm);
+
+					PassParameters->NumTotalHeaders  = Group.NumEvaluateHeaders;
+					PassParameters->HeaderBaseOffset = Group.EvaluateHeaderBaseOffset;
+
+					const uint32 NumGroups = FMath::DivideAndRoundUp(Group.NumEvaluateHeaders, Group.InstancesPerGroup);
+					FComputeShaderUtils::Dispatch(RHICmdList, ComputeShader, *PassParameters, FIntVector(NumGroups, 1, 1));
+				}
+			}
+		});
 	}
 
-	if (NumTotalEvaluateBlendConcatenateScatterBlockHeaders > 0)
+	// Single-pass blend evaluate dispatches
+	if (SinglePassBlendEvaluateHeaderBuffer)
 	{
 		FEvaluateBlendConcatenateScatterCS::FParameters* PassParameters = GraphBuilder.AllocParameters<FEvaluateBlendConcatenateScatterCS::FParameters>();
-		PassParameters->RetargetingDataBuffer = GraphBuilder.CreateSRV(RetargetingDataBuffer);
-		PassParameters->SequenceTransformBuffer = GraphBuilder.CreateSRV(SequenceTransformBuffer);
-		PassParameters->EvaluateBlockHeaderBuffer = GraphBuilder.CreateSRV(EvaluateBlendConcatenateScatterBlockHeaderBuffer);
-		PassParameters->BoneMapBuffer = GraphBuilder.CreateSRV(LeafBonePathBuffer);
+		PassParameters->RetargetingDataBuffer      = GraphBuilder.CreateSRV(RetargetingDataBuffer);
+		PassParameters->SequenceTransformBuffer    = GraphBuilder.CreateSRV(SequenceTransformBuffer);
+		PassParameters->EvaluateBlockHeaderBuffer  = GraphBuilder.CreateSRV(SinglePassBlendEvaluateHeaderBuffer);
+		PassParameters->BoneMapBuffer              = GraphBuilder.CreateSRV(LeafBonePathBuffer);
 		PassParameters->SkinningTransformBufferUAV = GetCompressedBoneTransformUAV(GraphBuilder, Context.TransformBuffer);
 
-		TShaderMapRef<FEvaluateBlendConcatenateScatterCS> ComputeShader(ShaderMap);
-
-		FComputeShaderUtils::AddPass(
-			GraphBuilder,
+		GraphBuilder.AddPass(
 			RDG_EVENT_NAME("EvaluateBlendConcatenateScatter"),
-			ComputeShader,
 			PassParameters,
-			FIntVector(NumTotalEvaluateBlendConcatenateScatterBlockHeaders, 1, 1)
-		);
+			ERDGPassFlags::Compute,
+			[ShaderMap, PassParameters, &SinglePassDispatchGroups](FRDGAsyncTask, FRHIComputeCommandList& RHICmdList)
+		{
+			for (const FSinglePassDispatchGroup& Group : SinglePassDispatchGroups)
+			{
+				if (Group.NumBlendEvaluateHeaders > 0)
+				{
+					FEvaluateBlendConcatenateScatterCS::FPermutationDomain Perm;
+					Perm.Set<FEvaluateBlendConcatenateScatterCS::FBonesPerGroupDim>(Group.GroupSize);
+					TShaderMapRef<FEvaluateBlendConcatenateScatterCS> ComputeShader(ShaderMap, Perm);
+
+					PassParameters->NumTotalHeaders  = Group.NumBlendEvaluateHeaders;
+					PassParameters->HeaderBaseOffset = Group.BlendEvaluateHeaderBaseOffset;
+
+					const uint32 NumGroups = FMath::DivideAndRoundUp(Group.NumBlendEvaluateHeaders, Group.InstancesPerGroup);
+					FComputeShaderUtils::Dispatch(RHICmdList, ComputeShader, *PassParameters, FIntVector(NumGroups, 1, 1));
+				}
+			}
+		});
 	}
 
-	// Dispatch 1.A: Evaluate animation sequences. Outputs animated local space bone transforms with retargeting.
-	if (NumTotalEvaluateBlockHeaders > 0)
+	// 3-pass evaluate dispatches
+	if (MultiPassEvaluateHeaderBuffer)
 	{
 		FEvaluateAnimSequenceCS::FParameters* PassParameters = GraphBuilder.AllocParameters<FEvaluateAnimSequenceCS::FParameters>();
-		PassParameters->RetargetingDataBuffer = GraphBuilder.CreateSRV(RetargetingDataBuffer);
-		PassParameters->SequenceTransformBuffer = GraphBuilder.CreateSRV(SequenceTransformBuffer);
-		PassParameters->EvaluateBlockHeaderBuffer = GraphBuilder.CreateSRV(EvaluateBlockHeaderBuffer);
+		PassParameters->RetargetingDataBuffer        = GraphBuilder.CreateSRV(RetargetingDataBuffer);
+		PassParameters->SequenceTransformBuffer      = GraphBuilder.CreateSRV(SequenceTransformBuffer);
+		PassParameters->EvaluateBlockHeaderBuffer    = GraphBuilder.CreateSRV(MultiPassEvaluateHeaderBuffer);
 		PassParameters->LocalSpaceTransformBufferUAV = GraphBuilder.CreateUAV(LocalSpaceTransformBuffer);
 
-		TShaderMapRef<FEvaluateAnimSequenceCS> ComputeShader(ShaderMap);
-
-		FComputeShaderUtils::AddPass(
-			GraphBuilder,
+		GraphBuilder.AddPass(
 			RDG_EVENT_NAME("EvaluateAnimSequence"),
-			ComputeShader,
 			PassParameters,
-			FIntVector(NumTotalEvaluateBlockHeaders, 1, 1)
-		);
+			ERDGPassFlags::Compute,
+			[ShaderMap, PassParameters, &MultiPassDispatchGroups](FRDGAsyncTask, FRHIComputeCommandList& RHICmdList)
+		{
+			for (const FMultiPassDispatchGroup& Group : MultiPassDispatchGroups)
+			{
+				if (Group.NumEvaluateHeaders > 0)
+				{
+					const uint32 NumBlocksPerInstance = FMath::DivideAndRoundUp(Group.NumCompactBones, FEvaluateAnimSequenceCS::BonesPerGroup);
+					TShaderMapRef<FEvaluateAnimSequenceCS> ComputeShader(ShaderMap);
+					PassParameters->HeaderBaseOffset = Group.EvaluateHeaderBaseOffset;
+					FComputeShaderUtils::Dispatch(RHICmdList, ComputeShader, *PassParameters, FIntVector(NumBlocksPerInstance, Group.NumEvaluateHeaders, 1));
+				}
+			}
+		});
 	}
 
-	// Dispatch 1.B: Evaluate animation sequences for blends. Outputs animated local space bone transforms with retargeting.
-	if (NumTotalEvaluateBlendBlockHeaders > 0)
+	// 3-pass blend evaluate dispatches
+	if (MultiPassBlendEvaluateHeaderBuffer)
 	{
 		FEvaluateAnimSequenceBlendCS::FParameters* PassParameters = GraphBuilder.AllocParameters<FEvaluateAnimSequenceBlendCS::FParameters>();
-		PassParameters->RetargetingDataBuffer = GraphBuilder.CreateSRV(RetargetingDataBuffer);
-		PassParameters->SequenceTransformBuffer = GraphBuilder.CreateSRV(SequenceTransformBuffer);
-		PassParameters->EvaluateBlockHeaderBuffer = GraphBuilder.CreateSRV(EvaluateBlendBlockHeaderBuffer);
+		PassParameters->RetargetingDataBuffer        = GraphBuilder.CreateSRV(RetargetingDataBuffer);
+		PassParameters->SequenceTransformBuffer      = GraphBuilder.CreateSRV(SequenceTransformBuffer);
+		PassParameters->EvaluateBlockHeaderBuffer    = GraphBuilder.CreateSRV(MultiPassBlendEvaluateHeaderBuffer);
 		PassParameters->LocalSpaceTransformBufferUAV = GraphBuilder.CreateUAV(LocalSpaceTransformBuffer);
 
-		TShaderMapRef<FEvaluateAnimSequenceBlendCS> ComputeShader(ShaderMap);
-
-		FComputeShaderUtils::AddPass(
-			GraphBuilder,
+		GraphBuilder.AddPass(
 			RDG_EVENT_NAME("EvaluateAnimSequenceBlend"),
-			ComputeShader,
 			PassParameters,
-			FIntVector(NumTotalEvaluateBlendBlockHeaders, 1, 1)
-		);
+			ERDGPassFlags::Compute,
+			[ShaderMap, PassParameters, &MultiPassDispatchGroups](FRDGAsyncTask, FRHIComputeCommandList& RHICmdList)
+		{
+			for (const FMultiPassDispatchGroup& Group : MultiPassDispatchGroups)
+			{
+				if (Group.NumBlendEvaluateHeaders > 0)
+				{
+					const uint32 NumBlocksPerInstance = FMath::DivideAndRoundUp(Group.NumCompactBones, FEvaluateAnimSequenceCS::BonesPerGroup);
+					TShaderMapRef<FEvaluateAnimSequenceBlendCS> ComputeShader(ShaderMap);
+					PassParameters->HeaderBaseOffset = Group.BlendEvaluateHeaderBaseOffset;
+					FComputeShaderUtils::Dispatch(RHICmdList, ComputeShader, *PassParameters, FIntVector(NumBlocksPerInstance, Group.NumBlendEvaluateHeaders, 1));
+				}
+			}
+		});
 	}
 
-	// Dispatch 2: Concatenate leaf bone transforms (local to global). Outputs global space bone transforms.
-	if (NumTotalConcatenateBlockHeaders > 0)
+	// 3-pass dispatches: Concatenate (one RDG pass for all multi-pass groups)
+	if (!MultiPassDispatchGroups.IsEmpty() && LocalSpaceTransformBuffer)
 	{
 		FConcatenateLocalToGlobalCS::FParameters* PassParameters = GraphBuilder.AllocParameters<FConcatenateLocalToGlobalCS::FParameters>();
-		PassParameters->LocalSpaceTransformBuffer = GraphBuilder.CreateSRV(LocalSpaceTransformBuffer);
-		PassParameters->ConcatenateBlockHeaderBuffer = GraphBuilder.CreateSRV(ConcatenateBlockHeaderBuffer);
-		PassParameters->LeafBonePathBuffer = GraphBuilder.CreateSRV(LeafBonePathBuffer);
-		PassParameters->GlobalSpaceTransformBufferUAV = GraphBuilder.CreateUAV(GlobalSpaceTransformBuffer);
-
-		TShaderMapRef<FConcatenateLocalToGlobalCS> ComputeShader(ShaderMap);
+		PassParameters->LocalSpaceTransformBuffer      = GraphBuilder.CreateSRV(LocalSpaceTransformBuffer);
+		PassParameters->ConcatenateBlockHeaderBuffer   = GraphBuilder.CreateSRV(ConcatenateHeaderBuffer);
+		PassParameters->LeafBonePathBuffer             = GraphBuilder.CreateSRV(LeafBonePathBuffer);
+		PassParameters->GlobalSpaceTransformBufferUAV  = GraphBuilder.CreateUAV(GlobalSpaceTransformBuffer);
 
-		FComputeShaderUtils::AddPass(
-			GraphBuilder,
+		GraphBuilder.AddPass(
 			RDG_EVENT_NAME("ConcatenateLeafBoneTransforms"),
-			ComputeShader,
 			PassParameters,
-			FIntVector(NumTotalConcatenateBlockHeaders, 1, 1)
-		);
+			ERDGPassFlags::Compute,
+			[ShaderMap, PassParameters, &MultiPassDispatchGroups](FRDGAsyncTask, FRHIComputeCommandList& RHICmdList)
+		{
+			for (const FMultiPassDispatchGroup& Group : MultiPassDispatchGroups)
+			{
+				const uint32 NumTotalHeaders = Group.NumEvaluateHeaders + Group.NumBlendEvaluateHeaders;
+				if (NumTotalHeaders == 0)
+				{
+					continue;
+				}
+
+				const uint32 NumBlocksPerInstance = FMath::DivideAndRoundUp(Group.NumLeafBones, FEvaluateAnimSequenceCS::BonesPerGroup);
+
+				TShaderMapRef<FConcatenateLocalToGlobalCS> ComputeShader(ShaderMap);
+				PassParameters->HeaderBaseOffset = Group.ConcatenateHeaderBaseOffset;
+				FComputeShaderUtils::Dispatch(RHICmdList, ComputeShader, *PassParameters, FIntVector(NumBlocksPerInstance, NumTotalHeaders, 1));
+			}
+		});
 	}
 
-	// Dispatch 3: Scatter transforms to final skinning scene buffer.
-	if (NumTotalScatterBlockHeaders > 0)
+	// 3-pass dispatches: Scatter (one RDG pass for all multi-pass groups)
+	if (!MultiPassDispatchGroups.IsEmpty() && GlobalSpaceTransformBuffer)
 	{
 		FScatterTransformsCS::FParameters* PassParameters = GraphBuilder.AllocParameters<FScatterTransformsCS::FParameters>();
-		PassParameters->GlobalSpaceTransformBuffer = GraphBuilder.CreateSRV(GlobalSpaceTransformBuffer);
-		PassParameters->SequenceTransformBuffer = GraphBuilder.CreateSRV(SequenceTransformBuffer);
-		PassParameters->ScatterBlockHeaderBuffer = GraphBuilder.CreateSRV(ScatterBlockHeaderBuffer);
-		PassParameters->BoneMapBuffer = GraphBuilder.CreateSRV(LeafBonePathBuffer);
-		PassParameters->SkinningTransformBufferUAV = GetCompressedBoneTransformUAV(GraphBuilder, Context.TransformBuffer);
+		PassParameters->GlobalSpaceTransformBuffer   = GraphBuilder.CreateSRV(GlobalSpaceTransformBuffer);
+		PassParameters->SequenceTransformBuffer      = GraphBuilder.CreateSRV(SequenceTransformBuffer);
+		PassParameters->ScatterBlockHeaderBuffer     = GraphBuilder.CreateSRV(ScatterHeaderBuffer);
+		PassParameters->BoneMapBuffer                = GraphBuilder.CreateSRV(LeafBonePathBuffer);
+		PassParameters->SkinningTransformBufferUAV   = GetCompressedBoneTransformUAV(GraphBuilder, Context.TransformBuffer);
 
-		TShaderMapRef<FScatterTransformsCS> ComputeShader(ShaderMap);
-
-		FComputeShaderUtils::AddPass(
-			GraphBuilder,
+		GraphBuilder.AddPass(
 			RDG_EVENT_NAME("ScatterTransforms"),
-			ComputeShader,
 			PassParameters,
-			FIntVector(NumTotalScatterBlockHeaders, 1, 1)
-		);
+			ERDGPassFlags::Compute,
+			[ShaderMap, PassParameters, &MultiPassDispatchGroups](FRDGAsyncTask, FRHIComputeCommandList& RHICmdList)
+		{
+			for (const FMultiPassDispatchGroup& Group : MultiPassDispatchGroups)
+			{
+				const uint32 NumTotalHeaders = Group.NumEvaluateHeaders + Group.NumBlendEvaluateHeaders;
+				if (NumTotalHeaders == 0)
+				{
+					continue;
+				}
+
+				const uint32 NumBlocksPerInstance = FMath::DivideAndRoundUp(Group.MaxBoneTransformCount, FEvaluateAnimSequenceCS::BonesPerGroup);
+
+				TShaderMapRef<FScatterTransformsCS> ComputeShader(ShaderMap);
+				PassParameters->HeaderBaseOffset = Group.ScatterHeaderBaseOffset;
+				FComputeShaderUtils::Dispatch(RHICmdList, ComputeShader, *PassParameters, FIntVector(NumBlocksPerInstance, NumTotalHeaders, 1));
+			}
+		});
 	}
 }
 
diff --git a/Engine/Source/Runtime/Renderer/Private/Skinning/AnimSequenceTransformProvider.h b/Engine/Source/Runtime/Renderer/Private/Skinning/AnimSequenceTransformProvider.h
--- a/Engine/Source/Runtime/Renderer/Private/Skinning/AnimSequenceTransformProvider.h
+++ b/Engine/Source/Runtime/Renderer/Private/Skinning/AnimSequenceTransformProvider.h
@@ -33,6 +33,38 @@ private:
 	void ProvideGPUTransforms(FSkinningTransformProvider::FProviderContext& Context);
 	void ProvideCPUTransforms(FSkinningTransformProvider::FProviderContext& Context);
 
+	static void ComputeSinglePassGroupParams(
+		uint32 NumCompactBones,
+		uint32& OutGroupSize,
+		uint32& OutInstancesPerGroup,
+		uint32& OutSlotsPerInstance)
+	{
+		if (NumCompactBones > 256)
+		{
+			OutGroupSize = 0;
+			OutInstancesPerGroup = 1;
+			OutSlotsPerInstance = 0;
+			return;
+		}
+
+		OutSlotsPerInstance = FMath::RoundUpToPowerOfTwo(NumCompactBones);
+
+		if (NumCompactBones <= 64)
+		{
+			OutGroupSize = 64;
+		}
+		else if (NumCompactBones <= 128)
+		{
+			OutGroupSize = 128;
+		}
+		else
+		{
+			OutGroupSize = 256;
+		}
+
+		OutInstancesPerGroup = OutGroupSize / OutSlotsPerInstance;
+	}
+
 	struct FBoneTables
 	{
 		TArray<uint16> MeshToCompact;
@@ -301,6 +333,12 @@ private:
 			bool bSinglePass = false;
 			bool bUploaded = false;
 
+			// Dispatch grouping params (computed at registration)
+			uint32 GroupSize         = 0; // 64, 128, 256 for single-pass; 0 for 3-pass
+			uint32 InstancesPerGroup = 1; // Wave-packing factor (single-pass only)
+			uint32 SlotsPerInstance  = 0; // Power-of-two >= NumCompactBones (single-pass only)
+			void*  DispatchGroup     = nullptr; // Points to per-frame FSinglePassDispatchGroup or FMultiPassDispatchGroup
+
 			bool IsValid() const
 			{
 				return BoneTables.NumCompactBones != 0;
