Fix validation error when Hair and SLW are active at the same time.

#Synced-CL 52056680

diff --git a/Engine/Shaders/Private/HairStrands/HairStrandsComposition.usf b/Engine/Shaders/Private/HairStrands/HairStrandsComposition.usf
--- a/Engine/Shaders/Private/HairStrands/HairStrandsComposition.usf
+++ b/Engine/Shaders/Private/HairStrands/HairStrandsComposition.usf
@@ -77,7 +77,7 @@ int2 OutputResolution;
 uint bComposeDofDepth;
 uint bHasHoldout;
 
-#if PERMUTATION_SLW
+#if PERMUTATION_READ_SLW_DEPTH
 #define SLW_COMPOSE_NONE 0
 #define SLW_COMPOSE_BEHIND_WATER 1
 #define SLW_COMPOSE_INFRONTOF_WATER 2
@@ -90,8 +90,11 @@ Texture2D<float4>	HairLightingSampleBuffer;
 
 void ComposeSamplePS(
 	in FScreenVertexOutput Input,
-	out float4 OutColor : SV_Target0,
-	out float  OutDepth : SV_DEPTH)
+	out float4 OutColor : SV_Target0
+	#if PERMUTATION_WRITE_DEPTH
+	,out float  OutDepth : SV_DEPTH
+	#endif
+)
 #if PERMUTATION_DEBUG == 0
 {
 	OutColor = 0.0f;
@@ -111,7 +114,7 @@ void ComposeSamplePS(
 		discard;
 	}
 
-	#if PERMUTATION_SLW
+	#if PERMUTATION_READ_SLW_DEPTH
 	const float SLWDepth = SLWDepthTexture.Load(PixelCoord);
 	#endif
 	const uint TotalNodeCount = HairStrands.HairSampleCount[0];
@@ -134,7 +137,7 @@ void ComposeSamplePS(
 		}
 		#endif
 		
-		#if PERMUTATION_SLW
+		#if PERMUTATION_READ_SLW_DEPTH
 		{
 			const FHairSample Sample = UnpackHairSample(HairStrands.HairSampleData[NodeDesc.Offset + SampleIt]);
 			if (SLWComposeType == SLW_COMPOSE_BEHIND_WATER)
@@ -153,7 +156,7 @@ void ComposeSamplePS(
 		LocalAccColor += HairLightingSampleBuffer.Load(uint3(Coord, 0)).xyz;
 	}
 
-	#if PERMUTATION_SLW
+	#if PERMUTATION_READ_SLW_DEPTH
 	if (bSkipComposition)
 	{
 		discard;
@@ -170,11 +173,13 @@ void ComposeSamplePS(
 	OutColor.rgb = (LocalAccColor * Volumetric.a + Volumetric.rgb) * PixelCoverage;
 	OutColor.a = PixelCoverage;
 
+	#if PERMUTATION_WRITE_DEPTH
 	OutDepth = HairDeviceZ;
 	if (bComposeDofDepth)
 	{
 		OutDepth = HairDOFDepthTexture.Load(PixelCoord);
 	}
+	#endif
 }
 #else // PERMUTATION_DEBUG == 1
 {
@@ -206,7 +211,9 @@ void ComposeSamplePS(
 
 	OutColor.rgb *= PixelCoverage;
 	OutColor.a = PixelCoverage;
+	#if PERMUTATION_WRITE_DEPTH
 	OutDepth = HairDeviceZ;
+	#endif
 }
 #endif
 
diff --git a/Engine/Source/Runtime/Renderer/Private/HairStrands/HairStrandsComposition.cpp b/Engine/Source/Runtime/Renderer/Private/HairStrands/HairStrandsComposition.cpp
--- a/Engine/Source/Runtime/Renderer/Private/HairStrands/HairStrandsComposition.cpp
+++ b/Engine/Source/Runtime/Renderer/Private/HairStrands/HairStrandsComposition.cpp
@@ -139,7 +139,14 @@ void InternalCommonDrawPass(
 
 		if (Type == EHairStrandsCommonPassType::Composition || Type == EHairStrandsCommonPassType::HoldoutScene || Type == EHairStrandsCommonPassType::HoldoutHair)
 		{
-			GraphicsPSOInit.DepthStencilState = TStaticDepthStencilState<true, CF_DepthNearOrEqual>::GetRHI();
+			if (bWriteDepth)
+			{
+				GraphicsPSOInit.DepthStencilState = TStaticDepthStencilState<true, CF_DepthNearOrEqual>::GetRHI();
+			}
+			else
+			{
+				GraphicsPSOInit.DepthStencilState = TStaticDepthStencilState<false, CF_Always>::GetRHI();
+			}
 		}
 		else if (Type == EHairStrandsCommonPassType::DOF)
 		{
@@ -238,7 +245,6 @@ static void AddHairHoldoutPass(
 		Parameters->HairStrands = View.HairStrandsViewData.UniformBuffer;
 		Parameters->FogStruct = FogBuffer;
 		Parameters->RenderTargets[0] = FRenderTargetBinding(OutColorTexture, ERenderTargetLoadAction::ELoad);
-		Parameters->RenderTargets.DepthStencil = FDepthStencilBinding(OutDepthTexture, ERenderTargetLoadAction::ELoad, ERenderTargetLoadAction::ELoad, FExclusiveDepthStencil::DepthWrite_StencilRead);
 
 		FHairHoldoutPS::FPermutationDomain PermutationVector;
 		TShaderMapRef<FHairHoldoutPS> PixelShader(View.ShaderMap, PermutationVector);
@@ -266,7 +272,6 @@ static void AddHairHoldoutPass(
 		Parameters->HairStrands = View.HairStrandsViewData.UniformBuffer;
 		Parameters->FogStruct = FogBuffer;
 		Parameters->RenderTargets[0] = FRenderTargetBinding(OutColorTexture, ERenderTargetLoadAction::ELoad);
-		Parameters->RenderTargets.DepthStencil = FDepthStencilBinding(OutDepthTexture, ERenderTargetLoadAction::ELoad, ERenderTargetLoadAction::ELoad, FExclusiveDepthStencil::DepthWrite_StencilRead);
 
 		FHairHoldoutPS::FPermutationDomain PermutationVector;
 		TShaderMapRef<FHairHoldoutPS> PixelShader(View.ShaderMap, PermutationVector);
@@ -292,8 +297,9 @@ class FHairVisibilityComposeSamplePS : public FGlobalShader
 	SHADER_USE_PARAMETER_STRUCT(FHairVisibilityComposeSamplePS, FGlobalShader);
 
 	class FDebug : SHADER_PERMUTATION_BOOL("PERMUTATION_DEBUG");
-	class FSlw : SHADER_PERMUTATION_BOOL("PERMUTATION_SLW");
-	using FPermutationDomain = TShaderPermutationDomain<FDebug, FSlw>;
+	class FReadSLWDepth : SHADER_PERMUTATION_BOOL("PERMUTATION_READ_SLW_DEPTH");
+	class FWriteDepth : SHADER_PERMUTATION_BOOL("PERMUTATION_WRITE_DEPTH");
+	using FPermutationDomain = TShaderPermutationDomain<FDebug, FReadSLWDepth, FWriteDepth>;
 	
 	BEGIN_SHADER_PARAMETER_STRUCT(FParameters, )
 		SHADER_PARAMETER_STRUCT_REF(FViewUniformShaderParameters, ViewUniformBuffer)
@@ -309,7 +315,7 @@ class FHairVisibilityComposeSamplePS : public FGlobalShader
 		SHADER_PARAMETER_RDG_TEXTURE_UAV(RWTexture2D<uint>, OutMetaTexture)
 		SHADER_PARAMETER_RDG_UNIFORM_BUFFER(FFogUniformParameters, FogStruct)
 		RENDER_TARGET_BINDING_SLOTS()
-		END_SHADER_PARAMETER_STRUCT()
+	END_SHADER_PARAMETER_STRUCT()
 
 public:
 		static bool ShouldCompilePermutation(const FGlobalShaderPermutationParameters& Parameters) { return IsHairStrandsSupported(EHairStrandsShaderType::Strands, Parameters.Platform); }
@@ -338,6 +344,10 @@ static void AddHairVisibilityComposeSamplePass(
 
 	TRDGUniformBufferRef<FFogUniformParameters> FogBuffer = CreateFogUniformBuffer(GraphBuilder, View);
 
+	// When bReadSLWDepth == true, the SLW copy pass has been executed, and thus the hair depth has already be written into the scene depth (no need to rewrite it again)
+	const bool bReadSLWDepth = SLWDepthTexture && SLWComposeType != EHairSLWCompositionType::All;
+	const bool bWriteDepth = !bReadSLWDepth;
+
 	FHairVisibilityComposeSamplePS::FParameters* Parameters = GraphBuilder.AllocParameters<FHairVisibilityComposeSamplePS::FParameters>();
 	Parameters->bComposeDofDepth = bDOFEnable ? 1 : 0;
 	Parameters->SLWComposeType = 0;
@@ -349,10 +359,13 @@ static void AddHairVisibilityComposeSamplePass(
 	Parameters->FogStruct = FogBuffer;
 	Parameters->bHasHoldout = bHasHoldout ? 1u : 0u;
 	Parameters->RenderTargets[0] = FRenderTargetBinding(OutColorTexture, ERenderTargetLoadAction::ELoad);
-	Parameters->RenderTargets.DepthStencil = FDepthStencilBinding(OutDepthTexture, ERenderTargetLoadAction::ELoad, ERenderTargetLoadAction::ELoad, FExclusiveDepthStencil::DepthWrite_StencilRead);
 
-	const bool bSWL = SLWDepthTexture && SLWComposeType != EHairSLWCompositionType::All;
-	if (bSWL)
+	if (bWriteDepth)
+	{
+		Parameters->RenderTargets.DepthStencil = FDepthStencilBinding(OutDepthTexture, ERenderTargetLoadAction::ELoad, ERenderTargetLoadAction::ELoad, FExclusiveDepthStencil::DepthWrite_StencilRead);
+	}
+
+	if (bReadSLWDepth)
 	{
 		Parameters->SLWComposeType = SLWComposeType;
 		Parameters->SLWDepthTexture = SLWDepthTexture;
@@ -361,15 +374,16 @@ static void AddHairVisibilityComposeSamplePass(
 	const bool bDebugComposition = View.Family->EngineShowFlags.LODColoration;
 	FHairVisibilityComposeSamplePS::FPermutationDomain PermutationVector;
 	PermutationVector.Set<FHairVisibilityComposeSamplePS::FDebug>(bDebugComposition);
-	PermutationVector.Set<FHairVisibilityComposeSamplePS::FSlw>(bSWL);
+	PermutationVector.Set<FHairVisibilityComposeSamplePS::FReadSLWDepth>(bReadSLWDepth);
+	PermutationVector.Set<FHairVisibilityComposeSamplePS::FWriteDepth>(bWriteDepth);
 	TShaderMapRef<FHairVisibilityComposeSamplePS> PixelShader(View.ShaderMap, PermutationVector);
 	InternalCommonDrawPass(
 		GraphBuilder,
-		RDG_EVENT_NAME("HairStrands::ComposeSample"),
+		RDG_EVENT_NAME("HairStrands::ComposeSample(WriteDepth:%s)", bWriteDepth?TEXT("1"):TEXT("0"), bReadSLWDepth?TEXT(",ReadSLW"):TEXT("")),
 		View,
 		OutColorTexture->Desc.Extent,
 		EHairStrandsCommonPassType::Composition,
-		false /*bWriteDepth*/,
+		bWriteDepth,
 		bHasHoldout,
 		VisibilityData.TileData,
 		PixelShader,
