Move skinning pipeline to PreGPUSceneUpdate.
- Moves FinishSkinningBufferUpload + PerformSkinning from UpdateSceneUniformBuffer (post-visibility) to PreGPUSceneUpdate (inside FScene::Update, pre-GPU-Scene). This ensures bone transforms are available before hair groom reads them
in PreVisibilityFrameSetup, and before GPU Scene uploads instance data.

#Synced-CL 52105729

diff --git a/Engine/Source/Runtime/Renderer/Private/RendererScene.cpp b/Engine/Source/Runtime/Renderer/Private/RendererScene.cpp
--- a/Engine/Source/Runtime/Renderer/Private/RendererScene.cpp
+++ b/Engine/Source/Runtime/Renderer/Private/RendererScene.cpp
@@ -5374,6 +5374,13 @@ FViewSceneChangeSet* FScene::ProcessViewChanges(FRDGBuilder& GraphBuilder, TCons
 	ChangeSet->bIsRequiresDebugMaterialChanged = NewViewStateDebugFlags != PersistentViewStateDebugFlags;
 	PersistentViewStateDebugFlags = NewViewStateDebugFlags;
 
+	// Extract scene time from the view family. Views should always be non-empty here;
+	// ProcessViewChanges is only called when SceneUpdateInputs provides rendered views.
+	if (ensureMsgf(Views.Num() > 0, TEXT("ProcessViewChanges called with no views")))
+	{
+		ChangeSet->CurrentTime = Views[0]->Family->Time;
+	}
+
 	return ChangeSet;
 }
 
@@ -6664,6 +6671,14 @@ void FScene::Update(FRDGBuilder& GraphBuilder, const FUpdateParameters& Paramete
 		});
 	}
 	
+	if (Parameters.ViewUpdateChangeSet)
+	{
+		if (auto SkinningUpdater = SceneExtensionsUpdaters.GetUpdaterPtr<FSkinningSceneExtension::FUpdater>())
+		{
+			SkinningUpdater->PreGPUSceneUpdate(GraphBuilder, Parameters.ViewUpdateChangeSet->CurrentTime);
+		}
+	}
+
 	{
 		RDG_CSV_STAT_EXCLUSIVE_SCOPE(GraphBuilder, UpdateGPUScene);
 		RDG_EVENT_SCOPE_STAT(GraphBuilder, GPUSceneUpdate, "GPUSceneUpdate");
diff --git a/Engine/Source/Runtime/Renderer/Private/ScenePrivate.h b/Engine/Source/Runtime/Renderer/Private/ScenePrivate.h
--- a/Engine/Source/Runtime/Renderer/Private/ScenePrivate.h
+++ b/Engine/Source/Runtime/Renderer/Private/ScenePrivate.h
@@ -1582,6 +1582,7 @@ struct FViewSceneChangeSet
 {
 	TArray<FPersistentViewId, SceneRenderingAllocator> AddedViewIds;
 	TArray<FPersistentViewId> RemovedViewIds;
+	FGameTime CurrentTime;
 	bool bIsRequiresDebugMaterialChanged = false;
 };
 
diff --git a/Engine/Source/Runtime/Renderer/Private/Skinning/AnimBankTransformProvider.cpp b/Engine/Source/Runtime/Renderer/Private/Skinning/AnimBankTransformProvider.cpp
--- a/Engine/Source/Runtime/Renderer/Private/Skinning/AnimBankTransformProvider.cpp
+++ b/Engine/Source/Runtime/Renderer/Private/Skinning/AnimBankTransformProvider.cpp
@@ -313,11 +313,11 @@ void FAnimBankTransformProvider::ProvideGPUBankTransforms(FSkinningTransformProv
 	int32& BlockCount = *GraphBuilder.AllocPOD<int32>();
 	BlockCount = 0;
 
-	GraphBuilder.AddSetupTask([this, ProxyDatas, BlockHeaders, &ViewFamily = Context.ViewFamily, Indirections = Context.Indirections, &BlockCount, BonesPerGroup]
+	GraphBuilder.AddSetupTask([this, ProxyDatas, BlockHeaders, CurrentTime = Context.CurrentTime, Indirections = Context.Indirections, &BlockCount, BonesPerGroup]
 	{
 		TRACE_CPUPROFILER_EVENT_SCOPE(FAnimBankTransformProvider::BuildBoneBlocks);
 		const float GlobalTimeScale = CVarAnimBankTimeScale.GetValueOnRenderThread();
-		const float DeltaTime = ViewFamily.Time.GetDeltaWorldTimeSeconds() * GlobalTimeScale;
+		const float DeltaTime = CurrentTime.GetDeltaWorldTimeSeconds() * GlobalTimeScale;
 
 		for (const FSkinningTransformProvider::FProviderIndirection Indirection : Indirections)
 		{
@@ -475,7 +475,7 @@ void FAnimBankTransformProvider::ProvideCPUBankTransforms(FSkinningTransformProv
 	}
 
 	Context.ScatterUploadBuilder.AddPass(GraphBuilder, SkinningTransformUploadBuffer, Context.TransformBuffer, GlobalTransformCount, sizeof(FCompressedBoneTransform), TEXT("Skinning.AnimTransforms"),
-		[this, Indirections = Context.Indirections, Proxies = Context.Proxies, &ViewFamily = Context.ViewFamily, GlobalTransformCount] (FRDGScatterUploader& ScatterUploader)
+		[this, Indirections = Context.Indirections, Proxies = Context.Proxies, CurrentTime = Context.CurrentTime, GlobalTransformCount] (FRDGScatterUploader& ScatterUploader)
 	{
 		TRACE_CPUPROFILER_EVENT_SCOPE(FSkinningSceneExtension::ProvideAnimBankRuntimeTransformsTask);
 		TArray<FCompressedBoneTransform, SceneRenderingAllocator> SourceTransforms;
@@ -575,7 +575,7 @@ void FAnimBankTransformProvider::ProvideCPUBankTransforms(FSkinningTransformProv
 		};
 
 		const float GlobalTimeScale = CVarAnimBankTimeScale.GetValueOnRenderThread();
-		const float DeltaTime = ViewFamily.Time.GetDeltaWorldTimeSeconds() * GlobalTimeScale;
+		const float DeltaTime = CurrentTime.GetDeltaWorldTimeSeconds() * GlobalTimeScale;
 
 		for (const FSkinningTransformProvider::FProviderIndirection Indirection : Indirections)
 		{
diff --git a/Engine/Source/Runtime/Renderer/Private/Skinning/AnimRuntimeTransformProvider.cpp b/Engine/Source/Runtime/Renderer/Private/Skinning/AnimRuntimeTransformProvider.cpp
--- a/Engine/Source/Runtime/Renderer/Private/Skinning/AnimRuntimeTransformProvider.cpp
+++ b/Engine/Source/Runtime/Renderer/Private/Skinning/AnimRuntimeTransformProvider.cpp
@@ -64,7 +64,7 @@ void FAnimRuntimeTransformProvider::ProvideAnimRuntimeTransforms(FSkinningTransf
 	FRDGAsyncScatterUploadBuffer& TransformUploadBuffer = *GraphBuilder.AllocObject<FRDGAsyncScatterUploadBuffer>();
 
 	Context.ScatterUploadBuilder.AddPass(GraphBuilder, TransformUploadBuffer, Context.TransformBuffer, GlobalTransformCount, sizeof(FCompressedBoneTransform), TEXT("Skinning.AnimRuntimeTransforms"),
-		[this, Indirections = Context.Indirections, Proxies = Context.Proxies, &ViewFamily = Context.ViewFamily, GlobalTransformCount] (FRDGScatterUploader& ScatterUploader)
+		[this, Indirections = Context.Indirections, Proxies = Context.Proxies, CurrentTime = Context.CurrentTime, GlobalTransformCount] (FRDGScatterUploader& ScatterUploader)
 	{
 		TRACE_CPUPROFILER_EVENT_SCOPE(FAnimRuntimeTransformProvider::ProvideAnimRuntimeTransformsTask);
 
diff --git a/Engine/Source/Runtime/Renderer/Private/Skinning/AnimSequenceTransformProvider.cpp b/Engine/Source/Runtime/Renderer/Private/Skinning/AnimSequenceTransformProvider.cpp
--- a/Engine/Source/Runtime/Renderer/Private/Skinning/AnimSequenceTransformProvider.cpp
+++ b/Engine/Source/Runtime/Renderer/Private/Skinning/AnimSequenceTransformProvider.cpp
@@ -1031,8 +1031,8 @@ void FAnimSequenceTransformProvider::ProvideGPUTransforms(FSkinningTransformProv
 	TArray<FLODUploadRequest, SceneRenderingAllocator> LODUploadRequests;
 
 	const float GlobalTimeScale = CVarAnimSequenceTransformProviderTimeScale.GetValueOnRenderThread();
-	const double WorldTime = Context.ViewFamily.Time.GetWorldTimeSeconds() * GlobalTimeScale;
-	const float DeltaTime = Context.ViewFamily.Time.GetDeltaWorldTimeSeconds() * GlobalTimeScale;
+	const double WorldTime = Context.CurrentTime.GetWorldTimeSeconds() * GlobalTimeScale;
+	const float DeltaTime = Context.CurrentTime.GetDeltaWorldTimeSeconds() * GlobalTimeScale;
 
 	uint32 NumTotalSubTracks = 0;
 	uint32 NumTotalIntermediateTransforms = 0;
@@ -2082,7 +2082,7 @@ void FAnimSequenceTransformProvider::ProvideCPUTransforms(FSkinningTransformProv
 	}
 
 	Context.ScatterUploadBuilder.AddPass(GraphBuilder, SkinningTransformUploadBuffer, Context.TransformBuffer, GlobalTransformCount, sizeof(FCompressedBoneTransform), TEXT("Skinning.AnimTransforms"),
-		[this, Indirections = Context.Indirections, Proxies = Context.Proxies, &ViewFamily = Context.ViewFamily, GlobalTransformCount] (FRDGScatterUploader& ScatterUploader)
+		[this, Indirections = Context.Indirections, Proxies = Context.Proxies, CurrentTime = Context.CurrentTime, GlobalTransformCount] (FRDGScatterUploader& ScatterUploader)
 	{
 		TRACE_CPUPROFILER_EVENT_SCOPE(FAnimSequenceTransformProvider::ProvideCPUTransformsTask);
 
@@ -2341,8 +2341,8 @@ void FAnimSequenceTransformProvider::ProvideCPUTransforms(FSkinningTransformProv
 		};
 
 		const float GlobalTimeScale = CVarAnimSequenceTransformProviderTimeScale.GetValueOnRenderThread();
-		const double WorldTime = ViewFamily.Time.GetWorldTimeSeconds() * GlobalTimeScale;
-		const float DeltaTime = ViewFamily.Time.GetDeltaWorldTimeSeconds() * GlobalTimeScale;
+		const double WorldTime = CurrentTime.GetWorldTimeSeconds() * GlobalTimeScale;
+		const float DeltaTime = CurrentTime.GetDeltaWorldTimeSeconds() * GlobalTimeScale;
 
 		uint32 NumActiveTracks = 0;
 		uint32 NumCurrentTransformUpdates = 0;
diff --git a/Engine/Source/Runtime/Renderer/Private/Skinning/SkinningSceneExtension.cpp b/Engine/Source/Runtime/Renderer/Private/Skinning/SkinningSceneExtension.cpp
--- a/Engine/Source/Runtime/Renderer/Private/Skinning/SkinningSceneExtension.cpp
+++ b/Engine/Source/Runtime/Renderer/Private/Skinning/SkinningSceneExtension.cpp
@@ -359,7 +359,7 @@ void FSkinningSceneExtension::FinishSkinningBufferUpload(FRDGBuilder& GraphBuild
 	}
 }
 
-void FSkinningSceneExtension::PerformSkinning(FSkinningSceneParameters& Parameters, FRDGBuilder& GraphBuilder, const FSceneViewFamily& ViewFamily)
+void FSkinningSceneExtension::PerformSkinning(FSkinningSceneParameters& Parameters, FRDGBuilder& GraphBuilder, const FGameTime& CurrentTime)
 {
 	RDG_EVENT_SCOPE(GraphBuilder, "Skinning");
 
@@ -531,7 +531,7 @@ void FSkinningSceneExtension::PerformSkinning(FSkinningSceneParameters& Paramete
 				TConstArrayView<FSkinningSceneExtensionProxy*>{},
 				BatchIndices,
 				Batches,
-				ViewFamily,
+				CurrentTime,
 				GraphBuilder,
 				ScatterUploadBuilder,
 				Parameters.BoneTransforms->GetParent(),
@@ -624,7 +624,7 @@ void FSkinningSceneExtension::PerformSkinning(FSkinningSceneParameters& Paramete
 				Proxies,
 				PrimitiveIndices,
 				TConstArrayView<FSkeletonBatch>(),
-				ViewFamily,
+				CurrentTime,
 				GraphBuilder,
 				ScatterUploadBuilder,
 				Parameters.BoneTransforms->GetParent(),
@@ -1634,16 +1634,31 @@ void FSkinningSceneExtension::FRenderer::UpdateViewData(FRDGBuilder& GraphBuilde
 	}
 }
 
+void FSkinningSceneExtension::FUpdater::PreGPUSceneUpdate(FRDGBuilder& GraphBuilder, const FGameTime& CurrentTime)
+{
+	if (!SceneData->IsEnabled())
+	{
+		return;
+	}
+
+	SCOPED_NAMED_EVENT(FSkinningSceneExtension_FUpdater_PreGPUSceneUpdate, FColor::Silver);
+	CSV_SCOPED_TIMING_STAT_EXCLUSIVE(STAT_SkinningSceneExtension);
+
+	FSkinningSceneParameters Parameters;
+	const bool bUpdateStats = true;
+	SceneData->FinishSkinningBufferUpload(GraphBuilder, &Parameters, bUpdateStats);
+	SceneData->PerformSkinning(Parameters, GraphBuilder, CurrentTime);
+}
+
 void FSkinningSceneExtension::FRenderer::UpdateSceneUniformBuffer(FRDGBuilder& GraphBuilder, FSceneUniformBuffer& SceneUniformBuffer)
 {
 	SCOPED_NAMED_EVENT(FSkinningSceneExtension_FRenderer_UpdateSceneUniformBuffer, FColor::Silver);
 	CSV_SCOPED_TIMING_STAT_EXCLUSIVE(STAT_SkinningSceneExtension);
 	check(SceneData->IsEnabled());
 	FSkinningSceneParameters Parameters;
-	const bool bUpdateStats = true;
+	const bool bUpdateStats = false;
 	SceneData->FinishSkinningBufferUpload(GraphBuilder, &Parameters, bUpdateStats);
 	SceneUniformBuffer.Set(SceneUB::Skinning, Parameters);
-	SceneData->PerformSkinning(Parameters, GraphBuilder, ViewFamily);
 }
 
 void FSkinningSceneExtension::GetSkinnedPrimitives(TArray<FPrimitiveSceneInfo*>& OutPrimitives) const
@@ -1779,7 +1794,7 @@ void FSkinningSceneExtension::ProvideMeshObjectTransforms(FSkinningTransformProv
 	FRDGAsyncScatterUploadBuffer& TransformUploadBuffer = *GraphBuilder.AllocObject<FRDGAsyncScatterUploadBuffer>();
 
 	Context.ScatterUploadBuilder.AddPass(GraphBuilder, TransformUploadBuffer, Context.TransformBuffer, GlobalTransformCount, sizeof(FCompressedBoneTransform), TEXT("Skinning.AnimTransforms"),
-		[Indirections = Context.Indirections, Proxies = Context.Proxies, &ViewFamily = Context.ViewFamily, GlobalTransformCount] (FRDGScatterUploader& ScatterUploader)
+		[Indirections = Context.Indirections, Proxies = Context.Proxies, GlobalTransformCount] (FRDGScatterUploader& ScatterUploader)
 	{
 		TRACE_CPUPROFILER_EVENT_SCOPE(FSkinningSceneExtension::ProvideMeshObjectTransformsTask);
 
diff --git a/Engine/Source/Runtime/Renderer/Private/Skinning/SkinningSceneExtension.h b/Engine/Source/Runtime/Renderer/Private/Skinning/SkinningSceneExtension.h
--- a/Engine/Source/Runtime/Renderer/Private/Skinning/SkinningSceneExtension.h
+++ b/Engine/Source/Runtime/Renderer/Private/Skinning/SkinningSceneExtension.h
@@ -37,6 +37,8 @@ public:
 		virtual void PreSceneUpdate(FRDGBuilder& GraphBuilder, const FScenePreUpdateChangeSet& ChangeSet) override;
 		virtual void PostSceneUpdate(FRDGBuilder& GraphBuilder, const FScenePostUpdateChangeSet& ChangeSet) override;
 
+		void PreGPUSceneUpdate(FRDGBuilder& GraphBuilder, const FGameTime& CurrentTime);
+
 		void PostMeshUpdate(FRDGBuilder& GraphBuilder, const TConstArrayView<FPrimitiveSceneInfo*>& SceneInfosWithStaticDrawListUpdate, const FSkeletalMeshUpdater::FSubmitData* SkeletalMeshSubmitData);
 
 	private:
@@ -200,7 +202,7 @@ private:
 	void PerformSkinning(
 		FSkinningSceneParameters& Parameters,
 		FRDGBuilder& GraphBuilder,
-		const FSceneViewFamily& ViewFamily
+		const FGameTime& CurrentTime
 	);
 
 	bool ProcessBufferDefragmentation();
diff --git a/Engine/Source/Runtime/Renderer/Private/Skinning/SkinningTransformProvider.h b/Engine/Source/Runtime/Renderer/Private/Skinning/SkinningTransformProvider.h
--- a/Engine/Source/Runtime/Renderer/Private/Skinning/SkinningTransformProvider.h
+++ b/Engine/Source/Runtime/Renderer/Private/Skinning/SkinningTransformProvider.h
@@ -98,7 +98,7 @@ public:
 			const TConstArrayView<FSkinningSceneExtensionProxy*> InProxies,
 			const TConstArrayView<FProviderIndirection> InIndirections,
 			const TConstArrayView<FSkeletonBatch> InSkeletonBatches,
-			const FSceneViewFamily& InViewFamily,
+			const FGameTime& InCurrentTime,
 			FRDGBuilder& InGraphBuilder,
 			FRDGScatterUploadBuilder& InScatterUploadBuilder,
 			FRDGBufferRef InTransformBuffer,
@@ -108,11 +108,11 @@ public:
 		, Proxies(InProxies)
 		, Indirections(InIndirections)
 		, SkeletonBatches(InSkeletonBatches)
+		, CurrentTime(InCurrentTime)
 		, GraphBuilder(InGraphBuilder)
 		, ScatterUploadBuilder(InScatterUploadBuilder)
 		, TransformBuffer(InTransformBuffer)
 		, BoneMapBufferSRV(InBoneMapBufferSRV)
-		, ViewFamily(InViewFamily)
 		{
 		}
 
@@ -121,12 +121,11 @@ public:
 		TConstArrayView<FProviderIndirection> Indirections;
 		TConstArrayView<FSkeletonBatch> SkeletonBatches;
 
+		const FGameTime& CurrentTime;
 		FRDGBuilder& GraphBuilder;
 		FRDGScatterUploadBuilder& ScatterUploadBuilder;
 		FRDGBufferRef TransformBuffer;
 		FRDGBufferSRVRef BoneMapBufferSRV;
-
-		const FSceneViewFamily& ViewFamily;
 	};
 
 	DECLARE_DELEGATE_OneParam(FOnProvideTransforms, FProviderContext&);
