Fix empty RDG pass names.

#rnx

#Synced-CL 52126719

diff --git a/Engine/Source/Runtime/Renderer/Private/DistortionRendering.cpp b/Engine/Source/Runtime/Renderer/Private/DistortionRendering.cpp
--- a/Engine/Source/Runtime/Renderer/Private/DistortionRendering.cpp
+++ b/Engine/Source/Runtime/Renderer/Private/DistortionRendering.cpp
@@ -392,7 +392,7 @@ static void AddCopySceneColorDepthPass(FRDGBuilder& GraphBuilder, const FViewInf
 	const FScreenPassTextureViewport InputViewport(SceneColorTexture, ViewRect);
 	const FScreenPassTextureViewport OutputViewport(SceneColorCopyTexture);
 
-	AddDrawScreenPass(GraphBuilder, {}, View, InputViewport, OutputViewport, VertexShader, PixelShader, PassParameters);
+	AddDrawScreenPass(GraphBuilder, RDG_EVENT_NAME("Distortion::CopySceneColorDepth"), View, InputViewport, OutputViewport, VertexShader, PixelShader, PassParameters);
 }
 
 //////////////////////////////////////////////////////////////////////////
@@ -641,10 +641,7 @@ void FDeferredShadingSceneRenderer::RenderDistortion(
 
 			// Copy scene color into the first mip level
 
-			{
-				RDG_EVENT_SCOPE(GraphBuilder, "CopySceneColorDepth");
-				AddCopySceneColorDepthPass(GraphBuilder, View, TAASceneColorTexture, SceneDepthTexture, SceneColorMipchainTexture, SceneDepthMeterMipchainTexture);
-			}
+			AddCopySceneColorDepthPass(GraphBuilder, View, TAASceneColorTexture, SceneDepthTexture, SceneColorMipchainTexture, SceneDepthMeterMipchainTexture);
 
 			// Now render the mip chain
 			// SUBSTRATE_TODO we could optimize that pass by doing one pass with a tile of 16x16 writing out the 8x8, 4x4, 2x2 and 1x1 down sampled output
@@ -761,7 +758,6 @@ void FDeferredShadingSceneRenderer::RenderDistortion(
 
 	{
 		QUICK_SCOPE_CYCLE_COUNTER(STAT_FSceneRenderer_RenderDistortion_Accumulate);
-		RDG_EVENT_SCOPE(GraphBuilder, "Accumulate");
 
 		// Use RGBA8 light target for accumulating distortion offsets.
 		// R = positive X offset
@@ -817,7 +813,6 @@ void FDeferredShadingSceneRenderer::RenderDistortion(
 			if (auto* Pass = View.ParallelMeshDrawCommandPasses[EMeshPass::Distortion])
 			{
 				RDG_GPU_MASK_SCOPE(GraphBuilder, View.GPUMask);
-				RDG_EVENT_SCOPE_CONDITIONAL(GraphBuilder, Views.Num() > 1, "View%d", ViewIndex);
 				View.BeginRenderView();
 
 				auto* PassParameters = GraphBuilder.AllocParameters<FDistortionPassParameters>();
@@ -834,7 +829,7 @@ void FDeferredShadingSceneRenderer::RenderDistortion(
 				Pass->BuildRenderingCommands(GraphBuilder, Scene->GPUScene, PassParameters->InstanceCullingDrawParams);
 
 				GraphBuilder.AddPass(
-					{},
+					RDG_EVENT_NAME("Distortion::Accumulate View%d", ViewIndex),
 					PassParameters,
 					ERDGPassFlags::Raster,
 					[&View, Pass, PassParameters](FRDGAsyncTask, FRHICommandList& RHICmdList)
@@ -908,7 +903,6 @@ void FDeferredShadingSceneRenderer::RenderDistortion(
 	// Apply distortion and store off-screen.
 	{
 		QUICK_SCOPE_CYCLE_COUNTER(STAT_FSceneRenderer_RenderDistortion_Apply);
-		RDG_EVENT_SCOPE(GraphBuilder, "Apply");
 		CommonParameters.SceneColorMSAATexture = SceneColorTexture;
 		CommonParameters.SceneColorTexture = SceneColorTexture;
 		CommonParameters.RenderTargets.DepthStencil = StencilReadBinding;
@@ -946,7 +940,7 @@ void FDeferredShadingSceneRenderer::RenderDistortion(
 			Viewport.Rect = View.ViewRect;
 
 			ClearUnusedGraphResources(ApplyPixelShader, PassParameters);
-			AddDrawScreenPass(GraphBuilder, {}, FScreenPassViewInfo(View, 1), Viewport, Viewport, PipelineState, PassParameters,
+			AddDrawScreenPass(GraphBuilder, RDG_EVENT_NAME("Distortion::Apply View%d", ViewIndex), FScreenPassViewInfo(View, 1), Viewport, Viewport, PipelineState, PassParameters,
 				[ApplyPixelShader, PassParameters](FRHICommandList& RHICmdList)
 			{
 				SetShaderParameters(RHICmdList, ApplyPixelShader, ApplyPixelShader.GetPixelShader(), *PassParameters);
@@ -959,7 +953,6 @@ void FDeferredShadingSceneRenderer::RenderDistortion(
 	// Merge distortion back to scene color.
 	{
 		QUICK_SCOPE_CYCLE_COUNTER(STAT_FSceneRenderer_RenderDistortion_Merge);
-		RDG_EVENT_SCOPE(GraphBuilder, "Merge");
 		CommonParameters.SceneColorMSAATexture = DistortionSceneColorTexture;
 		CommonParameters.SceneColorTexture = DistortionSceneColorTexture;
 		CommonParameters.RenderTargets.DepthStencil = StencilWriteBinding;
@@ -983,7 +976,6 @@ void FDeferredShadingSceneRenderer::RenderDistortion(
 		{
 			const FViewInfo& View = Views[ViewIndex];
 			RDG_GPU_MASK_SCOPE(GraphBuilder, View.GPUMask);
-			RDG_EVENT_SCOPE_CONDITIONAL(GraphBuilder, Views.Num() > 1, "View%d", ViewIndex);
 
 			auto* PassParameters = GraphBuilder.AllocParameters<FDistortionScreenPS::FParameters>();
 			*PassParameters = CommonParameters;
@@ -999,7 +991,7 @@ void FDeferredShadingSceneRenderer::RenderDistortion(
 			Viewport.Rect = View.ViewRect;
 
 			ClearUnusedGraphResources(MergePixelShader, PassParameters);
-			AddDrawScreenPass(GraphBuilder, {}, FScreenPassViewInfo(View, 1), Viewport, Viewport, PipelineState, PassParameters,
+			AddDrawScreenPass(GraphBuilder, RDG_EVENT_NAME("Distortion::Merge View%d", ViewIndex), FScreenPassViewInfo(View, 1), Viewport, Viewport, PipelineState, PassParameters,
 				[MergePixelShader, PassParameters](FRHICommandList& RHICmdList)
 			{
 				SetShaderParameters(RHICmdList, MergePixelShader, MergePixelShader.GetPixelShader(), *PassParameters);
diff --git a/Engine/Source/Runtime/Renderer/Private/LightRendering.cpp b/Engine/Source/Runtime/Renderer/Private/LightRendering.cpp
--- a/Engine/Source/Runtime/Renderer/Private/LightRendering.cpp
+++ b/Engine/Source/Runtime/Renderer/Private/LightRendering.cpp
@@ -2962,7 +2962,7 @@ static void InternalRenderLight(
 	const ELightComponentType LightType = (ELightComponentType)LightProxy->GetLightType();
 
 	GraphBuilder.AddPass(
-		RDG_EVENT_NAME("%s: %s", ShaderName, *LightProxy->GetOwnerNameOrLabel()),
+		RDG_EVENT_NAME("RenderLight %s: %s", ShaderName, *LightProxy->GetOwnerNameOrLabel()),
 		PassParameters,
 		ERDGPassFlags::Raster,
 		[Scene, &View, PixelShader, UberPixelShader, LightSceneInfo, PassParameters, LightBounds, LightType, SubstrateTileMaterialType](FRDGAsyncTask, FRHICommandList& RHICmdList)
@@ -3379,7 +3379,6 @@ void FDeferredShadingSceneRenderer::RenderLightForHair(
 
 	SCOPE_CYCLE_COUNTER(STAT_DirectLightRenderingTime);
 	INC_DWORD_STAT(STAT_NumLightsUsingStandardDeferred);
-	RDG_EVENT_SCOPE(GraphBuilder, "StandardDeferredLighting_Hair");
 	RDG_GPU_MASK_SCOPE(GraphBuilder, View.GPUMask);
 
 	const bool bIsDirectional = LightSceneInfo->Proxy->GetLightType() == LightType_Directional;
@@ -3461,7 +3460,7 @@ void FDeferredShadingSceneRenderer::RenderLightForHair(
 	TShaderMapRef<FDeferredLightPS> PixelShader(View.ShaderMap, PermutationVector);
 
 	GraphBuilder.AddPass(
-		{},
+		RDG_EVENT_NAME("StandardDeferredLighting_Hair"),
 		PassParameters,
 		ERDGPassFlags::Raster,
 		[this, VertexShader, PixelShader, PermutationVector, &View, PassParameters, SampleLightingViewportResolution](FRDGAsyncTask, FRHICommandList& RHICmdList)
@@ -3972,8 +3971,6 @@ FRDGTextureRef FDeferredShadingSceneRenderer::CopyStencilToLightingChannelTextur
 
 	if (bNeedToCopyStencilToTexture)
 	{
-		RDG_EVENT_SCOPE(GraphBuilder, "CopyStencilToLightingChannels");
-
 		{
 			check(SceneStencilTexture && SceneStencilTexture->Desc.Texture);
 			const FIntPoint TextureExtent = SceneStencilTexture->Desc.Texture->Desc.Extent;
@@ -3988,7 +3985,6 @@ FRDGTextureRef FDeferredShadingSceneRenderer::CopyStencilToLightingChannelTextur
 		for (int32 ViewIndex = 0, ViewCount = Views.Num(); ViewIndex < ViewCount; ++ViewIndex)
 		{
 			const FViewInfo& View = Views[ViewIndex];
-			RDG_EVENT_SCOPE_CONDITIONAL(GraphBuilder, Views.Num() > 1, "View%d", ViewIndex);
 			RDG_GPU_MASK_SCOPE(GraphBuilder, View.GPUMask);
 
 			auto* PassParameters = GraphBuilder.AllocParameters<FCopyStencilToLightingChannelsPS::FParameters>();
@@ -4007,7 +4003,7 @@ FRDGTextureRef FDeferredShadingSceneRenderer::CopyStencilToLightingChannelTextur
 			PermutationVector.Set<FCopyStencilToLightingChannelsPS::FNaniteCompositeDim>(PassParameters->NaniteShadingMask != nullptr);
 			TShaderMapRef<FCopyStencilToLightingChannelsPS> PixelShader(View.ShaderMap, PermutationVector);
 
-			AddDrawScreenPass(GraphBuilder, {}, View, Viewport, Viewport, PixelShader, PassParameters);
+			AddDrawScreenPass(GraphBuilder, RDG_EVENT_NAME("CopyStencilToLightingChannels View%d", ViewIndex), View, Viewport, Viewport, PixelShader, PassParameters);
 		}
 	}
 
diff --git a/Engine/Source/Runtime/Renderer/Private/MobileDistortionPass.cpp b/Engine/Source/Runtime/Renderer/Private/MobileDistortionPass.cpp
--- a/Engine/Source/Runtime/Renderer/Private/MobileDistortionPass.cpp
+++ b/Engine/Source/Runtime/Renderer/Private/MobileDistortionPass.cpp
@@ -264,7 +264,6 @@ void FMobileSceneRenderer::RenderDistortion(
 	// Apply distortion and store off-screen.
 	{
 		QUICK_SCOPE_CYCLE_COUNTER(STAT_FSceneRenderer_RenderDistortion_Apply);
-		RDG_EVENT_SCOPE(GraphBuilder, "Apply");
 
 		// Test against stencil mask but don't clear.
 		PipelineState.DepthStencilState = TStaticDepthStencilState<
@@ -285,7 +284,7 @@ void FMobileSceneRenderer::RenderDistortion(
 		PassParameters->DistortionAccumulateTexture = DistortionTexture;
 		PassParameters->DistortionAccumulateSampler = TStaticSamplerState<SF_Bilinear, AM_Clamp, AM_Clamp, AM_Clamp>::GetRHI();
 
-		AddDrawScreenPass(GraphBuilder, {}, FScreenPassViewInfo(View, 1), Viewport, Viewport, PipelineState, PassParameters,
+		AddDrawScreenPass(GraphBuilder, RDG_EVENT_NAME("Distortion::Apply"), FScreenPassViewInfo(View, 1), Viewport, Viewport, PipelineState, PassParameters,
 			[MergePixelShader, PassParameters](FRHICommandList& RHICmdList)
 			{
 				SetShaderParameters(RHICmdList, MergePixelShader, MergePixelShader.GetPixelShader(), *PassParameters);
@@ -297,7 +296,6 @@ void FMobileSceneRenderer::RenderDistortion(
 	// Merge distortion back to scene color.
 	{
 		QUICK_SCOPE_CYCLE_COUNTER(STAT_FSceneRenderer_RenderDistortion_Merge);
-		RDG_EVENT_SCOPE(GraphBuilder, "Merge");
 			
 		TShaderMapRef<FCopyRectPS> CopyPixelShader(ShaderMap);
 
@@ -315,7 +313,7 @@ void FMobileSceneRenderer::RenderDistortion(
 		Parameters->RenderTargets[0] = FRenderTargetBinding(SceneColorTexture, ERenderTargetLoadAction::ENoAction);
 		Parameters->RenderTargets.DepthStencil = StencilWriteBinding;
 		Parameters->InputSampler = TStaticSamplerState<SF_Bilinear>::GetRHI();
-		AddDrawScreenPass(GraphBuilder, {}, FScreenPassViewInfo(View, 1), Viewport, Viewport, PipelineState, Parameters,
+		AddDrawScreenPass(GraphBuilder, RDG_EVENT_NAME("Distortion::Merge"), FScreenPassViewInfo(View, 1), Viewport, Viewport, PipelineState, Parameters,
 			[CopyPixelShader, Parameters](FRHICommandList& RHICmdList)
 			{
 				SetShaderParameters(RHICmdList, CopyPixelShader, CopyPixelShader.GetPixelShader(), *Parameters);
diff --git a/Engine/Source/Runtime/Renderer/Private/TranslucentRendering.cpp b/Engine/Source/Runtime/Renderer/Private/TranslucentRendering.cpp
--- a/Engine/Source/Runtime/Renderer/Private/TranslucentRendering.cpp
+++ b/Engine/Source/Runtime/Renderer/Private/TranslucentRendering.cpp
@@ -442,7 +442,7 @@ FRDGTextureRef AddCopySceneColorPass(FRDGBuilder& GraphBuilder, TArrayView<const
 				BlendState = TStaticBlendState<CW_RGBA, BO_Add, BF_Zero, BF_One, BO_Add, BF_Zero, BF_Zero>::GetRHI();
 			}
 
-			AddDrawScreenPass(GraphBuilder, {}, View, Viewport, Viewport, VertexShader, PixelShader, PassParameters);
+			AddDrawScreenPass(GraphBuilder, RDG_EVENT_NAME("TranslucencyCopySceneColor"), View, Viewport, Viewport, VertexShader, PixelShader, PassParameters);
 		}
 	}
 
@@ -512,7 +512,7 @@ static void AddCopySceneColorAlphaPass(FRDGBuilder& GraphBuilder, TArrayView<con
 		}
 
 		FRHIBlendState* BlendState = TStaticBlendState<CW_RGBA, BO_Add, BF_Zero, BF_One, BO_Add, BF_One, BF_Zero>::GetRHI();
-		AddDrawScreenPass(GraphBuilder, {}, View, Viewport, Viewport, VertexShader, PixelShader, BlendState, PassParameters);
+		AddDrawScreenPass(GraphBuilder, RDG_EVENT_NAME("TranslucencyCopySceneColorAlpha"), View, Viewport, Viewport, VertexShader, PixelShader, BlendState, PassParameters);
 	}
 }
 
