[TSR]  Fix memory spike at camera cut and teleport.

When it happens, memory allocation of history color texture 2d array is 1 while regular rendering is 2 due to motion blur is requried.
This reduces 31.6MB spike at 1080p and upscale to 4k without resurrection. and 126.6MB when resurrection is enabled. This also
fixs the quality fully drop issue when the teleport is very close.

bMotionBlurNeedsHalfResInput meaning changes from need to `would need`. The upscaler needs to decide if generating or not for memory/gpu usage.

#tests test on a sequence with multiple camera cuts. See memory spike eliminated
#Synced-CL 51948989

diff --git a/Engine/Source/Runtime/Renderer/Private/PostProcess/PostProcessing.cpp b/Engine/Source/Runtime/Renderer/Private/PostProcess/PostProcessing.cpp
--- a/Engine/Source/Runtime/Renderer/Private/PostProcess/PostProcessing.cpp
+++ b/Engine/Source/Runtime/Renderer/Private/PostProcess/PostProcessing.cpp
@@ -871,7 +871,7 @@ void AddPostProcessingPasses(
 
 		const bool bProcessQuarterResolution = CVarPostProcessingQuarterResolutionDownsample.GetValueOnRenderThread() == 1;
 		const bool bProcessEighthResolution = CVarPostProcessingQuarterResolutionDownsample.GetValueOnRenderThread() == 2;
-		const bool bMotionBlurNeedsHalfResInput = PassSequence.IsEnabled(EPass::MotionBlur) && DoesMotionBlurNeedsHalfResInput() && !bVisualizeMotionBlur;
+		const bool bMotionBlurNeedsHalfResInput = IsMotionBlurEnabled(View) && DoesMotionBlurNeedsHalfResInput() && !bVisualizeMotionBlur && !bVisualizeTSR;
 
 		const float FFTBloomResolutionFraction = GetFFTBloomResolutionFraction(PostTAAViewSize);
 
@@ -1109,6 +1109,7 @@ void AddPostProcessingPasses(
 			UpscalerPassInputs.SceneColor = FScreenPassTexture(SceneColorSlice);
 			UpscalerPassInputs.SceneDepth = FScreenPassTexture(SceneDepth.Texture, View.ViewRect);
 			UpscalerPassInputs.SceneVelocity = FScreenPassTexture(Velocity.Texture, View.ViewRect);
+			UpscalerPassInputs.bGenerateOutputMip1 = bMotionBlurNeedsHalfResInput;
 			if (PassSequence.IsEnabled(EPass::MotionBlur))
 			{
 				if (bVisualizeMotionBlur)
@@ -1117,7 +1118,6 @@ void AddPostProcessingPasses(
 				}
 				else
 				{
-					UpscalerPassInputs.bGenerateOutputMip1 = bMotionBlurNeedsHalfResInput;
 					UpscalerPassInputs.bGenerateVelocityFlattenTextures = FVelocityFlattenTextures::AllowExternal(View) && !bVisualizeMotionBlur && !bApplyLensDistortionInTSR;
 				}
 			}
diff --git a/Engine/Source/Runtime/Renderer/Private/PostProcess/TemporalAA.h b/Engine/Source/Runtime/Renderer/Private/PostProcess/TemporalAA.h
--- a/Engine/Source/Runtime/Renderer/Private/PostProcess/TemporalAA.h
+++ b/Engine/Source/Runtime/Renderer/Private/PostProcess/TemporalAA.h
@@ -276,6 +276,8 @@ struct FTSRPassConfig
 
 	float VelocityWeightClampingSampleCount = 4.0f;
 	float VelocityWeightClampingPixelSpeed = 1.0f;
+	
+	bool  MotionBlurValidInCurrentFrame = true;
 
 	int32  Visualize = -1;
 
diff --git a/Engine/Source/Runtime/Renderer/Private/PostProcess/TemporalSuperResolution.cpp b/Engine/Source/Runtime/Renderer/Private/PostProcess/TemporalSuperResolution.cpp
--- a/Engine/Source/Runtime/Renderer/Private/PostProcess/TemporalSuperResolution.cpp
+++ b/Engine/Source/Runtime/Renderer/Private/PostProcess/TemporalSuperResolution.cpp
@@ -1832,6 +1832,12 @@ FTSRPassConfig GetTSRMainPassConfig(const FViewInfo& View)
 
 	PassConfig.VelocityWeightClampingSampleCount = CVarTSRWeightClampingSampleCount.GetValueOnRenderThread();
 	PassConfig.VelocityWeightClampingPixelSpeed = CVarTSRWeightClampingPixelSpeed.GetValueOnRenderThread();
+
+	{
+		static const IConsoleVariable* CVar = IConsoleManager::Get().FindConsoleVariable(TEXT("r.MotionVectorSimulation"));
+		bool MotionVectorSimulation = CVar ? CVar->GetInt() != 0 : false;
+		PassConfig.MotionBlurValidInCurrentFrame = (MotionVectorSimulation || (!View.bCameraCut && !View.bPrevTransformsReset));
+	}
 #if UE_BUILD_OPTIMIZED_SHOWFLAGS
 	PassConfig.Visualize = false;
 #else
@@ -2136,11 +2142,14 @@ FDefaultTemporalUpscaler::FOutputs AddTemporalSuperResolutionPasses(
 	};
 
 	// Allocate a new history
+	bool bAllocateOutputMip1 = PassInputs.bGenerateOutputMip1;
+	bool bWriteOutputMip1 = PassInputs.bGenerateOutputMip1 && PassConfig.MotionBlurValidInCurrentFrame;
+
 	FTSRHistoryTextures History;
 	const int32 HistoryColorGuideSliceCountWithoutResurrection = bSupportsAlpha ? 2 : 1;
 	{
 		{
-			bool bRequires2Mips = HistorySize == OutputRect.Size() && PassInputs.bGenerateOutputMip1;
+			bool bRequires2Mips = HistorySize == OutputRect.Size() && bAllocateOutputMip1;
 			FIntPoint MipClampedHistoryExtent = FIntPoint(FMath::Max(HistoryExtent.X, bRequires2Mips ? 2: 1), 
 				                                          FMath::Max(HistoryExtent.Y, bRequires2Mips ? 2: 1));
 			FRDGTextureDesc ArrayDesc = FRDGTextureDesc::Create2DArray(
@@ -3211,7 +3220,7 @@ FDefaultTemporalUpscaler::FOutputs AddTemporalSuperResolutionPasses(
 				History.ColorArray, ColorUAVDesc.FirstArraySlice + HistoryArrayIndices.HighFrequency));
 		}
 
-		if (PassInputs.bGenerateOutputMip1 && HistorySize == OutputRect.Size())
+		if (bWriteOutputMip1 && HistorySize == OutputRect.Size())
 		{
 			FRDGTextureUAVDesc Mip1Desc(History.ColorArray);
 			Mip1Desc.MipLevel = 1;
@@ -3321,21 +3330,21 @@ FDefaultTemporalUpscaler::FOutputs AddTemporalSuperResolutionPasses(
 		
 		FRDGTextureRef SceneColorOutputTexture;
 		{
-			FIntPoint MipClampedOutputExtent = FIntPoint(FMath::Max(OutputExtent.X, PassInputs.bGenerateOutputMip1 ? 2 : 1), 
-				                                         FMath::Max(OutputExtent.Y, PassInputs.bGenerateOutputMip1 ? 2 : 1));
+			FIntPoint MipClampedOutputExtent = FIntPoint(FMath::Max(OutputExtent.X, bAllocateOutputMip1 ? 2 : 1), 
+				                                         FMath::Max(OutputExtent.Y, bAllocateOutputMip1 ? 2 : 1));
 			FRDGTextureDesc OutputDesc = FRDGTextureDesc::Create2D(
 				MipClampedOutputExtent,
 				ColorFormat,
 				FClearValueBinding::None,
 				/* InFlags = */ TexCreate_ShaderResource | TexCreate_UAV | TexCreate_RenderTargetable,
-				/* NumMips = */ PassInputs.bGenerateOutputMip1 ? 2 : 1);
+				/* NumMips = */ bAllocateOutputMip1 ? 2 : 1);
 			SceneColorOutputTexture = GraphBuilder.CreateTexture(OutputDesc, TEXT("TSR.Output"));
 
 			PassParameters->SceneColorOutputMip0 = GraphBuilder.CreateUAV(FRDGTextureUAVDesc(SceneColorOutputTexture, /* InMipLevel = */ 0));
 			SceneColorOutputTextureSRV = GraphBuilder.CreateSRV(FRDGTextureSRVDesc(SceneColorOutputTexture));
 		}
 
-		if (PassInputs.bGenerateOutputMip1)
+		if (bWriteOutputMip1)
 		{
 			PassParameters->bGenerateOutputMip1 = true;
 			PassParameters->SceneColorOutputMip1 = GraphBuilder.CreateUAV(FRDGTextureUAVDesc(SceneColorOutputTexture, /* InMipLevel = */ 1));
