[Engine][Animatioon] ASTP Reference Pose Raw vs Compressed
- When setting up sequence data ASTP uses RawNumberOfBones but when unpacking it uses Num this can lead to a mismatch and crash as the sequence data accesses the array out of bounds.

#rnx
#Synced-CL 52134719

diff --git a/Engine/Source/Runtime/Renderer/Private/Skinning/AnimSequenceTransformProvider.cpp b/Engine/Source/Runtime/Renderer/Private/Skinning/AnimSequenceTransformProvider.cpp
--- a/Engine/Source/Runtime/Renderer/Private/Skinning/AnimSequenceTransformProvider.cpp
+++ b/Engine/Source/Runtime/Renderer/Private/Skinning/AnimSequenceTransformProvider.cpp
@@ -403,7 +403,7 @@ void FAnimSequenceTransformProvider::BuildBoneTables(
 	const USkeleton* TargetSkeleton = SkinnedMesh->GetSkeleton();
 	const FReferenceSkeleton& MeshRefSkeleton = SkinnedMesh->GetRefSkeleton();
 	const FReferenceSkeleton& TargetSkeletonRefSkeleton = TargetSkeleton->GetReferenceSkeleton();
-	const uint32 NumMeshBones = MeshRefSkeleton.GetRawBoneNum();
+	const uint32 NumMeshBones = MeshRefSkeleton.GetNum();
 	const uint32 NumCompactBones = LODRenderData.RequiredBones.Num();
 
 	BoneTables.MeshToCompact.Init(INVALID_BONE_INDEX, NumMeshBones);
@@ -451,7 +451,7 @@ void FAnimSequenceTransformProvider::BuildRetargetingTables(
 
 	const FSkeletonRemapping& SkeletonRemapping = UE::Anim::FSkeletonRemappingRegistry::Get().GetRemapping(AnimSkeleton, TargetSkeleton);
 	const bool bRequiresRefPoseRetarget = SkeletonRemapping.IsValid() && SkeletonRemapping.RequiresReferencePoseRetarget();
-	const uint32 NumAnimBones = AnimSkeleton->GetReferenceSkeleton().GetRawBoneNum();
+	const uint32 NumAnimBones = AnimSkeleton->GetReferenceSkeleton().GetNum();
 
 	RetargetingTables.Mode.Reset(BoneTables.NumCompactBones);
 	RetargetingTables.Transform.Reset(BoneTables.NumCompactBones);
@@ -747,7 +747,7 @@ void FAnimSequenceTransformProvider::RegisterProxy(FAnimSequenceTransformProvide
 			Sequence->PlayLength             = AnimSequence->GetPlayLength();
 			Sequence->SampleRate             = GetAnimSequenceSampleRate(AnimSequence);
 			Sequence->NumFrames              = Sequence->SampleRate.AsFrameTime(Sequence->PlayLength).RoundToFrame().Value;
-			Sequence->NumAnimBones           = SourceSkeleton->GetReferenceSkeleton().GetRawBoneNum();
+			Sequence->NumAnimBones           = SourceSkeleton->GetReferenceSkeleton().GetNum();
 			Sequence->NumSampleRateLevels    = GetMaxSampleRateLevels(Sequence->NumFrames);
 			Sequence->CurrentSampleRateLevel = Sequence->NumSampleRateLevels;
 			Sequence->TargetSampleRateLevel  = Sequence->NumSampleRateLevels - 1;
