Ray Tracing: While using the fallback RTPSO, display an on-screen warning and only log once

This makes it more obvious that the render is not ready yet, which is particularly useful for path tracing which will display a black screen until the RTPSO is done compiling

#Synced-CL 52135942

diff --git a/Engine/Source/Runtime/Renderer/Private/RayTracing/RayTracingMaterialHitShaders.cpp b/Engine/Source/Runtime/Renderer/Private/RayTracing/RayTracingMaterialHitShaders.cpp
--- a/Engine/Source/Runtime/Renderer/Private/RayTracing/RayTracingMaterialHitShaders.cpp
+++ b/Engine/Source/Runtime/Renderer/Private/RayTracing/RayTracingMaterialHitShaders.cpp
@@ -790,6 +790,7 @@ void FDeferredShadingSceneRenderer::CreateMaterialRayTracingMaterialPipeline(
 	{
 		// Save the current pipeline to be used as fallback in future frames
 		Scene->LastRayTracingMaterialPipelineSignature = static_cast<FRayTracingPipelineStateSignature&>(Initializer);
+        Scene->bLastFrameUsedFallback = false;
 	}
 	else
 	{
@@ -797,7 +798,18 @@ void FDeferredShadingSceneRenderer::CreateMaterialRayTracingMaterialPipeline(
 		check(FallbackPipelineState);
 		PipelineState = FallbackPipelineState;
 		bOutIsUsingFallbackRTPSO = true;
-		UE_LOG(LogRenderer, Log, TEXT("Using fallback RTPSO"));
+#if WITH_EDITOR
+        static const FText Message = NSLOCTEXT("Renderer", "RayTracingPSONotReady", "Ray Tracing PSO is not ready yet, using fallback.");
+		OnGetOnScreenMessages.AddLambda([](FScreenMessageWriter& ScreenMessageWriter) -> void
+        {
+            ScreenMessageWriter.DrawLine(Message);
+        });
+        if (!Scene->bLastFrameUsedFallback)
+        {
+            UE_LOGF(LogRenderer, Log, "%ls", *Message.ToString());
+            Scene->bLastFrameUsedFallback = true;
+        }
+#endif
 	}
 
 	// Retrieve the binding data size from the actual used RTPSO because the requested RTPSO could still be non blocking async compiling
diff --git a/Engine/Source/Runtime/Renderer/Private/ScenePrivate.h b/Engine/Source/Runtime/Renderer/Private/ScenePrivate.h
--- a/Engine/Source/Runtime/Renderer/Private/ScenePrivate.h
+++ b/Engine/Source/Runtime/Renderer/Private/ScenePrivate.h
@@ -1914,6 +1914,8 @@ public:
 	// Last valid RTPSO is saved, so it could be used as fallback in future frames if background PSO compilation is enabled.
 	// This RTPSO can be used only if the only difference from previous PSO is the material hit shaders.
 	FRayTracingPipelineStateSignature LastRayTracingMaterialPipelineSignature;
+
+    bool bLastFrameUsedFallback = false;
 	
 	/** SBT object to be used with hardware ray tracing */
 	FRayTracingShaderBindingTable RayTracingSBT;
