Virtual texture: Deal with case of overlapping virtual texture feedback scopes.
This can happen during material editing where the preview window opens a scope and the individual canvas nodes each generate a DrawTileMesh scope.
#Synced-CL 52125649

diff --git a/Engine/Source/Runtime/Renderer/Private/VT/VirtualTextureFeedbackResource.cpp b/Engine/Source/Runtime/Renderer/Private/VT/VirtualTextureFeedbackResource.cpp
--- a/Engine/Source/Runtime/Renderer/Private/VT/VirtualTextureFeedbackResource.cpp
+++ b/Engine/Source/Runtime/Renderer/Private/VT/VirtualTextureFeedbackResource.cpp
@@ -60,7 +60,7 @@ public:
 	//~ End FRenderResource Interface
 
 private:
-	bool bIsInBeginEndScope = false;
+	int32 BeginEndScopeCount = 0;
 	uint32 FeedbackBufferSize = 0;
 	uint32 ExtendedDebugBufferSize = 0;
 	ERHIFeatureLevel::Type FeatureLevel = ERHIFeatureLevel::Num;
@@ -79,6 +79,14 @@ void FVirtualTextureFeedbackBufferResource::Begin(FRDGBuilder& GraphBuilder, uin
 	// the view uniform buffer, which is not an RDG uniform buffer. If it can be factored out into its own RDG
 	// uniform buffer (or put on the pass uniform buffers), then the resource can be fully converted to RDG.
 
+	if (++BeginEndScopeCount != 1)
+	{
+		// Where overlapping scopes occur we will assume that it is OK to use the UAV from the outer scope.
+		// This remains correct, but not ideal, if the inner scope requests a bigger buffer size then the outer scope.
+		// In that case the feedback shader still respects the smaller size when writing, but we are more likely to fill the buffer and lose feedback.
+		return;
+	}
+
 	FeedbackBufferSize = InFeedbackBufferSize;
 	ExtendedDebugBufferSize = InExtendedDebugBufferSize;
 	FeatureLevel = InFeatureLevel;
@@ -99,17 +107,19 @@ void FVirtualTextureFeedbackBufferResource::Begin(FRDGBuilder& GraphBuilder, uin
 		RHICmdList.Transition(FRHITransitionInfo(UAV, ERHIAccess::UAVCompute, ERHIAccess::UAVMask));
 		RHICmdList.BeginUAVOverlap(UAV);
 	});
-
-	bIsInBeginEndScope = true;
 }
 
 void FVirtualTextureFeedbackBufferResource::End(FRDGBuilder& GraphBuilder)
 {
-	if (!bIsInBeginEndScope)
+	if (--BeginEndScopeCount != 0)
 	{
+		if (!ensure(BeginEndScopeCount > 0))
+		{ 
+			// Calling End() more often than Begin() should never happen. Recover by resetting the count.
+			BeginEndScopeCount = 0;
+		}
 		return;
 	}
-	bIsInBeginEndScope = false;
 
 	// VirtualTextureFeedback would be a more descriptive stat name, but VirtualTextureUpdate was used historically and some profile tools may depend on that.
 	RDG_EVENT_SCOPE_STAT(GraphBuilder, VirtualTextureUpdate, "VirtualTextureUpdate");
@@ -253,7 +263,7 @@ void FVirtualTextureFeedbackBufferResource::End(FRDGBuilder& GraphBuilder)
 
 FRDGBuffer* FVirtualTextureFeedbackBufferResource::ResolveExtendedDebugBuffer(FRDGBuilder& GraphBuilder)
 {
-	if (!bIsInBeginEndScope || ExtendedDebugBufferSize == 0)
+	if (BeginEndScopeCount <= 0 || ExtendedDebugBufferSize == 0)
 	{
 		return nullptr;
 	}
@@ -285,17 +295,17 @@ FRDGBuffer* FVirtualTextureFeedbackBufferResource::ResolveExtendedDebugBuffer(FR
 
 uint32 FVirtualTextureFeedbackBufferResource::GetBufferSize() const
 {
-	return bIsInBeginEndScope ? FeedbackBufferSize : 0;
+	return BeginEndScopeCount > 0 ? FeedbackBufferSize : 0;
 }
 
 uint32 FVirtualTextureFeedbackBufferResource::GetExtendedDebugBufferSize() const
 {
-	return bIsInBeginEndScope ? ExtendedDebugBufferSize : 0;
+	return BeginEndScopeCount > 0 ? ExtendedDebugBufferSize : 0;
 }
 
 FRHIUnorderedAccessView* FVirtualTextureFeedbackBufferResource::GetUAV() const
 {
-	return UAV != nullptr && bIsInBeginEndScope ? UAV : GEmptyStructuredBufferWithUAV->UnorderedAccessViewRHI.GetReference();
+	return UAV != nullptr && BeginEndScopeCount > 0 ? UAV : GEmptyStructuredBufferWithUAV->UnorderedAccessViewRHI.GetReference();
 }
 
 void FVirtualTextureFeedbackBufferResource::ReleaseRHI()
