Enable LOD dithering visualization for dynamic meshes in LODColoration mode

LODColoration uses dynamic mesh batches to override materials for visualization.
Since these dynamic batches don't have StaticMeshId entries in StaticMeshFadeOutDitheredLODMap, we cache the dithering state directly in FMeshBatch.

- Added EDitherState enum and CachedLODDither to FMeshBatch
- GetDynamicMeshElements calculates dithering state from LODMask
- GetOptionalDitheringFlag checks cached value before static mesh lookup

#tests Editor preview modes

diff --git a/Engine/Source/Runtime/Engine/Private/DebugViewModeInterface.cpp b/Engine/Source/Runtime/Engine/Private/DebugViewModeInterface.cpp
--- a/Engine/Source/Runtime/Engine/Private/DebugViewModeInterface.cpp
+++ b/Engine/Source/Runtime/Engine/Private/DebugViewModeInterface.cpp
@@ -13,7 +13,7 @@ DebugViewModeInterface.cpp: Contains definitions for rendering debug viewmodes.
 
 FDebugViewModeInterface* FDebugViewModeInterface::Singleton = nullptr;
 
-void FDebugViewModeInterface::SetDrawRenderState(EDebugViewShaderMode DebugViewMode, EBlendMode InBlendMode, FRenderState& DrawRenderState, bool bHasDepthPrepassForMaskedMaterial) const
+void FDebugViewModeInterface::SetDrawRenderState(EDebugViewShaderMode DebugViewMode, EBlendMode InBlendMode, FRenderState& DrawRenderState, bool bHasDepthPrepassForMaskedMaterial, bool bDitheredLODTransition) const
 {
 	// No writes to RT1
 	FBlendStateRHIRef DefaultOpaqueBlendState = TStaticBlendState<
@@ -94,6 +94,10 @@ void FDebugViewModeInterface::SetDrawRenderState(EDebugViewShaderMode DebugViewM
 			DrawRenderState.BlendState = TStaticBlendState<CW_RGBA, BO_Add, BF_One, BF_Zero, BO_Add, BF_One, BF_Zero, CW_RED>::GetRHI();
 		}
 	}
+	else if (DebugViewMode == DVSM_LODColoration && bDitheredLODTransition)
+	{
+		DrawRenderState.DepthStencilState = TStaticDepthStencilState<false, CF_Equal>::GetRHI();
+	}
 }
 
 void FDebugViewModeInterface::SetInterface(FDebugViewModeInterface* Interface)
diff --git a/Engine/Source/Runtime/Engine/Private/StaticMeshSceneProxy.cpp b/Engine/Source/Runtime/Engine/Private/StaticMeshSceneProxy.cpp
--- a/Engine/Source/Runtime/Engine/Private/StaticMeshSceneProxy.cpp
+++ b/Engine/Source/Runtime/Engine/Private/StaticMeshSceneProxy.cpp
@@ -1717,6 +1717,7 @@ void FStaticMeshSceneProxy::GetDynamicMeshElements(const TArray<const FSceneView
 		for (int32 ViewIndex = 0; ViewIndex < Views.Num(); ViewIndex++)
 		{
 			const FSceneView* View = Views[ViewIndex];
+			const bool bLODColoration = View->Family->EngineShowFlags.LODColoration;
 
 			if (IsShown(View) && (VisibilityMap & (1 << ViewIndex)))
 			{
@@ -1861,7 +1862,19 @@ void FStaticMeshSceneProxy::GetDynamicMeshElements(const TArray<const FSceneView
 	#endif
 										if (MeshElement.bDitheredLODTransition && LODMask.IsDithered())
 										{
-
+#if UE_ENABLE_DEBUG_DRAWING
+											if(bLODColoration)
+											{ 
+												if (LODIndex == LODMask.LODIndex0)
+												{
+													MeshElement.CachedLODDither = FMeshBatch::EDitherState::FadeOut;
+												}
+												else if (LODIndex == LODMask.LODIndex1)
+												{
+													MeshElement.CachedLODDither = FMeshBatch::EDitherState::FadeIn;
+												}
+											}
+#endif										
 										}
 										else
 										{
diff --git a/Engine/Source/Runtime/Engine/Public/DebugViewModeInterface.h b/Engine/Source/Runtime/Engine/Public/DebugViewModeInterface.h
--- a/Engine/Source/Runtime/Engine/Public/DebugViewModeInterface.h
+++ b/Engine/Source/Runtime/Engine/Public/DebugViewModeInterface.h
@@ -47,7 +47,7 @@ public:
 		const FVertexFactoryType* InVertexFactoryType,
 		FMaterialShaderTypes& OutShaderTypes) const = 0;
 
-	ENGINE_API virtual void SetDrawRenderState(EDebugViewShaderMode DebugViewMode, EBlendMode BlendMode, FRenderState& DrawRenderState, bool bHasDepthPrepassForMaskedMaterial) const;
+	ENGINE_API virtual void SetDrawRenderState(EDebugViewShaderMode DebugViewMode, EBlendMode BlendMode, FRenderState& DrawRenderState, bool bHasDepthPrepassForMaskedMaterial, bool bDitheredLODTransition = false) const;
 
 	virtual void GetDebugViewModeShaderBindings(
 		const FDebugViewModePS& Shader,
diff --git a/Engine/Source/Runtime/Engine/Public/MeshBatch.h b/Engine/Source/Runtime/Engine/Public/MeshBatch.h
--- a/Engine/Source/Runtime/Engine/Public/MeshBatch.h
+++ b/Engine/Source/Runtime/Engine/Public/MeshBatch.h
@@ -459,6 +459,17 @@ struct FMeshBatch
 
 	/** Conceptual LOD index used for the LOD Coloration visualization. */
 	int8 VisualizeLODIndex;
+
+	/** Dithering direction for LOD transitions in debug visualization */
+	enum class EDitherState : uint8
+	{
+		None = 0,
+		FadeOut = 1,
+		FadeIn = 2
+	};
+
+	/** Cached dithering state value for LOD transitions. Computed in GetDynamicMeshElements. */
+	EDitherState CachedLODDither = EDitherState::None;
 #endif
 
 	ENGINE_API bool IsTranslucent(ERHIFeatureLevel::Type InFeatureLevel) const;
diff --git a/Engine/Source/Runtime/Renderer/Private/DebugViewModeRendering.cpp b/Engine/Source/Runtime/Renderer/Private/DebugViewModeRendering.cpp
--- a/Engine/Source/Runtime/Renderer/Private/DebugViewModeRendering.cpp
+++ b/Engine/Source/Runtime/Renderer/Private/DebugViewModeRendering.cpp
@@ -507,11 +507,11 @@ void FDebugViewModeMeshProcessor::AddMeshBatch(const FMeshBatch& RESTRICT MeshBa
 	const FMeshDrawingPolicyOverrideSettings OverrideSettings = ComputeMeshOverrideSettings(MeshBatch);
 	const ERasterizerFillMode MeshFillMode = ComputeMeshFillMode(*BatchMaterial, OverrideSettings);
 	const ERasterizerCullMode MeshCullMode = ComputeMeshCullMode(*BatchMaterial, OverrideSettings);
+	const bool bDitheredLODTransition = MeshBatch.bDitheredLODTransition && MeshBatch.CachedLODDither != FMeshBatch::EDitherState::None;
 
 	FMeshPassProcessorRenderState DrawRenderState;
-
 	FDebugViewModeInterface::FRenderState InterfaceRenderState;
-	DebugViewModeInterface->SetDrawRenderState(DebugViewMode, BatchMaterial->GetBlendMode(), InterfaceRenderState, Scene ? (Scene->GetShadingPath() == EShadingPath::Deferred && Scene->EarlyZPassMode != DDM_NonMaskedOnly) : false);
+	DebugViewModeInterface->SetDrawRenderState(DebugViewMode, BatchMaterial->GetBlendMode(), InterfaceRenderState, Scene ? (Scene->GetShadingPath() == EShadingPath::Deferred && Scene->EarlyZPassMode != DDM_NonMaskedOnly) : false, bDitheredLODTransition);
 	DrawRenderState.SetBlendState(InterfaceRenderState.BlendState);
 	DrawRenderState.SetDepthStencilState(InterfaceRenderState.DepthStencilState);
 
diff --git a/Engine/Source/Runtime/Renderer/Private/DepthRendering.cpp b/Engine/Source/Runtime/Renderer/Private/DepthRendering.cpp
--- a/Engine/Source/Runtime/Renderer/Private/DepthRendering.cpp
+++ b/Engine/Source/Runtime/Renderer/Private/DepthRendering.cpp
@@ -220,6 +220,20 @@ TOptional<uint32> GetOptionalDitheringFlag(
 	const FMeshBatch& RESTRICT Mesh,
 	int32 StaticMeshId)
 {
+#if UE_ENABLE_DEBUG_DRAWING
+	// Check cached value from GetDynamicMeshElements first for debug view mode.
+	if (Mesh.CachedLODDither != FMeshBatch::EDitherState::None)
+	{
+		if (Mesh.CachedLODDither == FMeshBatch::EDitherState::FadeOut)
+		{
+			return STENCIL_SANDBOX_MASK;
+		}
+		else if (Mesh.CachedLODDither == FMeshBatch::EDitherState::FadeIn)
+		{
+			return 0;
+		}
+	}
+#endif
 	if (SceneView && StaticMeshId >= 0 && Mesh.bDitheredLODTransition)
 	{
 		checkSlow(SceneView->bIsViewInfo);
