Fix HZB ViewRect extending beyond view depth in split-screen

#Synced-CL 52149918

diff --git a/Engine/Source/Runtime/Renderer/Private/DeferredShadingRenderer.cpp b/Engine/Source/Runtime/Renderer/Private/DeferredShadingRenderer.cpp
--- a/Engine/Source/Runtime/Renderer/Private/DeferredShadingRenderer.cpp
+++ b/Engine/Source/Runtime/Renderer/Private/DeferredShadingRenderer.cpp
@@ -568,17 +568,11 @@ bool FDeferredShadingSceneRenderer::RenderHzb(FRDGBuilder& GraphBuilder, FRDGTex
 			FRDGTextureRef ClosestHZBTexture = nullptr;
 			FRDGTextureRef FurthestHZBTexture = nullptr;
 
-			// In split-screen, multiple views share a depth buffer. UpperBoundViewRect (used by
-			// GetHZBViewRect) can extend beyond a view's rendered depth region into an adjacent
-			// view's area, contaminating the HZB mipchain. Clamp to ViewRect in multi-view to
-			// prevent reading another view's depth. See UE-347165.
-			const FIntRect HZBViewRect = Views.Num() > 1 ? View.ViewRect : GetHZBViewRect(View);
-
 			BuildHZB(
 				GraphBuilder,
 				SceneDepthTexture,
 				/* VisBufferTexture = */ nullptr,
-				HZBViewRect,
+				GetHZBViewRect(View),
 				View.GetFeatureLevel(),
 				View.GetShaderPlatform(),
 				TEXT("HZBClosest"),
diff --git a/Engine/Source/Runtime/Renderer/Private/MobileShadingRenderer.cpp b/Engine/Source/Runtime/Renderer/Private/MobileShadingRenderer.cpp
--- a/Engine/Source/Runtime/Renderer/Private/MobileShadingRenderer.cpp
+++ b/Engine/Source/Runtime/Renderer/Private/MobileShadingRenderer.cpp
@@ -3435,17 +3435,11 @@ void FMobileSceneRenderer::RenderHZB(FRDGBuilder& GraphBuilder, FRDGTextureRef S
 
 				FRDGTextureRef FurthestHZBTexture = nullptr;
 
-				// In split-screen, multiple views share a depth buffer. UpperBoundViewRect (used by
-				// GetHZBViewRect) can extend beyond a view's rendered depth region into an adjacent
-				// view's area, contaminating the HZB mipchain. Clamp to ViewRect in multi-view to
-				// prevent reading another view's depth. See UE-347165.
-				const FIntRect HZBViewRect = Views.Num() > 1 ? View.ViewRect : GetHZBViewRect(View);
-
 				BuildHZBFurthest(
 					GraphBuilder,
 					SceneDepthTexture,
 					/* VisBufferTexture = */ nullptr,
-					HZBViewRect,
+					GetHZBViewRect(View),
 					View.GetFeatureLevel(),
 					View.GetShaderPlatform(),
 					TEXT("MobileHZBFurthest"),
diff --git a/Engine/Source/Runtime/Renderer/Private/SceneRendering.cpp b/Engine/Source/Runtime/Renderer/Private/SceneRendering.cpp
--- a/Engine/Source/Runtime/Renderer/Private/SceneRendering.cpp
+++ b/Engine/Source/Runtime/Renderer/Private/SceneRendering.cpp
@@ -3572,14 +3572,20 @@ void FSceneRenderer::PrepareViewRectsForRendering()
 	{
 		FViewInfo& View = Views[i];
 
-		auto InitViewRect = [&View, &MaxOverscanResolutionFraction, ViewFamilySecondaryViewFraction = ViewFamily.SecondaryViewFraction](float CurrentResFraction)
+		auto InitViewRect = [&View, &MaxOverscanResolutionFraction, ViewFamilySecondaryViewFraction = ViewFamily.SecondaryViewFraction](float PositionResFraction, float SizeResFraction)
 			{
-				const float ViewResolutionFraction = View.SceneViewInitOptions.OverridePrimaryResolutionFraction > 0.0 ? (View.SceneViewInitOptions.OverridePrimaryResolutionFraction * ViewFamilySecondaryViewFraction) : CurrentResFraction;
+				auto GetViewResFraction = [&](float Fraction)
+				{
+					return View.SceneViewInitOptions.OverridePrimaryResolutionFraction > 0.0 ? (View.SceneViewInitOptions.OverridePrimaryResolutionFraction * ViewFamilySecondaryViewFraction) : Fraction;
+				};
+
+				const float ViewPositionFraction = GetViewResFraction(PositionResFraction);
+				const float ViewSizeFraction = GetViewResFraction(SizeResFraction);
 
-				FIntPoint ViewSize = ApplyResolutionFraction(View.UnscaledViewRect.Size(), ViewResolutionFraction * View.SceneViewInitOptions.OverscanResolutionFraction);
+				FIntPoint ViewSize = ApplyResolutionFraction(View.UnscaledViewRect.Size(), ViewSizeFraction * View.SceneViewInitOptions.OverscanResolutionFraction);
 				FIntPoint ViewRectMin = QuantizeViewRectMin(FIntPoint(
-					FMath::CeilToInt(View.UnscaledViewRect.Min.X * ViewResolutionFraction * MaxOverscanResolutionFraction),
-					FMath::CeilToInt(View.UnscaledViewRect.Min.Y * ViewResolutionFraction * MaxOverscanResolutionFraction)));
+					FMath::CeilToInt(View.UnscaledViewRect.Min.X * ViewPositionFraction * MaxOverscanResolutionFraction),
+					FMath::CeilToInt(View.UnscaledViewRect.Min.Y * ViewPositionFraction * MaxOverscanResolutionFraction)));
 
 				// Use the bottom-left view rect if requested, instead of top-left
 				if (CVarViewRectUseScreenBottom.GetValueOnRenderThread())
@@ -3590,8 +3596,17 @@ void FSceneRenderer::PrepareViewRectsForRendering()
 				return FIntRect(ViewRectMin, ViewRectMin + ViewSize);
 			};
 
-		View.ViewRect = InitViewRect(ResolutionFraction);
-		View.UpperBoundViewRect = InitViewRect(ResolutionMaxFraction);
+		if (Views.Num() > 1)
+		{
+			// In split-screen, multiple views share a depth buffer.  Use ResolutionMaxFraction for position so that adjacent views' UpperBoundViewRects don't overlap. 
+			View.ViewRect = InitViewRect(ResolutionMaxFraction, ResolutionFraction);
+			View.UpperBoundViewRect = InitViewRect(ResolutionMaxFraction, ResolutionMaxFraction);
+		}
+		else
+		{
+			View.ViewRect = InitViewRect(ResolutionFraction, ResolutionFraction);
+			View.UpperBoundViewRect = InitViewRect(ResolutionMaxFraction, ResolutionMaxFraction);
+		}
 
 		#if !UE_BUILD_SHIPPING
 		// For testing purpose, override the screen percentage method.
@@ -3648,6 +3663,34 @@ void FSceneRenderer::PrepareViewRectsForRendering()
 		}
 	}
 
+#if DO_CHECK
+	// Verify that ViewRect is contained within UpperBoundViewRect, and that no two UpperBoundViewRects overlap in the shared depth buffer.
+	for (int32 i = 0; i < Views.Num(); i++)
+	{
+		const FIntRect& ViewRect = Views[i].ViewRect;
+		const FIntRect& UpperBound = Views[i].UpperBoundViewRect;
+
+		checkf(ViewRect.Min.X >= UpperBound.Min.X && ViewRect.Min.Y >= UpperBound.Min.Y &&
+			ViewRect.Max.X <= UpperBound.Max.X && ViewRect.Max.Y <= UpperBound.Max.Y,
+			TEXT("View %d ViewRect [%d,%d]-[%d,%d] is not contained within UpperBoundViewRect [%d,%d]-[%d,%d]"),
+			i, ViewRect.Min.X, ViewRect.Min.Y, ViewRect.Max.X, ViewRect.Max.Y,
+			UpperBound.Min.X, UpperBound.Min.Y, UpperBound.Max.X, UpperBound.Max.Y);
+
+		for (int32 j = i + 1; j < Views.Num(); j++)
+		{
+			const FIntRect& OtherUpperBound = Views[j].UpperBoundViewRect;
+
+			const bool bOverlaps = UpperBound.Min.X < OtherUpperBound.Max.X && UpperBound.Max.X > OtherUpperBound.Min.X &&
+				UpperBound.Min.Y < OtherUpperBound.Max.Y && UpperBound.Max.Y > OtherUpperBound.Min.Y;
+
+			checkf(!bOverlaps,
+				TEXT("View %d UpperBoundViewRect [%d,%d]-[%d,%d] overlaps View %d UpperBoundViewRect [%d,%d]-[%d,%d]"),
+				i, UpperBound.Min.X, UpperBound.Min.Y, UpperBound.Max.X, UpperBound.Max.Y,
+				j, OtherUpperBound.Min.X, OtherUpperBound.Min.Y, OtherUpperBound.Max.X, OtherUpperBound.Max.Y);
+		}
+	}
+#endif
+
 	#if !UE_BUILD_SHIPPING
 	{
 		int32 ViewRectOffset = CVarTestInternalViewRectOffset.GetValueOnRenderThread();
