Path Tracing: Make sure path tracer is not impacted by r.RayTracing.ExcludeTranslucent cvar

The cvar is useful to optimize performance for the deferred renderer, but may remove important geometry from the path tracer's view (in particular mesh decals which by definition use translucent blend mode).

This is similar to the policy for r.RayTracing.ExcludeSky which is also ignored by the path tracer.

#Synced-CL 52157500

diff --git a/Engine/Source/Runtime/Renderer/Private/RayTracing/RayTracing.cpp b/Engine/Source/Runtime/Renderer/Private/RayTracing/RayTracing.cpp
--- a/Engine/Source/Runtime/Renderer/Private/RayTracing/RayTracing.cpp
+++ b/Engine/Source/Runtime/Renderer/Private/RayTracing/RayTracing.cpp
@@ -135,7 +135,7 @@ static int32 GRayTracingExcludeTranslucent = 0;
 static FAutoConsoleVariableRef CRayTracingExcludeTranslucent(
 	TEXT("r.RayTracing.ExcludeTranslucent"),
 	GRayTracingExcludeTranslucent,
-	TEXT("A toggle that modifies the inclusion of translucent objects in the ray tracing scene.\n")
+	TEXT("A toggle that modifies the inclusion of translucent objects in the ray tracing scene. This setting is ignored for the Path Tracer.\n")
 	TEXT(" 0: Translucent objects included in the ray tracing scene (default)\n")
 	TEXT(" 1: Translucent objects excluded from the ray tracing scene"),
 	ECVF_RenderThreadSafe);
@@ -281,12 +281,6 @@ namespace RayTracing
 			bTranslucentGeometry = true;
 		}
 
-		if (ViewFamily.EngineShowFlags.PathTracing
-			&& FDataDrivenShaderPlatformInfo::GetSupportsPathTracing(Scene.GetShaderPlatform()))
-		{
-			bTranslucentGeometry = true;
-		}
-
 		if (ViewFamily.EngineShowFlags.RayTracingDebug && !IsRayTracingVisualizeOpaqueOnly())
 		{
 			if (!bTranslucentGeometry)
@@ -303,6 +297,12 @@ namespace RayTracing
 			bTranslucentGeometry = false;
 		}
 
+        if (ViewFamily.EngineShowFlags.PathTracing)
+		{
+            // When path tracing, always include translucent geometry
+			bTranslucentGeometry = true;
+		}
+
 		bIncludeSky = GRayTracingExcludeSky == 0 || ViewFamily.EngineShowFlags.PathTracing;
 
 		bLightingChannelsUsingAHS = MegaLights::IsEnabled(ViewFamily) && MegaLights::IsUsingLightingChannels();
