* Remove hair debug permutation from hair strands env lighting
* Remove hair bsdf plotting shaders.

#Synced-CL 52165209

diff --git a/Engine/Shaders/Private/HairStrands/HairStrandsDebug.usf b/Engine/Shaders/Private/HairStrands/HairStrandsDebug.usf
--- a/Engine/Shaders/Private/HairStrands/HairStrandsDebug.usf
+++ b/Engine/Shaders/Private/HairStrands/HairStrandsDebug.usf
@@ -2306,324 +2306,6 @@ void MainCS(uint3 DispatchThreadId : SV_DispatchThreadID)
 
 #endif //SHADER_CARDS_DIRECTION_OCCLUSION
 
-///////////////////////////////////////////////////////////////////////////////////////////////////
-
-#if SHADER_PLOTBSDF
-
-#include "../SceneTexturesCommon.ush"
-#include "../ShadingModels.ush"
-#include "../DeferredShadingCommon.ush"
-#include "../HairBsdf.ush"
-#include "HairStrandsCommon.ush"
-#include "HairStrandsVisibilityCommon.ush"
-#include "HairStrandsEnvironmentLightingCommon.ush"
-#include "HairStrandsDeepShadowCommon.ush"
-#include "HairStrandsDeepTransmittanceCommon.ush"
-#include "HairStrandsDeepTransmittanceDualScattering.ush"
-
-int2 InputCoord;
-int2 OutputOffset;
-int2 OutputResolution;
-int2 MaxResolution;
-uint HairComponents;
-float Roughness;
-float BaseColor;
-float Exposure;
-
-RWTexture2D<float4> OutputTexture;
-
-#define COMPARE_ROUGHNESS 0
-
-#if COMPARE_ROUGHNESS
-float DefaultLit(float3 N, float3 V, float3 L, float InRoughnesss)
-{
-	FGBufferData GBuffer;
-	GBuffer.Roughness = InRoughnesss;
-	GBuffer.SpecularColor = 1;
-	GBuffer.DiffuseColor = 0;
-
-	FAreaLight AreaLight;
-	AreaLight.IsRectAndDiffuseMicroReflWeight = 0;
-	AreaLight.FalloffColor = 1;
-	AreaLight.SphereSinAlpha = 0;
-	AreaLight.SphereSinAlphaSoft = 0;
-	AreaLight.LineCosSubtended = 1;
-	FShadowTerms Shadow;
-
-	if (dot(V,N) < 0)
-		V.z = -V.z;
-
-	const float NoL = saturate(dot(N, L));
-	const float3 Fs = DefaultLitBxDF(GBuffer, N, V, L, 1, NoL, AreaLight, Shadow).Specular;
-	return Fs.x;
-}
-#endif
-
-void MainPS(
-	in FScreenVertexOutput Input,
-	out float4 OutColor : SV_Target0)
-{
-	const int2 PixelCoord = Input.Position.xy;
-	if (PixelCoord.x < OutputOffset.x || PixelCoord.x >= OutputOffset.x+OutputResolution.x ||
-		PixelCoord.y < OutputOffset.y || PixelCoord.y >= OutputOffset.x+OutputResolution.y )
-		discard;
-
-	const float2 UV = ((PixelCoord-OutputOffset) + float2(0.5f,0.5f)) / float2(OutputResolution);
-	float2 InputUV = saturate(float2(InputCoord - OutputOffset) / float2(OutputResolution));
-	// InputUV = float2(0.5, 0.5);
-
-
-	float V_Phi   = (UV.x*2-1) * PI; //[-Pi..Pi]
-	float V_Theta = (UV.y*2-1) * PI * 0.5f;
-	float L_Phi   = (InputUV.x*2-1) * PI; //[-Pi..Pi]
-	float L_Theta = (InputUV.y*2-1) * PI * 0.5f;
-
-	const float3 Axis_U = float3(0,0,1); // Tangent
-	const float3 Axis_V = float3(1,0,0); // Perp. to the view direction
-	const float3 Axis_W = float3(0,1,0); // Ortho.
-
-	const float3 L_Local = float3(sin(L_Theta), cos(L_Theta)*cos(L_Phi), cos(L_Theta)*sin(L_Phi));
-	const float3 V_Local = float3(sin(V_Theta), cos(V_Theta)*cos(V_Phi), cos(V_Theta)*sin(V_Phi));
-	const float3 T 		 = float3(1,0,0);
-
-	const float3 V_World = 
-		V_Local.x * Axis_U + 
-		V_Local.y * Axis_V + 
-		V_Local.z * Axis_W;
-
-	FHairTransmittanceData TransmittanceData = (FHairTransmittanceData)0;
-	TransmittanceData.LocalScattering = 0;
-	TransmittanceData.GlobalScattering = 1;
-	TransmittanceData.ScatteringComponent = HairComponents;
-	TransmittanceData.OpaqueVisibility = 1;
-
-	const float Backlit = 1;
-	const float Specular = 0.5f;
-	FGBufferData HairGBuffer = (FGBufferData)0;
-	HairGBuffer.WorldNormal = T;
-	HairGBuffer.BaseColor = BaseColor.xxx;
-	HairGBuffer.Roughness = Roughness;
-	HairGBuffer.ShadingModelID = SHADINGMODELID_HAIR;
-	HairGBuffer.DiffuseColor = 0;
-	HairGBuffer.SpecularColor = 0;
-	HairGBuffer.Specular = Specular;
-	HairGBuffer.Metallic = 0;
-	HairGBuffer.Depth = 0.5f;
-	HairGBuffer.GBufferAO = 1;
-	HairGBuffer.CustomData = float4(0,0,Backlit,0);
-	HairGBuffer.IndirectIrradiance = 1000000;
-
-	OutColor = float4(0,0,0,1);
-
-#if COMPARE_ROUGHNESS
-	const float3 Fs_Lit_R   = DefaultLit(float3(0, 0,-1), V_Local, L_Local, HairToGGXRoughness_R  (HairGBuffer.Roughness));
-	const float3 Fs_Lit_TT  = DefaultLit(float3(0, 0,-1), V_Local, L_Local, HairToGGXRoughness_TT (HairGBuffer.Roughness));
-	const float3 Fs_Lit_TRT = DefaultLit(float3(0, 0,-1), V_Local, L_Local, HairToGGXRoughness_TRT(HairGBuffer.Roughness));
-#endif
-	const float3 Fs_Hair    = HairShading(HairGBuffer, L_Local, V_Local, T, 1, TransmittanceData, 1, 0, uint2(0,0) );
-
-#if COMPARE_ROUGHNESS
-	OutColor.xyz += Fs_Lit_R   * float3(1, 0, 0);
-	OutColor.xyz += Fs_Lit_TRT * float3(0, 1, 1);
-	OutColor.xyz += Fs_Lit_TT  * float3(0, 0, 1);
-#endif
-	OutColor.xyz += Fs_Hair;
-	OutColor.xyz *= Exposure;
-
-	//const float3 ColorV = (V_World+1)* 0.5f;
-	//OutColor = float4(ColorV, 1);
-	//OutColor = float4(UV, 0, 1);
-}
-
-#endif // SHADER_PLOTBSDF
-
-//////////////////////////////////////////////////////////////////////////////////////////////////////////////////////////////
-
-#if SHADER_PLOTSAMPLE
-
-#include "HairStrandsDebugCommonStruct.ush"
-#include "../SceneTexturesCommon.ush"
-#include "../ShadingModels.ush"
-#include "../DeferredShadingCommon.ush"
-#include "../HairBsdf.ush"
-#include "HairStrandsCommon.ush"
-#include "HairStrandsVisibilityCommon.ush"
-#include "HairStrandsEnvironmentLightingCommon.ush"
-#include "HairStrandsDeepShadowCommon.ush"
-#include "HairStrandsDeepTransmittanceCommon.ush"
-#include "HairStrandsDeepTransmittanceDualScattering.ush"
-
-StructuredBuffer<FDebugShadingInfo>	Debug_ShadingPointBuffer;
-Buffer<uint>						Debug_ShadingPointCounter;
-StructuredBuffer<FDebugSample>		Debug_SampleBuffer;
-Buffer<uint>						Debug_SampleCounter;
-uint								Debug_MaxShadingPointCount;
-uint								Debug_MaxSampleCount;
-
-Texture3D<float4>					HairScatteringLUTTexture;
-SamplerState 						HairLUTSampler;
-
-int2 OutputOffset;
-int2 OutputResolution;
-int2 MaxResolution;
-uint HairComponents;
-float Exposure;
-
-struct FPlotRect
-{
-	int2 Offset;
-	int2 Resolution;
-};
-
-bool IsInside(const int2 Coord, const FPlotRect Rect)
-{
-	return all(Coord >= Rect.Offset) && all(Coord < Rect.Offset + Rect.Resolution);
-}
-
-float2 ToUV(const int2 Coord, const FPlotRect Rect)
-{
-	return ((Coord - Rect.Offset) + float2(0.5f, 0.5f)) / float2(OutputResolution);
-}
-
-RWTexture2D<float4> OutputTexture;
-
-float2 CartesianToUnwrap(float3 D)
-{
-	const float Theta = asin(D.x);
-	const float Phi = atan2(D.z, D.y);
-	return float2(Theta, Phi);
-}
-
-float2 CartesianToLatlong(float3 D)
-{
-	const float Theta = asin(D.x);
-	const float Phi = atan2(D.z, D.y);
-	return float2(((Phi/PI)+1)*0.5f, Theta/PI);
-}
-
-void MainPS(
-	in FScreenVertexOutput Input,
-	out float4 OutColor : SV_Target0)
-{
-	OutColor = float4(0, 0, 0, 1);
-
-	// Prune all pixel which are not within the band horizontal band
-	const int2 PixelCoord = Input.Position.xy;
-	if (PixelCoord.y < OutputOffset.y || PixelCoord.y >= OutputOffset.x + OutputResolution.y)
-	{
-		discard;
-	}
-
-	const float CosAngleThreshold = 0.999f;
-	const float PointSize = 3.f / OutputResolution.x;
-
-	// Local 
-	// Plot sample points
-	bool bIsValid = false;
-	const uint ShadingPointCount = Debug_ShadingPointCounter[0];
-	for (uint PointIt = 0; PointIt < ShadingPointCount; PointIt++)
-	{
-		FPlotRect Rect;
-		Rect.Offset = OutputOffset + PointIt * int2(OutputResolution.x + 10, 0);
-		Rect.Resolution = OutputResolution;
-
-		if (!IsInside(PixelCoord, Rect))
-			continue;
-
-		bIsValid = true;
-		const float2 UV = ToUV(PixelCoord, Rect);
-
-		const float  Pixel_Phi	= (UV.x * 2 - 1) * PI + PI*0.5; //[3PI/2, 0, 2PI, 3PI/2 ], so that the hair fiber is aligned on the center
-		const float  Pixel_Theta= (UV.y * 2 - 1) * PI * 0.5f;
-		const float3 PixelDir_Local	= float3(sin(Pixel_Theta), cos(Pixel_Theta)*cos(Pixel_Phi), cos(Pixel_Theta)*sin(Pixel_Phi));
-
-		const FDebugShadingInfo P = Debug_ShadingPointBuffer[PointIt];
-
-		// Shading frame T, V
-		const float3 Local_X = P.T;
-		const float3 Local_Y = normalize(cross(Local_X, P.V));
-		const float3 Local_Z = normalize(cross(Local_Y, Local_X));
-		const float3x3 WorldToLocal = transpose(float3x3(Local_X, Local_Y, Local_Z));
-
-		const float3 V_Local = mul(P.V, WorldToLocal);
-		const float3 T_Local = mul(P.T, WorldToLocal); // Should be float3(1,0,0)
-
-		{
-
-			const float SampleBacklit = 1;
-			const float SampleSpecular = 0.5f;
-			FHairTransmittanceData TransmittanceData = InitHairStrandsTransmittanceData();
-
-			uint InHairComponents = HairComponents;
-			const bool bPlotDualScatteringContribution = true;
-			if (!bPlotDualScatteringContribution)
-			{
-				TransmittanceData.LocalScattering = 1;
-				TransmittanceData.GlobalScattering = 1;
-				TransmittanceData.ScatteringComponent = InHairComponents;
-				TransmittanceData.OpaqueVisibility = 1;
-			}
-			else
-			{
-				const float3 L = PixelDir_Local;
-				const float3 T = T_Local;
-				const float3 V = V_Local;
-				const float SinLightAngle = dot(L, T);
-
-				// Compute the transmittance based on precompute Hair transmittance LUT
-				FHairAverageScattering AverageScattering = SampleHairLUT(View.HairScatteringLUTTexture, HairScatteringLUTSampler, P.BaseColor, P.Roughness, SinLightAngle);
-
-				TransmittanceData = ComputeDualScatteringTerms(
-					InitHairTransmittanceMask(),
-					AverageScattering,
-					P.Roughness,
-					SampleBacklit,
-					V,
-					L,
-					T,
-					InHairComponents);
-			}
-
-			FGBufferData HairGBuffer = (FGBufferData)0;
-			HairGBuffer.WorldNormal = P.T;
-			HairGBuffer.BaseColor = P.BaseColor;
-			HairGBuffer.Roughness = P.Roughness;
-			HairGBuffer.ShadingModelID = SHADINGMODELID_HAIR;
-			HairGBuffer.DiffuseColor = 0;
-			HairGBuffer.SpecularColor = 0;
-			HairGBuffer.Specular = SampleSpecular;
-			HairGBuffer.Metallic = 0;
-			HairGBuffer.Depth = 0.5f;
-			HairGBuffer.GBufferAO = 1;
-			HairGBuffer.CustomData = float4(0,0,SampleBacklit,0);
-			HairGBuffer.IndirectIrradiance = 1000000;
-
-			const float Area = 0;
-			const float BackLit = 1;
-			OutColor.xyz = HairShading(HairGBuffer, PixelDir_Local, V_Local, T_Local, 1, TransmittanceData, BackLit, Area, uint2(0, 0)) * Exposure;
-		}
-
-		for (uint SampleIt = 0; SampleIt < P.SampleCount; SampleIt++)
-		{
-			const FDebugSample S_World = Debug_SampleBuffer[P.SampleOffset + SampleIt];
-
-			const float3 S_Local = mul(S_World.Direction, WorldToLocal);
-
-			if (dot(PixelDir_Local, S_Local) > CosAngleThreshold)
-			{
-				OutColor.xyz = float3(1, 0, 0);
-			}
-		}
-	}
-
-	if (!bIsValid)
-	{
-		discard;
-	}
-}
-
-#endif // SHADER_PLOTSAMPLE
-
 //////////////////////////////////////////////////////////////////////////////////////////////////////////////////////////////
 // Display groom culling from shadow perspective 
 
diff --git a/Engine/Source/Runtime/Renderer/Private/HairStrands/HairStrandsData.h b/Engine/Source/Runtime/Renderer/Private/HairStrands/HairStrandsData.h
--- a/Engine/Source/Runtime/Renderer/Private/HairStrands/HairStrandsData.h
+++ b/Engine/Source/Runtime/Renderer/Private/HairStrands/HairStrandsData.h
@@ -325,61 +325,6 @@ struct FHairStrandsMacroGroupData
 
 struct FHairStrandsDebugData
 {
-	BEGIN_SHADER_PARAMETER_STRUCT(FWriteParameters, )
-		SHADER_PARAMETER(uint32, Debug_MaxShadingPointCount)
-		SHADER_PARAMETER(uint32, Debug_MaxSampleCount)
-		SHADER_PARAMETER_RDG_BUFFER_UAV(RWStructuredBuffer, Debug_ShadingPointBuffer)
-		SHADER_PARAMETER_RDG_BUFFER_UAV(RWBuffer, Debug_ShadingPointCounter)
-		SHADER_PARAMETER_RDG_BUFFER_UAV(RWStructuredBuffer, Debug_SampleBuffer)
-		SHADER_PARAMETER_RDG_BUFFER_UAV(RWBuffer, Debug_SampleCounter)
-		END_SHADER_PARAMETER_STRUCT()
-
-		BEGIN_SHADER_PARAMETER_STRUCT(FReadParameters, )
-		SHADER_PARAMETER(uint32, Debug_MaxShadingPointCount)
-		SHADER_PARAMETER(uint32, Debug_MaxSampleCount)
-		SHADER_PARAMETER_RDG_BUFFER_SRV(StructuredBuffer, Debug_ShadingPointBuffer)
-		SHADER_PARAMETER_RDG_BUFFER_SRV(Buffer, Debug_ShadingPointCounter)
-		SHADER_PARAMETER_RDG_BUFFER_SRV(StructuredBuffer, Debug_SampleBuffer)
-		SHADER_PARAMETER_RDG_BUFFER_SRV(Buffer, Debug_SampleCounter)
-	END_SHADER_PARAMETER_STRUCT()
-
-	// Plot Data
-	struct ShadingInfo
-	{
-		FVector3f	BaseColor;
-		float		Roughness;
-		FVector3f	T;
-		uint32		SampleCount;
-		FVector3f	V;
-		float		SampleOffset;
-	};
-
-	struct Sample
-	{
-		FVector3f	Direction;
-		float		Pdf;
-		FVector3f	Weights;
-		float		Pad;
-	};
-	static const uint32 MaxShadingPointCount = 32;
-	static const uint32 MaxSampleCount = 1024 * 32;
-
-	struct FPlotData
-	{
-		FRDGBufferRef ShadingPointBuffer = nullptr;
-		FRDGBufferRef ShadingPointCounter = nullptr;
-		FRDGBufferRef SampleBuffer = nullptr;
-		FRDGBufferRef SampleCounter = nullptr;
-	} PlotData;
-
-	bool IsPlotDataValid() const
-	{
-		return PlotData.ShadingPointBuffer && PlotData.ShadingPointCounter && PlotData.SampleBuffer && PlotData.SampleCounter;
-	}
-
-	static FPlotData CreatePlotData(FRDGBuilder& GraphBuilder);
-	static void SetParameters(FRDGBuilder& GraphBuilder, const FPlotData& In, FWriteParameters& Out);
-	static void SetParameters(FRDGBuilder& GraphBuilder, const FPlotData& In, FReadParameters& Out);
 
 	// PPLL debug data
 	struct FPPLLData
diff --git a/Engine/Source/Runtime/Renderer/Private/HairStrands/HairStrandsDebug.cpp b/Engine/Source/Runtime/Renderer/Private/HairStrands/HairStrandsDebug.cpp
--- a/Engine/Source/Runtime/Renderer/Private/HairStrands/HairStrandsDebug.cpp
+++ b/Engine/Source/Runtime/Renderer/Private/HairStrands/HairStrandsDebug.cpp
@@ -54,18 +54,6 @@ static float GDeepShadowDebugScale = 20;
 static FAutoConsoleVariableRef CVarDeepShadowDebugDomIndex(TEXT("r.HairStrands.DeepShadow.DebugDOMIndex"), GDeepShadowDebugIndex, TEXT("Index of the DOM texture to draw"));
 static FAutoConsoleVariableRef CVarDeepShadowDebugDomScale(TEXT("r.HairStrands.DeepShadow.DebugDOMScale"), GDeepShadowDebugScale, TEXT("Scaling value for the DeepOpacityMap when drawing the deep shadow stats"));
 
-static int32 GHairStrandsDebugPlotBsdf = 0;
-static FAutoConsoleVariableRef CVarHairStrandsDebugBSDF(TEXT("r.HairStrands.PlotBsdf"), GHairStrandsDebugPlotBsdf, TEXT("Debug view for visualizing hair BSDF."));
-
-static float GHairStrandsDebugPlotBsdfRoughness = 0.3f;
-static FAutoConsoleVariableRef CVarHairStrandsDebugBSDFRoughness(TEXT("r.HairStrands.PlotBsdf.Roughness"), GHairStrandsDebugPlotBsdfRoughness, TEXT("Change the roughness of the debug BSDF plot."));
-
-static float GHairStrandsDebugPlotBsdfBaseColor = 1;
-static FAutoConsoleVariableRef CVarHairStrandsDebugBSDFAbsorption(TEXT("r.HairStrands.PlotBsdf.BaseColor"), GHairStrandsDebugPlotBsdfBaseColor, TEXT("Change the base color / absorption of the debug BSDF plot."));
-
-static float GHairStrandsDebugPlotBsdfExposure = 1.1f;
-static FAutoConsoleVariableRef CVarHairStrandsDebugBSDFExposure(TEXT("r.HairStrands.PlotBsdf.Exposure"), GHairStrandsDebugPlotBsdfExposure, TEXT("Change the exposure of the plot."));
-
 static int32 GHairVirtualVoxel_DebugTraversalType = 0;
 static FAutoConsoleVariableRef CVarHairVirtualVoxel_DebugTraversalType(TEXT("r.HairStrands.Voxelization.Virtual.DebugTraversalType"), GHairVirtualVoxel_DebugTraversalType, TEXT("Traversal mode (0:linear, 1:mip) for debug voxel visualization."));
 
@@ -80,38 +68,6 @@ static bool TryEnableShaderDrawAndShaderPrint(const FViewInfo& View, uint32 Resq
 	return ShaderPrint::IsEnabled(View.ShaderPrintData);
 }
 
-FHairStrandsDebugData::FPlotData FHairStrandsDebugData::CreatePlotData(FRDGBuilder& GraphBuilder)
-{
-	FHairStrandsDebugData::FPlotData Out;
-	Out.ShadingPointBuffer = GraphBuilder.CreateBuffer(FRDGBufferDesc::CreateStructuredDesc(sizeof(ShadingInfo), MaxShadingPointCount), TEXT("Hair.DebugShadingPoint"));
-	Out.ShadingPointCounter = GraphBuilder.CreateBuffer(FRDGBufferDesc::CreateBufferDesc(sizeof(uint32), 1), TEXT("Hair.DebugShadingPointCounter"));
-	Out.SampleBuffer = GraphBuilder.CreateBuffer(FRDGBufferDesc::CreateStructuredDesc(sizeof(Sample), MaxSampleCount), TEXT("Hair.DebugSample"));
-	Out.SampleCounter = GraphBuilder.CreateBuffer(FRDGBufferDesc::CreateBufferDesc(sizeof(uint32), 1), TEXT("Hair.DebugSampleCounter"));
-	AddClearUAVPass(GraphBuilder, GraphBuilder.CreateUAV(Out.ShadingPointCounter, PF_R32_UINT), 0u);
-	AddClearUAVPass(GraphBuilder, GraphBuilder.CreateUAV(Out.SampleCounter, PF_R32_UINT), 0u);
-	return Out;
-}
-
-void FHairStrandsDebugData::SetParameters(FRDGBuilder& GraphBuilder, const FHairStrandsDebugData::FPlotData& In, FHairStrandsDebugData::FWriteParameters& Out)
-{
-	Out.Debug_MaxSampleCount = FHairStrandsDebugData::MaxSampleCount;
-	Out.Debug_MaxShadingPointCount = FHairStrandsDebugData::MaxShadingPointCount;
-	Out.Debug_ShadingPointBuffer = GraphBuilder.CreateUAV(In.ShadingPointBuffer);
-	Out.Debug_ShadingPointCounter = GraphBuilder.CreateUAV(In.ShadingPointCounter, PF_R32_UINT);
-	Out.Debug_SampleBuffer = GraphBuilder.CreateUAV(In.SampleBuffer);
-	Out.Debug_SampleCounter = GraphBuilder.CreateUAV(In.SampleCounter, PF_R32_UINT);
-}
-
-void FHairStrandsDebugData::SetParameters(FRDGBuilder& GraphBuilder, const FHairStrandsDebugData::FPlotData& In, FHairStrandsDebugData::FReadParameters& Out)
-{
-	Out.Debug_MaxSampleCount = FHairStrandsDebugData::MaxSampleCount;
-	Out.Debug_MaxShadingPointCount = FHairStrandsDebugData::MaxShadingPointCount;
-	Out.Debug_ShadingPointBuffer = GraphBuilder.CreateSRV(In.ShadingPointBuffer);
-	Out.Debug_ShadingPointCounter = GraphBuilder.CreateSRV(In.ShadingPointCounter, PF_R32_UINT);
-	Out.Debug_SampleBuffer = GraphBuilder.CreateSRV(In.SampleBuffer);
-	Out.Debug_SampleCounter = GraphBuilder.CreateSRV(In.SampleCounter, PF_R32_UINT);
-}
-
 ///////////////////////////////////////////////////////////////////////////////////////////////////
 class FHairPrintLODInfoCS : public FGlobalShader
 {
@@ -910,189 +866,6 @@ static void AddDebugHairTangentPass(
 	FComputeShaderUtils::AddPass(GraphBuilder, RDG_EVENT_NAME("HairStrands::DebugTangentCS"), ComputeShader, Parameters, DispatchCount);
 }
 
-///////////////////////////////////////////////////////////////////////////////////////////////////
-class FHairStrandsPlotBSDFPS : public FGlobalShader
-{
-	DECLARE_GLOBAL_SHADER(FHairStrandsPlotBSDFPS);
-	SHADER_USE_PARAMETER_STRUCT(FHairStrandsPlotBSDFPS, FGlobalShader);
-
-	BEGIN_SHADER_PARAMETER_STRUCT(FParameters, )
-		SHADER_PARAMETER(FIntPoint, InputCoord)
-		SHADER_PARAMETER(FIntPoint, OutputOffset)
-		SHADER_PARAMETER(FIntPoint, OutputResolution)
-		SHADER_PARAMETER(FIntPoint, MaxResolution)
-		SHADER_PARAMETER(uint32, HairComponents)
-		SHADER_PARAMETER(float, Roughness)
-		SHADER_PARAMETER(float, BaseColor)
-		SHADER_PARAMETER(float, Exposure)
-		RENDER_TARGET_BINDING_SLOTS()
-	END_SHADER_PARAMETER_STRUCT()
-
-public:
-	static bool ShouldCompilePermutation(const FGlobalShaderPermutationParameters& Parameters) { return IsHairStrandsSupported(EHairStrandsShaderType::Tool, Parameters.Platform); }
-	static EShaderPermutationPrecacheRequest ShouldPrecachePermutation(const FGlobalShaderPermutationParameters& Parameters) { return EShaderPermutationPrecacheRequest::NotPrecached; }
-	static void ModifyCompilationEnvironment(const FGlobalShaderPermutationParameters& Parameters, FShaderCompilerEnvironment& OutEnvironment)
-	{
-		FGlobalShader::ModifyCompilationEnvironment(Parameters, OutEnvironment);
-		OutEnvironment.SetDefine(TEXT("SHADER_PLOTBSDF"), 1);
-	}
-};
-
-IMPLEMENT_GLOBAL_SHADER(FHairStrandsPlotBSDFPS, "/Engine/Private/HairStrands/HairStrandsDebug.usf", "MainPS", SF_Pixel);
-
-static void AddPlotBSDFPass(
-	FRDGBuilder& GraphBuilder,
-	const FViewInfo& View,
-	FRDGTextureRef& OutputTexture)
-{
-	
-	FSceneTextureParameters SceneTextures = GetSceneTextureParameters(GraphBuilder, View);
-
-	const FIntPoint Resolution(OutputTexture->Desc.Extent);
-	FHairStrandsPlotBSDFPS::FParameters* Parameters = GraphBuilder.AllocParameters<FHairStrandsPlotBSDFPS::FParameters>();
-	Parameters->InputCoord = View.CursorPos;
-	Parameters->OutputOffset = FIntPoint(10,100);
-	Parameters->OutputResolution = FIntPoint(256, 256);
-	Parameters->MaxResolution = OutputTexture->Desc.Extent;
-	Parameters->HairComponents = ToBitfield(GetHairComponents());
-	Parameters->Roughness = GHairStrandsDebugPlotBsdfRoughness;
-	Parameters->BaseColor = GHairStrandsDebugPlotBsdfBaseColor;
-	Parameters->Exposure = GHairStrandsDebugPlotBsdfExposure;
-	Parameters->RenderTargets[0] = FRenderTargetBinding(OutputTexture, ERenderTargetLoadAction::ELoad);
-
-	const FIntPoint OutputResolution = SceneTextures.SceneDepthTexture->Desc.Extent;
-	TShaderMapRef<FScreenVS> VertexShader(View.ShaderMap);
-	TShaderMapRef<FHairStrandsPlotBSDFPS> PixelShader(View.ShaderMap);
-	const FGlobalShaderMap* GlobalShaderMap = View.ShaderMap;
-	const FIntRect Viewport = View.ViewRect;
-
-	ClearUnusedGraphResources(PixelShader, Parameters);
-
-	GraphBuilder.AddPass(
-		RDG_EVENT_NAME("HairStrands::BsdfPlot"),
-		Parameters,
-		ERDGPassFlags::Raster,
-		[Parameters, VertexShader, PixelShader, Viewport, Resolution](FRDGAsyncTask, FRHICommandList& RHICmdList)
-	{
-		FGraphicsPipelineStateInitializer GraphicsPSOInit;
-		RHICmdList.ApplyCachedRenderTargets(GraphicsPSOInit);
-		GraphicsPSOInit.BlendState = TStaticBlendState<CW_RGBA, BO_Add, BF_One, BF_Zero, BO_Add, BF_One, BF_Zero>::GetRHI();
-		GraphicsPSOInit.RasterizerState = TStaticRasterizerState<>::GetRHI();
-		GraphicsPSOInit.DepthStencilState = TStaticDepthStencilState<false, CF_Always>::GetRHI();
-
-		GraphicsPSOInit.BoundShaderState.VertexDeclarationRHI = GFilterVertexDeclaration.VertexDeclarationRHI;
-		GraphicsPSOInit.BoundShaderState.VertexShaderRHI = VertexShader.GetVertexShader();
-		GraphicsPSOInit.BoundShaderState.PixelShaderRHI = PixelShader.GetPixelShader();
-		GraphicsPSOInit.PrimitiveType = PT_TriangleList;
-		SetGraphicsPipelineState(RHICmdList, GraphicsPSOInit, 0);
-
-		RHICmdList.SetViewport(Viewport.Min.X, Viewport.Min.Y, 0.0f, Viewport.Max.X, Viewport.Max.Y, 1.0f);
-		SetShaderParameters(RHICmdList, PixelShader, PixelShader.GetPixelShader(), *Parameters);
-
-		DrawRectangle(
-			RHICmdList,
-			0, 0,
-			Viewport.Width(), Viewport.Height(),
-			Viewport.Min.X, Viewport.Min.Y,
-			Viewport.Width(), Viewport.Height(),
-			Viewport.Size(),
-			Resolution,
-			VertexShader,
-			EDRF_UseTriangleOptimization);
-	});
-}
-
-///////////////////////////////////////////////////////////////////////////////////////////////////
-class FHairStrandsPlotSamplePS : public FGlobalShader
-{
-	DECLARE_GLOBAL_SHADER(FHairStrandsPlotSamplePS);
-	SHADER_USE_PARAMETER_STRUCT(FHairStrandsPlotSamplePS, FGlobalShader);
-
-	BEGIN_SHADER_PARAMETER_STRUCT(FParameters, )
-		SHADER_PARAMETER_STRUCT_INCLUDE(FHairStrandsDebugData::FReadParameters, DebugData)
-		SHADER_PARAMETER(FIntPoint, OutputOffset)
-		SHADER_PARAMETER(FIntPoint, OutputResolution)
-		SHADER_PARAMETER(FIntPoint, MaxResolution)
-		SHADER_PARAMETER(uint32, HairComponents)
-		SHADER_PARAMETER(float, Exposure)
-		SHADER_PARAMETER_STRUCT_REF(FViewUniformShaderParameters, ViewUniformBuffer)
-		RENDER_TARGET_BINDING_SLOTS()
-	END_SHADER_PARAMETER_STRUCT()
-
-public:
-	static bool ShouldCompilePermutation(const FGlobalShaderPermutationParameters& Parameters) { return IsHairStrandsSupported(EHairStrandsShaderType::Tool, Parameters.Platform); }
-	static EShaderPermutationPrecacheRequest ShouldPrecachePermutation(const FGlobalShaderPermutationParameters& Parameters) { return EShaderPermutationPrecacheRequest::NotPrecached; }
-	static void ModifyCompilationEnvironment(const FGlobalShaderPermutationParameters& Parameters, FShaderCompilerEnvironment& OutEnvironment)
-	{
-		FGlobalShader::ModifyCompilationEnvironment(Parameters, OutEnvironment);
-		OutEnvironment.SetDefine(TEXT("SHADER_PLOTSAMPLE"), 1);
-	}
-};
-
-IMPLEMENT_GLOBAL_SHADER(FHairStrandsPlotSamplePS, "/Engine/Private/HairStrands/HairStrandsDebug.usf", "MainPS", SF_Pixel);
-
-static void AddPlotSamplePass(
-	FRDGBuilder& GraphBuilder,
-	const FViewInfo& View,
-	const FHairStrandsDebugData::FPlotData& DebugData,
-	FRDGTextureRef& OutputTexture)
-{
-	FSceneTextureParameters SceneTextures = GetSceneTextureParameters(GraphBuilder, View);
-
-	const FIntPoint Resolution(OutputTexture->Desc.Extent);
-	FHairStrandsPlotSamplePS::FParameters* Parameters = GraphBuilder.AllocParameters<FHairStrandsPlotSamplePS::FParameters>();
-
-	FHairStrandsDebugData::SetParameters(GraphBuilder, DebugData, Parameters->DebugData);
-	Parameters->ViewUniformBuffer = View.ViewUniformBuffer;
-	Parameters->OutputOffset = FIntPoint(100, 100);
-	Parameters->OutputResolution = FIntPoint(256, 256);
-	Parameters->MaxResolution = OutputTexture->Desc.Extent;
-	Parameters->HairComponents = ToBitfield(GetHairComponents());
-	Parameters->Exposure = GHairStrandsDebugPlotBsdfExposure;
-	Parameters->RenderTargets[0] = FRenderTargetBinding(OutputTexture, ERenderTargetLoadAction::ELoad);
-
-	const FIntPoint OutputResolution = SceneTextures.SceneDepthTexture->Desc.Extent;
-	TShaderMapRef<FScreenVS> VertexShader(View.ShaderMap);
-	TShaderMapRef<FHairStrandsPlotSamplePS> PixelShader(View.ShaderMap);
-	const FGlobalShaderMap* GlobalShaderMap = View.ShaderMap;
-	const FIntRect Viewport = View.ViewRect;
-
-	ClearUnusedGraphResources(PixelShader, Parameters);
-
-	GraphBuilder.AddPass(
-		RDG_EVENT_NAME("HairStrands::SamplePlot"),
-		Parameters,
-		ERDGPassFlags::Raster,
-		[Parameters, VertexShader, PixelShader, Viewport, Resolution](FRDGAsyncTask, FRHICommandList& RHICmdList)
-	{
-		FGraphicsPipelineStateInitializer GraphicsPSOInit;
-		RHICmdList.ApplyCachedRenderTargets(GraphicsPSOInit);
-		GraphicsPSOInit.BlendState = TStaticBlendState<CW_RGBA, BO_Add, BF_One, BF_Zero, BO_Add, BF_One, BF_Zero>::GetRHI();
-		GraphicsPSOInit.RasterizerState = TStaticRasterizerState<>::GetRHI();
-		GraphicsPSOInit.DepthStencilState = TStaticDepthStencilState<false, CF_Always>::GetRHI();
-
-		GraphicsPSOInit.BoundShaderState.VertexDeclarationRHI = GFilterVertexDeclaration.VertexDeclarationRHI;
-		GraphicsPSOInit.BoundShaderState.VertexShaderRHI = VertexShader.GetVertexShader();
-		GraphicsPSOInit.BoundShaderState.PixelShaderRHI = PixelShader.GetPixelShader();
-		GraphicsPSOInit.PrimitiveType = PT_TriangleList;
-		SetGraphicsPipelineState(RHICmdList, GraphicsPSOInit, 0);
-
-		RHICmdList.SetViewport(Viewport.Min.X, Viewport.Min.Y, 0.0f, Viewport.Max.X, Viewport.Max.Y, 1.0f);
-		SetShaderParameters(RHICmdList, PixelShader, PixelShader.GetPixelShader(), *Parameters);
-
-		DrawRectangle(
-			RHICmdList,
-			0, 0,
-			Viewport.Width(), Viewport.Height(),
-			Viewport.Min.X, Viewport.Min.Y,
-			Viewport.Width(), Viewport.Height(),
-			Viewport.Size(),
-			Resolution,
-			VertexShader,
-			EDRF_UseTriangleOptimization);
-	});
-}
-
 ///////////////////////////////////////////////////////////////////////////////////////////////////
 
 class FHairVisibilityDebugPPLLCS : public FGlobalShader
@@ -1208,18 +981,6 @@ static void InternalRenderHairStrandsDebugInfo(
 
 	const FHairStrandsViewData& HairData = View.HairStrandsViewData;
 
-	if (GHairStrandsDebugPlotBsdf > 0 || HairData.DebugData.IsPlotDataValid())
-	{
-		if (GHairStrandsDebugPlotBsdf > 0)
-		{
-			AddPlotBSDFPass(GraphBuilder, View, Params.SceneColorTexture);
-		}
-		if (HairData.DebugData.IsPlotDataValid())
-		{
-			AddPlotSamplePass(GraphBuilder, View, HairData.DebugData.PlotData, Params.SceneColorTexture);
-		}	
-	}
-
 	float ClusterY = 38;
 
 	if (View.HairStrandsViewData.DebugData.CullData.bIsValid)
diff --git a/Engine/Source/Runtime/Renderer/Private/HairStrands/HairStrandsEnvironment.cpp b/Engine/Source/Runtime/Renderer/Private/HairStrands/HairStrandsEnvironment.cpp
--- a/Engine/Source/Runtime/Renderer/Private/HairStrands/HairStrandsEnvironment.cpp
+++ b/Engine/Source/Runtime/Renderer/Private/HairStrands/HairStrandsEnvironment.cpp
@@ -62,9 +62,6 @@ static FAutoConsoleVariableRef CVarHairStrandsSkyLightingDistanceThreshold(TEXT(
 static int32 GHairStrandsSkyLighting_IntegrationType = 2;
 static FAutoConsoleVariableRef CVarHairStrandsSkyLighting_IntegrationType(TEXT("r.HairStrands.SkyLighting.IntegrationType"), GHairStrandsSkyLighting_IntegrationType, TEXT("Hair env. lighting integration type (0:Adhoc, 1:Uniform."), ECVF_Scalability | ECVF_RenderThreadSafe);
 
-static int32 GHairStrandsSkyLighting_DebugSample = 0;
-static FAutoConsoleVariableRef CVarHairStrandsSkyLighting_DebugSample(TEXT("r.HairStrands.SkyLighting.DebugSample"), GHairStrandsSkyLighting_DebugSample, TEXT("Enable debug view for visualizing sample used for the sky integration"), ECVF_Scalability | ECVF_RenderThreadSafe);
-
 static int32 GHairStrandsSkyLighting_ScreenTraceOcclusion = 0;
 static FAutoConsoleVariableRef CVarHairStrandsSkyLighting_ScreenTraceOcclusion(TEXT("r.HairStrands.SkyLighting.ScreenTraceOcclusion"), GHairStrandsSkyLighting_ScreenTraceOcclusion, TEXT("Enable screen trace occlusion for during hair env. lighting integration"), ECVF_Scalability | ECVF_RenderThreadSafe);
 
@@ -92,7 +89,6 @@ enum class EHairLightingIntegrationType : uint8
 };
 
 bool GetHairStrandsSkyLightingEnable()				{ return GHairSkylightingEnable > 0; }
-bool GetHairStrandsSkyLightingDebugEnable()			{ return GHairSkylightingEnable > 0 && GHairStrandsSkyLighting_DebugSample > 0; }
 static bool GetHairStrandsSkyAOEnable()				{ return GHairSkyAOEnable > 0; }
 static float GetHairStrandsSkyLightingConeAngle()	{ return FMath::Max(0.f, GHairSkylightingConeAngle); }
 DECLARE_GPU_STAT_NAMED(HairStrandsReflectionEnvironment, TEXT("Hair Strands Reflection Environment"));
@@ -236,17 +232,15 @@ class FHairEnvironmentLightingPS : public FGlobalShader
 	SHADER_USE_PARAMETER_STRUCT(FHairEnvironmentLightingPS, FGlobalShader)
 
 	class FIntegrationType	: SHADER_PERMUTATION_INT("PERMUTATION_INTEGRATION_TYPE", uint32(EHairLightingIntegrationType::Count));
-	class FDebug			: SHADER_PERMUTATION_BOOL("PERMUTATION_DEBUG"); 
 	class FLumen			: SHADER_PERMUTATION_BOOL("PERMUTATION_LUMEN"); 
 	class FStaticLighting	: SHADER_PERMUTATION_BOOL("PERMUTATION_STATIC_LIGHTING"); 
-	using FPermutationDomain = TShaderPermutationDomain<FIntegrationType, FLumen, FDebug, FStaticLighting>;
+	using FPermutationDomain = TShaderPermutationDomain<FIntegrationType, FLumen, FStaticLighting>;
 
 	BEGIN_SHADER_PARAMETER_STRUCT(FParameters, )
 		SHADER_PARAMETER_STRUCT_INCLUDE(FSkyDiffuseLightingParameters, SkyDiffuseLighting)
 		SHADER_PARAMETER_STRUCT_INCLUDE(FSceneTextureParameters, SceneTextures)
 		SHADER_PARAMETER_STRUCT_INCLUDE(LumenRadianceCache::FRadianceCacheInterpolationParameters, RadianceCache)
 		SHADER_PARAMETER_STRUCT_INCLUDE(ShaderPrint::FShaderParameters, ShaderPrintParameters)
-		SHADER_PARAMETER_STRUCT_INCLUDE(FHairStrandsDebugData::FWriteParameters, DebugData)
 		SHADER_PARAMETER_STRUCT_REF(FBlueNoise, BlueNoise)
 
 		SHADER_PARAMETER(uint32, bDynamicSkyLight)
@@ -285,12 +279,7 @@ class FHairEnvironmentLightingPS : public FGlobalShader
 
 	static bool ShouldCompilePermutation(const FGlobalShaderPermutationParameters& Parameters)
 	{
-		// Compile debug permutation only for the uniform integrator
 		FPermutationDomain PermutationVector(Parameters.PermutationId);
-		if (PermutationVector.Get<FIntegrationType>() != int32(EHairLightingIntegrationType::Uniform) && PermutationVector.Get<FDebug>())
-		{
-			return false;
-		}
 
 		// If using Lumen, no static lighting
 		if (PermutationVector.Get<FStaticLighting>() && PermutationVector.Get<FLumen>())
@@ -309,10 +298,6 @@ class FHairEnvironmentLightingPS : public FGlobalShader
 
 	static FPermutationDomain RemapPermutation(FPermutationDomain PermutationVector)
 	{
-		if (PermutationVector.Get<FIntegrationType>() != int32(EHairLightingIntegrationType::Uniform))
-		{
-			PermutationVector.Set<FDebug>(false);
-		}
 		return PermutationVector;
 	}
 };
@@ -327,8 +312,7 @@ static void AddHairStrandsEnvironmentLightingPassPS(
 	const FHairStrandsVisibilityData& VisibilityData,
 	const FHairStrandsVoxelResources& VirtualVoxelResources,
 	const FRDGTextureRef SceneColorTexture,
-	const EHairLightingSourceType LightingType,
-	const FHairStrandsDebugData::FPlotData* DebugData)
+	const EHairLightingSourceType LightingType)
 {
 	if (!Scene)
 	{
@@ -427,15 +411,9 @@ static void AddHairStrandsEnvironmentLightingPassPS(
 		ShaderPrint::SetParameters(GraphBuilder, View.ShaderPrintData, ParametersPS->ShaderPrintParameters);
 	}
 
-	if (DebugData)
-	{
-		FHairStrandsDebugData::SetParameters(GraphBuilder, *DebugData, ParametersPS->DebugData);
-	}
-
 	FHairEnvironmentLightingPS::FPermutationDomain PermutationVector;
 	PermutationVector.Set<FHairEnvironmentLightingPS::FIntegrationType>(uint32(IntegrationType));
 	PermutationVector.Set<FHairEnvironmentLightingPS::FLumen>(LightingType == EHairLightingSourceType::Lumen);
-	PermutationVector.Set<FHairEnvironmentLightingPS::FDebug>(DebugData != nullptr);
 	PermutationVector.Set<FHairEnvironmentLightingPS::FStaticLighting>(bHasStaticLighting);
 	PermutationVector = FHairEnvironmentLightingPS::RemapPermutation(PermutationVector);
 
@@ -500,7 +478,7 @@ void RenderHairStrandsSceneColorScattering(
 
 		if (HasHairFlags(View.HairStrandsViewData.Flags, HAIR_FLAGS_SCATTER_SCENE_LIGHT))
 		{
-			AddHairStrandsEnvironmentLightingPassPS(GraphBuilder, Scene, View, VisibilityData, VoxelResources, SceneColorTexture, EHairLightingSourceType::SceneColor, nullptr);
+			AddHairStrandsEnvironmentLightingPassPS(GraphBuilder, Scene, View, VisibilityData, VoxelResources, SceneColorTexture, EHairLightingSourceType::SceneColor);
 		}
 	}
 }
@@ -522,7 +500,7 @@ static void InternalRenderHairStrandsEnvironmentLighting(
 		return;
 	}
 	
-	AddHairStrandsEnvironmentLightingPassPS(GraphBuilder, Scene, View, VisibilityData, VoxelResources, nullptr, LightingType, View.HairStrandsViewData.DebugData.IsPlotDataValid() ? &View.HairStrandsViewData.DebugData.PlotData : nullptr);
+	AddHairStrandsEnvironmentLightingPassPS(GraphBuilder, Scene, View, VisibilityData, VoxelResources, nullptr, LightingType);
 }
 
 void RenderHairStrandsLumenLighting(
diff --git a/Engine/Source/Runtime/Renderer/Private/HairStrands/HairStrandsRendering.cpp b/Engine/Source/Runtime/Renderer/Private/HairStrands/HairStrandsRendering.cpp
--- a/Engine/Source/Runtime/Renderer/Private/HairStrands/HairStrandsRendering.cpp
+++ b/Engine/Source/Runtime/Renderer/Private/HairStrands/HairStrandsRendering.cpp
@@ -92,7 +92,6 @@ static TRDGUniformBufferRef<FHairStrandsViewUniformParameters> InternalCreateHai
 
 IMPLEMENT_GLOBAL_SHADER_PARAMETER_STRUCT(FHairStrandsViewUniformParameters, "HairStrands");
 
-bool GetHairStrandsSkyLightingDebugEnable();
 void AddMeshDrawTransitionPass(
 	FRDGBuilder& GraphBuilder,
 	const FViewInfo& ViewInfo,
@@ -221,12 +220,6 @@ void RenderHairBasePass(
 				SceneTextures.Depth.Resolve,
 				SceneTextures.Velocity,
 				InstanceCullingManager);
-
-			const bool bDebugSamplingEnable = GetHairStrandsSkyLightingDebugEnable();
-			if (bDebugSamplingEnable)
-			{
-				View.HairStrandsViewData.DebugData.PlotData = FHairStrandsDebugData::CreatePlotData(GraphBuilder);
-			}
 		}
 		
 		if (View.HairStrandsViewData.VisibilityData.CoverageTexture)
