SkyAtmosphere - FSkyAtmosphereRenderSceneInfo must be created on the render thread so simplify the code.

#rnx
#tests PC editor with Skyatmosphere

#Synced-CL 52139128

diff --git a/Engine/Source/Runtime/Renderer/Private/SkyAtmosphereRendering.cpp b/Engine/Source/Runtime/Renderer/Private/SkyAtmosphereRendering.cpp
--- a/Engine/Source/Runtime/Renderer/Private/SkyAtmosphereRendering.cpp
+++ b/Engine/Source/Runtime/Renderer/Private/SkyAtmosphereRendering.cpp
@@ -609,15 +609,11 @@ bool IsLightAtmospherePerPixelTransmittanceEnabled(const FScene* Scene, const FV
 FSkyAtmosphereRenderSceneInfo::FSkyAtmosphereRenderSceneInfo(FSkyAtmosphereSceneProxy& SkyAtmosphereSceneProxyIn)
 	:SkyAtmosphereSceneProxy(SkyAtmosphereSceneProxyIn)
 {
-	// Create a multiframe uniform buffer. A render command is used because FSkyAtmosphereRenderSceneInfo ctor is called on the Game thread.
-	TUniformBufferRef<FAtmosphereUniformShaderParameters>* AtmosphereUniformBufferPtr = &AtmosphereUniformBuffer;
-	FAtmosphereUniformShaderParameters* AtmosphereUniformShaderParametersPtr = &AtmosphereUniformShaderParameters;
 	CopyAtmosphereSetupToUniformShaderParameters(AtmosphereUniformShaderParameters, SkyAtmosphereSceneProxy.GetAtmosphereSetup());
-	ENQUEUE_RENDER_COMMAND(FCreateUniformBuffer)(
-		[AtmosphereUniformBufferPtr, AtmosphereUniformShaderParametersPtr](FRHICommandListImmediate& RHICmdList)
-	{
-		*AtmosphereUniformBufferPtr = TUniformBufferRef<FAtmosphereUniformShaderParameters>::CreateUniformBufferImmediate(*AtmosphereUniformShaderParametersPtr, UniformBuffer_MultiFrame);
-	});
+
+	// FSkyAtmosphereRenderSceneInfo must be created on the render thread to be able to create a multiframe uniform buffer. (see FScene::AddSkyAtmosphere)
+	check(IsInRenderingThread());
+	AtmosphereUniformBuffer = TUniformBufferRef<FAtmosphereUniformShaderParameters>::CreateUniformBufferImmediate(AtmosphereUniformShaderParameters, UniformBuffer_MultiFrame);
 }
 
 FSkyAtmosphereRenderSceneInfo::~FSkyAtmosphereRenderSceneInfo()
