Fixed flickering race condition in ASTP due to Lock / Unlock happening on the immediate command list instead of in the task.

#Synced-CL 52152383

diff --git a/Engine/Source/Runtime/Renderer/Private/Skinning/AnimSequenceTransformProvider.cpp b/Engine/Source/Runtime/Renderer/Private/Skinning/AnimSequenceTransformProvider.cpp
--- a/Engine/Source/Runtime/Renderer/Private/Skinning/AnimSequenceTransformProvider.cpp
+++ b/Engine/Source/Runtime/Renderer/Private/Skinning/AnimSequenceTransformProvider.cpp
@@ -887,25 +887,30 @@ void FAnimSequenceTransformProvider::UnregisterProxy(FAnimSequenceTransformProvi
 template <typename T>
 FRDGBuffer* CreateByteAddressBuffer(FRDGBuilder& GraphBuilder, uint32 MaxNumElements, const TCHAR* Name)
 {
-	return MaxNumElements > 0 ? GraphBuilder.CreateBuffer(FRDGBufferDesc::CreateByteAddressDesc(MaxNumElements * sizeof(T)), Name) : nullptr;
+	if (MaxNumElements > 0)
+	{
+		FRDGBuffer* Buffer = GraphBuilder.CreateBuffer(FRDGBufferDesc::CreateByteAddressDesc(MaxNumElements * sizeof(T)), Name);
+		GraphBuilder.ConvertToExternalBuffer(Buffer);
+		return Buffer;
+	}
+	return nullptr;
 }
 
 template <typename T>
-TArrayView<T> Lock(FRDGBuilder& GraphBuilder, FRDGBuffer* Buffer)
+TArrayView<T> Lock(FRHICommandList& RHICmdList, FRDGBuffer* Buffer)
 {
 	if (Buffer)
 	{
-		GraphBuilder.ConvertToExternalBuffer(Buffer);
-		return TArrayView<T>(reinterpret_cast<T*>(GraphBuilder.RHICmdList.LockBuffer(Buffer->GetRHIUnchecked(), 0, Buffer->Desc.GetSize(), RLM_WriteOnly)), Buffer->Desc.GetSize() / sizeof(T));
+		return TArrayView<T>(reinterpret_cast<T*>(RHICmdList.LockBuffer(Buffer->GetRHIUnchecked(), 0, Buffer->Desc.GetSize(), RLM_WriteOnly)), Buffer->Desc.GetSize() / sizeof(T));
 	}
 	return {};
 }
 
-void Unlock(FRDGBuilder& GraphBuilder, FRDGBuffer* Buffer)
+void Unlock(FRHICommandList& RHICmdList, FRDGBuffer* Buffer)
 {
 	if (Buffer)
 	{
-		GraphBuilder.RHICmdList.UnlockBuffer(Buffer->GetRHIUnchecked());
+		RHICmdList.UnlockBuffer(Buffer->GetRHIUnchecked());
 	}
 }
 
@@ -1218,8 +1223,8 @@ void FAnimSequenceTransformProvider::ProvideGPUTransforms(FSkinningTransformProv
 
 	for (FSinglePassDispatchGroup& Group : SinglePassDispatchGroups)
 	{
-		Group.EvaluateHeaderBaseOffset      = NumTotalSinglePassEvaluateHeaders;
-		Group.BlendEvaluateHeaderBaseOffset = NumTotalSinglePassBlendEvaluateHeaders;
+		Group.EvaluateHeaderBaseOffset          = NumTotalSinglePassEvaluateHeaders;
+		Group.BlendEvaluateHeaderBaseOffset     = NumTotalSinglePassBlendEvaluateHeaders;
 		NumTotalSinglePassEvaluateHeaders      += Group.NumEvaluateHeaders;
 		NumTotalSinglePassBlendEvaluateHeaders += Group.NumBlendEvaluateHeaders;
 	}
@@ -1227,10 +1232,10 @@ void FAnimSequenceTransformProvider::ProvideGPUTransforms(FSkinningTransformProv
 	for (FMultiPassDispatchGroup& Group : MultiPassDispatchGroups)
 	{
 		const uint32 NumGroupTotalHeaders = Group.NumEvaluateHeaders + Group.NumBlendEvaluateHeaders;
-		Group.EvaluateHeaderBaseOffset      = NumTotalMultiPassEvaluateHeaders;
-		Group.BlendEvaluateHeaderBaseOffset = NumTotalMultiPassBlendEvaluateHeaders;
-		Group.ConcatenateHeaderBaseOffset   = NumTotalConcatenateHeaders;
-		Group.ScatterHeaderBaseOffset       = NumTotalScatterHeaders;
+		Group.EvaluateHeaderBaseOffset         = NumTotalMultiPassEvaluateHeaders;
+		Group.BlendEvaluateHeaderBaseOffset    = NumTotalMultiPassBlendEvaluateHeaders;
+		Group.ConcatenateHeaderBaseOffset      = NumTotalConcatenateHeaders;
+		Group.ScatterHeaderBaseOffset          = NumTotalScatterHeaders;
 		NumTotalMultiPassEvaluateHeaders      += Group.NumEvaluateHeaders;
 		NumTotalMultiPassBlendEvaluateHeaders += Group.NumBlendEvaluateHeaders;
 		NumTotalConcatenateHeaders            += NumGroupTotalHeaders;
@@ -1245,58 +1250,35 @@ void FAnimSequenceTransformProvider::ProvideGPUTransforms(FSkinningTransformProv
 	FRDGBuffer* ConcatenateHeaderBuffer             = CreateByteAddressBuffer<UE::HLSL::FConcatenateBlockHeader>(GraphBuilder, NumTotalConcatenateHeaders, TEXT("AnimSequenceTransformProvider.ConcatenateHeaders"));
 	FRDGBuffer* ScatterHeaderBuffer                 = CreateByteAddressBuffer<UE::HLSL::FScatterTransformsBlockHeader>(GraphBuilder, NumTotalScatterHeaders, TEXT("AnimSequenceTransformProvider.ScatterHeaders"));
 
-	TArrayView<UE::HLSL::FEvaluateConcatenateScatterBlockHeader>      SinglePassEvaluateUpload      = Lock<UE::HLSL::FEvaluateConcatenateScatterBlockHeader>(GraphBuilder, SinglePassEvaluateHeaderBuffer);
-	TArrayView<UE::HLSL::FEvaluateBlendConcatenateScatterBlockHeader> SinglePassBlendEvaluateUpload = Lock<UE::HLSL::FEvaluateBlendConcatenateScatterBlockHeader>(GraphBuilder, SinglePassBlendEvaluateHeaderBuffer);
-	TArrayView<UE::HLSL::FEvaluateAnimSequenceBlockHeader>            MultiPassEvaluateUpload       = Lock<UE::HLSL::FEvaluateAnimSequenceBlockHeader>(GraphBuilder, MultiPassEvaluateHeaderBuffer);
-	TArrayView<UE::HLSL::FEvaluateAnimSequenceBlendBlockHeader>       MultiPassBlendEvaluateUpload  = Lock<UE::HLSL::FEvaluateAnimSequenceBlendBlockHeader>(GraphBuilder, MultiPassBlendEvaluateHeaderBuffer);
-	TArrayView<UE::HLSL::FConcatenateBlockHeader>                     ConcatenateUpload             = Lock<UE::HLSL::FConcatenateBlockHeader>(GraphBuilder, ConcatenateHeaderBuffer);
-	TArrayView<UE::HLSL::FScatterTransformsBlockHeader>               ScatterUpload                 = Lock<UE::HLSL::FScatterTransformsBlockHeader>(GraphBuilder, ScatterHeaderBuffer);
-
-	// Assign GroupBufferIndex on each LOD so the fill phase can index directly
-	for (const FSkinningTransformProvider::FProviderIndirection Indirection : Context.Indirections)
-	{
-		FInstancedSkinningSceneExtensionProxy* ExtensionProxy = static_cast<FInstancedSkinningSceneExtensionProxy*>(Context.Proxies[Indirection.Index]);
-		FProxyUserData& ProxyUserData = static_cast<FProxyUserData&>(ExtensionProxy->GetUserDataChecked());
-		const uint32 LODIndex = ExtensionProxy->GetLOD();
-		FProxyUserData::FLOD& LOD = ProxyUserData.LODs[LODIndex];
-
-		if (LOD.bSinglePass)
-		{
-			FSinglePassDispatchGroup LookupKey;
-			LookupKey.NumCompactBones = LOD.BoneTables.NumCompactBones;
-			LOD.DispatchGroup = SinglePassDispatchGroups.Find(LookupKey);
-			check(LOD.DispatchGroup != nullptr);
-		}
-		else
-		{
-			FMultiPassDispatchGroup LookupKey;
-			LookupKey.NumCompactBones       = LOD.BoneTables.NumCompactBones;
-			LookupKey.NumLeafBones          = LOD.LeafBonePathData.NumLeafBones;
-			LookupKey.MaxBoneTransformCount = ExtensionProxy->GetMaxBoneTransformCount();
-			LOD.DispatchGroup = MultiPassDispatchGroups.Find(LookupKey);
-			check(LOD.DispatchGroup != nullptr);
-		}
-	}
-
-	UE::Tasks::FTask ProcessTask = GraphBuilder.AddSetupTask(
+	UE::Tasks::FTask ProcessTask = GraphBuilder.AddCommandListSetupTask(
 		[
-			  this
-			, WorldTime
+			  WorldTime
 			, DeltaTime
 			, Indirections = Context.Indirections
 			, Proxies = Context.Proxies
 			, &NumTotalFramesRequested
 			, TrackUpdateRequests = MoveTemp(TrackUpdateRequests)
-			, SinglePassEvaluateUpload
-			, SinglePassBlendEvaluateUpload
-			, MultiPassEvaluateUpload
-			, MultiPassBlendEvaluateUpload
-			, ConcatenateUpload
-			, ScatterUpload
-		] () mutable
+			, SinglePassEvaluateHeaderBuffer
+			, SinglePassBlendEvaluateHeaderBuffer
+			, MultiPassEvaluateHeaderBuffer
+			, MultiPassBlendEvaluateHeaderBuffer
+			, ConcatenateHeaderBuffer
+			, ScatterHeaderBuffer
+			, &SinglePassDispatchGroups
+			, &MultiPassDispatchGroups
+
+		] (FRHICommandList& RHICmdList) mutable
 	{
 		TRACE_CPUPROFILER_EVENT_SCOPE(FAnimSequenceTransformProvider::EvaluateAnimSequences);
 
+		// Lock all header buffers on the command list timeline.
+		TArrayView<UE::HLSL::FEvaluateConcatenateScatterBlockHeader>      SinglePassEvaluateUpload      = Lock<UE::HLSL::FEvaluateConcatenateScatterBlockHeader>(RHICmdList, SinglePassEvaluateHeaderBuffer);
+		TArrayView<UE::HLSL::FEvaluateBlendConcatenateScatterBlockHeader> SinglePassBlendEvaluateUpload = Lock<UE::HLSL::FEvaluateBlendConcatenateScatterBlockHeader>(RHICmdList, SinglePassBlendEvaluateHeaderBuffer);
+		TArrayView<UE::HLSL::FEvaluateAnimSequenceBlockHeader>            MultiPassEvaluateUpload       = Lock<UE::HLSL::FEvaluateAnimSequenceBlockHeader>(RHICmdList, MultiPassEvaluateHeaderBuffer);
+		TArrayView<UE::HLSL::FEvaluateAnimSequenceBlendBlockHeader>       MultiPassBlendEvaluateUpload  = Lock<UE::HLSL::FEvaluateAnimSequenceBlendBlockHeader>(RHICmdList, MultiPassBlendEvaluateHeaderBuffer);
+		TArrayView<UE::HLSL::FConcatenateBlockHeader>                     ConcatenateUpload             = Lock<UE::HLSL::FConcatenateBlockHeader>(RHICmdList, ConcatenateHeaderBuffer);
+		TArrayView<UE::HLSL::FScatterTransformsBlockHeader>               ScatterUpload                 = Lock<UE::HLSL::FScatterTransformsBlockHeader>(RHICmdList, ScatterHeaderBuffer);
+
 		static const uint32 MultiPassBonesPerGroup = FEvaluateAnimSequenceCS::BonesPerGroup;
 		uint32 IntermediateTransformOffset = 0;
 		uint32 NumActiveTracks = 0;
@@ -1309,6 +1291,7 @@ void FAnimSequenceTransformProvider::ProvideGPUTransforms(FSkinningTransformProv
 
 		const auto EvaluateAnimSequence = [&](
 			const FProxyUserData::FLOD& LOD,
+			void* InDispatchGroup,
 			const FMapping& SourceMapping,
 			const FMapping& TargetMapping,
 			FSequence& SourceSequence,
@@ -1353,7 +1336,7 @@ void FAnimSequenceTransformProvider::ProvideGPUTransforms(FSkinningTransformProv
 
 			if (LOD.bSinglePass)
 			{
-				FSinglePassDispatchGroup& Group = *static_cast<FSinglePassDispatchGroup*>(LOD.DispatchGroup);
+				FSinglePassDispatchGroup& Group = *static_cast<FSinglePassDispatchGroup*>(InDispatchGroup);
 
 				if (BlendWeight > 0.0f)
 				{
@@ -1407,7 +1390,7 @@ void FAnimSequenceTransformProvider::ProvideGPUTransforms(FSkinningTransformProv
 			}
 			else
 			{
-				FMultiPassDispatchGroup& Group = *static_cast<FMultiPassDispatchGroup*>(LOD.DispatchGroup);
+				FMultiPassDispatchGroup& Group = *static_cast<FMultiPassDispatchGroup*>(InDispatchGroup);
 
 				// Evaluate: single header per instance
 				if (BlendWeight > 0.0f)
@@ -1491,6 +1474,24 @@ void FAnimSequenceTransformProvider::ProvideGPUTransforms(FSkinningTransformProv
 			FProxyUserData::FLOD& LOD = ProxyUserData.LODs[LODIndex];
 			FAnimSequenceTrackRenderPool& Tracks = ProxyUserData.Tracks;
 
+			// Look up the dispatch group for this LOD's bone configuration.
+			void* DispatchGroup = nullptr;
+			if (LOD.bSinglePass)
+			{
+				FSinglePassDispatchGroup LookupKey;
+				LookupKey.NumCompactBones = LOD.BoneTables.NumCompactBones;
+				DispatchGroup = SinglePassDispatchGroups.Find(LookupKey);
+			}
+			else
+			{
+				FMultiPassDispatchGroup LookupKey;
+				LookupKey.NumCompactBones       = LOD.BoneTables.NumCompactBones;
+				LookupKey.NumLeafBones          = LOD.LeafBonePathData.NumLeafBones;
+				LookupKey.MaxBoneTransformCount = ExtensionProxy->GetMaxBoneTransformCount();
+				DispatchGroup = MultiPassDispatchGroups.Find(LookupKey);
+			}
+			check(DispatchGroup != nullptr);
+
 			const uint32 CurrentTransformOffset  = Indirection.CurrentTransformOffset;
 			const uint32 PreviousTransformOffset = Indirection.PreviousTransformOffset;
 
@@ -1513,7 +1514,7 @@ void FAnimSequenceTransformProvider::ProvideGPUTransforms(FSkinningTransformProv
 					uint32 OutputTransformOffset2 = UpdateRequest.GetUpdateMode() == EPreviousBoneTransformUpdateMode::DuplicateCurrentToPrevious ? TrackPreviousTransformOffset : TrackCurrentTransformOffset;
 
 					EvaluateAnimSequence(
-						LOD,
+						LOD, DispatchGroup,
 						SourceTrack.Mapping, TargetTrack.Mapping,
 						SourceTrack.Sequence, TargetTrack.Sequence,
 						SourceTrack.Position, TargetTrack.Position,
@@ -1527,7 +1528,7 @@ void FAnimSequenceTransformProvider::ProvideGPUTransforms(FSkinningTransformProv
 				if (UpdateRequest.GetUpdateMode() == EPreviousBoneTransformUpdateMode::UpdatePrevious)
 				{
 					EvaluateAnimSequence(
-						LOD,
+						LOD, DispatchGroup,
 						SourceTrack.Mapping, TargetTrack.Mapping,
 						SourceTrack.Sequence, TargetTrack.Sequence,
 						SourceTrack.PreviousPosition, TargetTrack.PreviousPosition,
@@ -1543,6 +1544,13 @@ void FAnimSequenceTransformProvider::ProvideGPUTransforms(FSkinningTransformProv
 			});
 		}
 
+		Unlock(RHICmdList, SinglePassEvaluateHeaderBuffer);
+		Unlock(RHICmdList, SinglePassBlendEvaluateHeaderBuffer);
+		Unlock(RHICmdList, MultiPassEvaluateHeaderBuffer);
+		Unlock(RHICmdList, MultiPassBlendEvaluateHeaderBuffer);
+		Unlock(RHICmdList, ConcatenateHeaderBuffer);
+		Unlock(RHICmdList, ScatterHeaderBuffer);
+
 		INC_DWORD_STAT_BY(STAT_AnimSequenceTransformProvider_ActiveTracks, NumActiveTracks);
 		INC_DWORD_STAT_BY(STAT_AnimSequenceTransformProvider_EvaluateSinglePassBlocks, NumSinglePassBlocks);
 		INC_DWORD_STAT_BY(STAT_AnimSequenceTransformProvider_EvaluateSinglePassBlendBlocks, NumSinglePassBlendBlocks);
@@ -1552,13 +1560,6 @@ void FAnimSequenceTransformProvider::ProvideGPUTransforms(FSkinningTransformProv
 		INC_DWORD_STAT_BY(STAT_AnimSequenceTransformProvider_PreviousTransformUpdates, NumPreviousTransformUpdates);
 	});
 
-	Unlock(GraphBuilder, SinglePassEvaluateHeaderBuffer);
-	Unlock(GraphBuilder, SinglePassBlendEvaluateHeaderBuffer);
-	Unlock(GraphBuilder, MultiPassEvaluateHeaderBuffer);
-	Unlock(GraphBuilder, MultiPassBlendEvaluateHeaderBuffer);
-	Unlock(GraphBuilder, ConcatenateHeaderBuffer);
-	Unlock(GraphBuilder, ScatterHeaderBuffer);
-
 	FRDGBuffer* SequenceTransformBuffer = SequenceTransformPersistentBuffer.ResizeBufferIfNeeded(GraphBuilder, SequenceTransformAllocator.GetMaxSize());
 	FRDGBuffer* RetargetingDataBuffer = RetargetingDataPersistentBuffer.ResizeBufferIfNeeded(GraphBuilder, RetargetingDataAllocator.GetMaxSize());
 	FRDGBuffer* LeafBonePathBuffer = LeafBonePathPersistentBuffer.ResizeBufferIfNeeded(GraphBuilder, LeafBonePathAllocator.GetMaxSize());
@@ -2435,4 +2436,4 @@ void FAnimSequenceTransformProvider::ProvideCPUTransforms(FSkinningTransformProv
 		INC_DWORD_STAT_BY(STAT_AnimSequenceTransformProvider_CurrentTransformUpdates, NumCurrentTransformUpdates);
 		INC_DWORD_STAT_BY(STAT_AnimSequenceTransformProvider_PreviousTransformUpdates, NumPreviousTransformUpdates);
 	});
-}
\ No newline at end of file
+}
