Fix validation error in FTSRDetectThinGeometryCS due to BlackArrayDummy being assigned as Texture2D parameter.

#Synced-CL 52158667

diff --git a/Engine/Source/Runtime/Renderer/Private/PostProcess/TemporalSuperResolution.cpp b/Engine/Source/Runtime/Renderer/Private/PostProcess/TemporalSuperResolution.cpp
--- a/Engine/Source/Runtime/Renderer/Private/PostProcess/TemporalSuperResolution.cpp
+++ b/Engine/Source/Runtime/Renderer/Private/PostProcess/TemporalSuperResolution.cpp
@@ -2725,7 +2725,7 @@ FDefaultTemporalUpscaler::FOutputs AddTemporalSuperResolutionPasses(
 			{
 				if (bCameraCut)
 				{
-					PassParameters->ReprojectedHistoryCoverageTexture = GraphBuilder.CreateSRV(BlackArrayDummy);
+					PassParameters->ReprojectedHistoryCoverageTexture = GraphBuilder.CreateSRV(BlackDummy);
 					PassParameters->CurrentCoverageTexture = GraphBuilder.CreateSRV(FRDGTextureSRVDesc(BlackDummy));
 				}
 				else
