RHI: Turn on specialization constants by default on desktop platforms that support them (only Vulkan for now).

#rnx
#Synced-CL 52150182

diff --git a/Engine/Source/Runtime/Renderer/Private/LightGridInjection.cpp b/Engine/Source/Runtime/Renderer/Private/LightGridInjection.cpp
--- a/Engine/Source/Runtime/Renderer/Private/LightGridInjection.cpp
+++ b/Engine/Source/Runtime/Renderer/Private/LightGridInjection.cpp
@@ -190,8 +190,8 @@ FAutoConsoleVariableRef CVarLightCullingMaxDistanceOverride(
 // Experimental version of FLightGridInjectionCS that uses specialization constants
 static TAutoConsoleVariable<int32> CVarLightGridInjectionUseSpecializationConstants(
 	TEXT("r.LightGridInjection.UseSpecializationConstants"),
-	0,
-	TEXT("Experimental: Selects whether LightGridInjection should use specialization constant. (default: 0)"),
+	1,
+	TEXT("Experimental: Selects whether LightGridInjection should use specialization constant. (default: 1)"),
 	ECVF_ReadOnly);
 
 static inline bool UseSpecializationConstants(EShaderPlatform ShaderPlatform)
