Fixed prefetch level mappings for virtual texture mapping requests

#Synced-CL 52137073

diff --git a/Engine/Source/Runtime/Renderer/Private/VT/VirtualTextureSystem.cpp b/Engine/Source/Runtime/Renderer/Private/VT/VirtualTextureSystem.cpp
--- a/Engine/Source/Runtime/Renderer/Private/VT/VirtualTextureSystem.cpp
+++ b/Engine/Source/Runtime/Renderer/Private/VT/VirtualTextureSystem.cpp
@@ -1882,7 +1882,8 @@ void FVirtualTextureSystem::GatherRequestsTask(const FGatherRequestsParameters&
 							ProducerPhysicalGroupMaskToPrefetch, PageCount, bStreamingRequest);
 						if (LoadRequestIndex != 0xffff)
 						{
-							const uint32 PrefetchMapping_vLevel = PrefetchLocal_vLevel - AllocatedAdaptiveBias + ProducerMipBias;
+							const uint32 PrefetchMapping_vLevel_Producer = PrefetchLocal_vLevel + ProducerMipBias;
+							const uint32 PrefetchMapping_vLevel = PrefetchMapping_vLevel_Producer - AllocatedAdaptiveBias;
 							ensure(PrefetchMapping_vLevel <= MaxLevel);
 							const uint32 Prefetch_vAddress = vAddress & (0xffffffff << (PrefetchMapping_vLevel * vDimensions));
 							for (uint32 LoadLayerIndex = 0u; LoadLayerIndex < NumPageTableLayersToLoad; ++LoadLayerIndex)
@@ -1895,7 +1896,7 @@ void FVirtualTextureSystem::GatherRequestsTask(const FGatherRequestsParameters&
 									{
 										RequestList->AddMappingRequest(
 											LoadRequestIndex, ProducerPhysicalGroupIndex, ID, LayerIndex, MaxLevel, 
-											Prefetch_vAddress, PrefetchMapping_vLevel, PrefetchLocal_vLevel
+											Prefetch_vAddress, PrefetchMapping_vLevel, PrefetchMapping_vLevel_Producer
 										);
 									}
 								}
