Path Tracer: Don't prepare an RTPSO for experimental configurations on startup, even if enabled

This should reduce startup time even if you need to keep the experimental permutations enabled for coverage.

#rnx
#Synced-CL 52176564

diff --git a/Engine/Source/Runtime/Renderer/Private/PathTracing.cpp b/Engine/Source/Runtime/Renderer/Private/PathTracing.cpp
--- a/Engine/Source/Runtime/Renderer/Private/PathTracing.cpp
+++ b/Engine/Source/Runtime/Renderer/Private/PathTracing.cpp
@@ -1388,22 +1388,21 @@ class FPathTracingRG : public FGlobalShader
 	class FHeterogeneousVolumesPipeline : SHADER_PERMUTATION_INT("PATH_TRACER_HETEROGENEOUS_VOLUMES_PIPELINE", 2);
 	using FPermutationDomain = TShaderPermutationDomain<FKernelType, FAdaptiveSampling, FCloudShader, FSubstrateComplexSpecialMaterial, FUseSER, FNeedTMinWorkaround, FHeterogeneousVolumesPipeline>;
 
+    static bool IsPermutationExperimental(const FGlobalShaderPermutationParameters& Parameters)
+    {
+        const FPermutationDomain PermutationVector(Parameters.PermutationId);
+        return PermutationVector.Get<FAdaptiveSampling>() ||
+               PermutationVector.Get<FHeterogeneousVolumesPipeline>();
+    }
+
 	static bool ShouldCompilePermutation(const FGlobalShaderPermutationParameters& Parameters)
 	{
 		const bool bUseExperimental = CVarPathTracingExperimental.GetValueOnAnyThread();
-		FPermutationDomain PermutationVector(Parameters.PermutationId);
-		if (bUseExperimental == false)
+		const FPermutationDomain PermutationVector(Parameters.PermutationId);
+		if (bUseExperimental == false && IsPermutationExperimental(Parameters))
 		{
-			// default case: experimental features and rarely used modes are not compiled
-			if (PermutationVector.Get<FAdaptiveSampling>())
-			{
-				return false;
-			}
-
-			if (PermutationVector.Get<FHeterogeneousVolumesPipeline>())
-			{
-				return false;
-			}
+            // do not compile experimental features if they are not enabled
+            return false;
 		}
 
 		if (PermutationVector.Get<FCloudShader>())
@@ -3898,6 +3897,11 @@ void PreparePathTracingRTPSO()
 					// Permutation is not enabled, nothing to pre-compile
 					continue;
 				}
+                if (FPathTracingRG::IsPermutationExperimental(Parameters))
+                {
+                    // Permutation is for an experimental feature, don't compile it on startup
+                    continue;
+                }
 				FGlobalShaderMap* ShaderMap = GetGlobalShaderMap(GMaxRHIFeatureLevel);
 
 				FPathTracingRG::FPermutationDomain PermutationVector(PermutationId);
