Move all CSM state out of FScene and into newly created FCSMShadowScene
#Synced-CL 52204226

diff --git a/Engine/Source/Runtime/Renderer/Private/LightSceneInfo.cpp b/Engine/Source/Runtime/Renderer/Private/LightSceneInfo.cpp
--- a/Engine/Source/Runtime/Renderer/Private/LightSceneInfo.cpp
+++ b/Engine/Source/Runtime/Renderer/Private/LightSceneInfo.cpp
@@ -153,8 +153,6 @@ void FLightSceneInfo::RemoveFromScene()
 		Scene->DirectionalShadowCastingLightIDs.RemoveSwap(Id);
 	}
 
-	Scene->CachedShadowMaps.Remove(Id);
-
 	// Detach the light from the primitives it affects.
 	Detach();
 }
diff --git a/Engine/Source/Runtime/Renderer/Private/RendererScene.cpp b/Engine/Source/Runtime/Renderer/Private/RendererScene.cpp
--- a/Engine/Source/Runtime/Renderer/Private/RendererScene.cpp
+++ b/Engine/Source/Runtime/Renderer/Private/RendererScene.cpp
@@ -3027,26 +3027,6 @@ const FPlanarReflectionSceneProxy* FScene::GetForwardPassGlobalPlanarReflection(
 	return nullptr;
 }
 
-int64 FScene::GetCachedWholeSceneShadowMapsSize() const
-{
-	int64 CachedShadowmapMemory = 0;
-
-	for (TMap<FLightSceneId, TArray<FCachedShadowMapData>>::TConstIterator CachedShadowMapIt(CachedShadowMaps); CachedShadowMapIt; ++CachedShadowMapIt)
-	{
-		const TArray<FCachedShadowMapData>& ShadowMapDatas = CachedShadowMapIt.Value();
-
-		for (const auto& ShadowMapData : ShadowMapDatas)
-		{
-			if (ShadowMapData.ShadowMap.IsValid())
-			{
-				CachedShadowmapMemory += ShadowMapData.ShadowMap.ComputeMemorySize();
-			}
-		}
-	}
-
-	return CachedShadowmapMemory;
-}
-
 void FScene::AddPrecomputedLightVolume(const FPrecomputedLightVolume* Volume)
 {
 	FScene* Scene = this;
@@ -6349,8 +6329,6 @@ void FScene::Update(FRDGBuilder& GraphBuilder, const FUpdateParameters& Paramete
 	RayTracingSBT.FlushAllocationsToClear(GraphBuilder.RHICmdList);
 #endif // RHI_RAYTRACING
 
-	UpdateCachedShadowState(SceneUpdateChangeSetStorage.GetPreUpdateSet(), SceneUpdateChangeSetStorage.GetPostUpdateSet());
-
 	SceneExtensionsUpdaters.PostSceneUpdate(GraphBuilder, SceneUpdateChangeSetStorage.GetPostUpdateSet());
 
 	const bool bAsyncCacheMeshDrawCommands = EnumHasAnyFlags(Parameters.AsyncOps, EUpdateAllPrimitiveSceneInfosAsyncOps::CacheMeshDrawCommands) && GRHISupportsMultithreadedShaderCreation;
@@ -7185,30 +7163,6 @@ bool FScene::AddCustomRenderPass(const FSceneViewFamily* ViewFamily, const FCust
 	return true;
 }
 
-void FScene::UpdateCachedShadowState(const FScenePreUpdateChangeSet &ScenePreUpdateChangeSet, const FScenePostUpdateChangeSet &ScenePostUpdateChangeSet)
-{
-	for (const FLightSceneInfo* LightSceneInfo : DirectionalLights)
-	{
-		TArray<FCachedShadowMapData>* CachedShadowMapDatas = GetCachedShadowMapDatas(LightSceneInfo->Id);
-
-		if (CachedShadowMapDatas)
-		{
-			for (auto& CachedShadowMapData : *CachedShadowMapDatas)
-			{
-				for (FPersistentPrimitiveIndex PersistentPrimitiveIndex : ScenePreUpdateChangeSet.RemovedPrimitiveIds)
-				{
-					if (CachedShadowMapData.StaticShadowSubjectPersistentPrimitiveIdMap[PersistentPrimitiveIndex.Index] == true)
-					{
-						CachedShadowMapData.InvalidateCachedShadow();
-						break;
-					}
-				}
-				CachedShadowMapData.StaticShadowSubjectPersistentPrimitiveIdMap.SetNum(GetMaxPersistentPrimitiveIndex(), false);
-			}
-		}
-	}
-}
-
 FVirtualShadowMapArrayCacheManager* FScene::GetVirtualShadowMapCache()
 {
 	return GetExtensionPtr<FVirtualShadowMapArrayCacheManager>();
diff --git a/Engine/Source/Runtime/Renderer/Private/SceneCore.cpp b/Engine/Source/Runtime/Renderer/Private/SceneCore.cpp
--- a/Engine/Source/Runtime/Renderer/Private/SceneCore.cpp
+++ b/Engine/Source/Runtime/Renderer/Private/SceneCore.cpp
@@ -16,6 +16,7 @@
 #include "MobileBasePassRendering.h"
 #include "RendererModule.h"
 #include "ScenePrivate.h"
+#include "Shadows/CSMSceneExtension.h"
 #include "Containers/AllocatorFixedSizeFreeList.h"
 #include "MaterialShared.h"
 #include "HAL/LowLevelMemTracker.h"
@@ -371,7 +372,7 @@ void FLightPrimitiveInteraction::FlushCachedShadowMapData()
 	{
 		if (bCastShadow && !PrimitiveSceneInfo->Proxy->IsMeshShapeOftenMoving())
 		{
-			TArray<FCachedShadowMapData>* CachedShadowMapDatas = PrimitiveSceneInfo->Scene->GetCachedShadowMapDatas(LightSceneInfo->Id);
+			TArray<FCachedShadowMapData>* CachedShadowMapDatas = PrimitiveSceneInfo->Scene->GetExtension<FCSMSceneExtension>().GetCachedShadowMapDatas(LightSceneInfo->Id);
 
 			if (CachedShadowMapDatas)
 			{
diff --git a/Engine/Source/Runtime/Renderer/Private/ScenePrivate.h b/Engine/Source/Runtime/Renderer/Private/ScenePrivate.h
--- a/Engine/Source/Runtime/Renderer/Private/ScenePrivate.h
+++ b/Engine/Source/Runtime/Renderer/Private/ScenePrivate.h
@@ -1390,48 +1390,6 @@ private:
 	void HideNodeChildren(FSceneViewState* ViewState, FLODSceneNode& Node);
 };
 
-// Enable the DEBUG_CSM_CACHING to make the debugging CSM caching with RenderDoc more easier
-#define DEBUG_CSM_CACHING 0
-
-class FCachedShadowMapData
-{
-public:
-	FWholeSceneProjectedShadowInitializer Initializer;
-	FShadowMapRenderTargetsRefCounted ShadowMap;
-	float LastUsedTime;
-	bool bCachedShadowMapHasPrimitives;
-	bool bCachedShadowMapHasNaniteGeometry;
-
-	/**
-	 * The static meshes cast shadow on this cached csm
-	 */
-	TBitArray<> StaticShadowSubjectPersistentPrimitiveIdMap;
-
-	FIntPoint ShadowBufferResolution;
-	FVector PreShadowTranslation;
-	float MaxSubjectZ;
-	float MinSubjectZ;
-
-	/**
-	* The extra static meshes cast shadow in last frame, if it exceeds the r.Shadow.MaxCSMScrollingStaticShadowSubjects, the cached csm should be rebuilt.
-	*/
-	int32 LastFrameExtraStaticShadowSubjects;
-
-	void InvalidateCachedShadow()
-	{
-		ShadowMap.Release();
-
-		StaticShadowSubjectPersistentPrimitiveIdMap.SetRange(0, StaticShadowSubjectPersistentPrimitiveIdMap.Num(), false);
-	}
-
-	FCachedShadowMapData(const FWholeSceneProjectedShadowInitializer& InInitializer, float InLastUsedTime) :
-		Initializer(InInitializer),
-		LastUsedTime(InLastUsedTime),
-		bCachedShadowMapHasPrimitives(true),
-		bCachedShadowMapHasNaniteGeometry(false)
-	{}
-};
-
 #if WITH_EDITOR
 class FPixelInspectorData
 {
@@ -1955,9 +1913,6 @@ public:
 	/** Maps regular landscape component IDs to their Nanite proxies. */
 	TMap<FPrimitiveComponentId, FPrimitiveSceneInfo*> LandscapeToNaniteProxyMap;
 
-	/** Map from light id to the cached shadowmap data for that light. */
-	TMap<FLightSceneId, TArray<FCachedShadowMapData>> CachedShadowMaps;
-	
 	/** Atlas HZB textures from the previous render. */
 	TArray<TRefCountPtr<IPooledRenderTarget>>	PrevAtlasHZBs;
 	TArray<TRefCountPtr<IPooledRenderTarget>>	PrevAtlasCompleteHZBs;
@@ -2346,8 +2301,6 @@ public:
 
 	const class FPlanarReflectionSceneProxy* GetForwardPassGlobalPlanarReflection() const;
 
-	int64 GetCachedWholeSceneShadowMapsSize() const;
-
 	/**
 	 * Get the default base pass depth stencil access
 	 */
@@ -2482,38 +2435,6 @@ public:
 
 	void DumpMeshDrawCommandMemoryStats();
 
-	FORCEINLINE TArray<FCachedShadowMapData>* GetCachedShadowMapDatas(FLightSceneId LightID)
-	{
-		return CachedShadowMaps.Find(LightID);
-	}
-
-	FORCEINLINE FCachedShadowMapData& GetCachedShadowMapDataRef(FLightSceneId LightID, int32 ShadowMapIndex = 0)
-	{
-		TArray<FCachedShadowMapData>& CachedShadowMapDatas = CachedShadowMaps.FindChecked(LightID);
-
-		checkSlow(ShadowMapIndex >= 0 && ShadowMapIndex < CachedShadowMapDatas.Num());
-
-		return CachedShadowMapDatas[ShadowMapIndex];
-	}
-
-	FORCEINLINE const FCachedShadowMapData& GetCachedShadowMapDataRef(FLightSceneId LightID, int32 ShadowMapIndex = 0) const
-	{
-		const TArray<FCachedShadowMapData>& CachedShadowMapDatas = CachedShadowMaps.FindChecked(LightID);
-
-		checkSlow(ShadowMapIndex >= 0 && ShadowMapIndex < CachedShadowMapDatas.Num());
-
-		return CachedShadowMapDatas[ShadowMapIndex];
-	}
-
-	FORCEINLINE const FCachedShadowMapData* GetCachedShadowMapData(FLightSceneId LightID, int32 ShadowMapIndex = 0) const
-	{
-		const TArray<FCachedShadowMapData>& CachedShadowMapDatas = CachedShadowMaps.FindChecked(LightID);
-
-		checkSlow(ShadowMapIndex >= 0 && ShadowMapIndex < CachedShadowMapDatas.Num());
-
-		return &CachedShadowMapDatas[ShadowMapIndex];
-	}
-
 	/**
 	 * Maximum used persistent Primitive Index, use to size arrays that store primitive data indexed by FPrimitiveSceneInfo::PersistentIndex.
 	 * Only changes during UpdateAllPrimitiveSceneInfos.
@@ -2761,11 +2682,6 @@ private:
 	template <typename UpdatePayloadType>
 	void UpdateLightInternal_RenderThread(FLightSceneProxy* LightSceneProxy, UpdatePayloadType&& InUpdatePayload);
 
-	/**
-	 * Update tracked scene state for cached CSM shadows
-	 */
-	void UpdateCachedShadowState(const FScenePreUpdateChangeSet &ScenePreUpdateChangeSet, const FScenePostUpdateChangeSet &ScenePostUpdateChangeSet);
-
 	FString FullWorldName;
 #if RHI_RAYTRACING
 	void UpdateRayTracingGroupBounds_AddPrimitives(const TArray<FPrimitiveSceneInfo*, SceneRenderingAllocator>& PrimitiveSceneInfos);
diff --git a/Engine/Source/Runtime/Renderer/Private/ShadowDepthRendering.cpp b/Engine/Source/Runtime/Renderer/Private/ShadowDepthRendering.cpp
--- a/Engine/Source/Runtime/Renderer/Private/ShadowDepthRendering.cpp
+++ b/Engine/Source/Runtime/Renderer/Private/ShadowDepthRendering.cpp
@@ -84,6 +84,7 @@
 #include "ShaderPlatformCachedIniValue.h"
 #include "ShadowRendering.h"
 #include "Shadows/ShadowSceneRenderer.h"
+#include "Shadows/CSMSceneExtension.h"
 #include "ShowFlags.h"
 #include "Stats/Stats.h"
 #include "SystemTextures.h"
@@ -958,7 +959,7 @@ void FProjectedShadowInfo::CopyCachedShadowMap(
 	const FMeshPassProcessorRenderState& DrawRenderState)
 {
 	check(CacheMode == SDCM_MovablePrimitivesOnly || CacheMode == SDCM_CSMScrolling);
-	const FCachedShadowMapData& CachedShadowMapData = SceneRenderer->Scene->GetCachedShadowMapDataRef(GetLightSceneInfo().Id, FMath::Max(CascadeSettings.ShadowSplitIndex, 0));
+	const FCachedShadowMapData& CachedShadowMapData = SceneRenderer->Scene->GetExtension<FCSMSceneExtension>().GetCachedShadowMapDataRef(GetLightSceneInfo().Id, FMath::Max(CascadeSettings.ShadowSplitIndex, 0));
 
 	if (CachedShadowMapData.bCachedShadowMapHasPrimitives && CachedShadowMapData.ShadowMap.IsValid())
 	{
@@ -1885,7 +1886,7 @@ void FSceneRenderer::RenderShadowDepthMaps(FRDGBuilder& GraphBuilder, FDynamicSh
 		RDG_EVENT_SCOPE(GraphBuilder, "Cubemap %s %u^2", *LightNameWithLevel, TargetSize.X, TargetSize.Y);
 
 		// Only clear when we're not copying from a cached shadow map.
-		if (ProjectedShadowInfo->CacheMode != SDCM_MovablePrimitivesOnly || !Scene->GetCachedShadowMapDataRef(ProjectedShadowInfo->GetLightSceneInfo().Id, FMath::Max(ProjectedShadowInfo->CascadeSettings.ShadowSplitIndex, 0)).bCachedShadowMapHasPrimitives)
+		if (ProjectedShadowInfo->CacheMode != SDCM_MovablePrimitivesOnly || !Scene->GetExtension<FCSMSceneExtension>().GetCachedShadowMapDataRef(ProjectedShadowInfo->GetLightSceneInfo().Id, FMath::Max(ProjectedShadowInfo->CascadeSettings.ShadowSplitIndex, 0)).bCachedShadowMapHasPrimitives)
 		{
 			AddClearShadowDepthPass(GraphBuilder, ShadowDepthTexture);
 		}
diff --git a/Engine/Source/Runtime/Renderer/Private/ShadowSetup.cpp b/Engine/Source/Runtime/Renderer/Private/ShadowSetup.cpp
--- a/Engine/Source/Runtime/Renderer/Private/ShadowSetup.cpp
+++ b/Engine/Source/Runtime/Renderer/Private/ShadowSetup.cpp
@@ -90,6 +90,7 @@
 #include "ShaderParameterMacros.h"
 #include "ShadowRendering.h"
 #include "Shadows/ShadowSceneRenderer.h"
+#include "Shadows/CSMSceneExtension.h"
 #include "ShowFlags.h"
 #include "SimpleMeshDrawCommandPass.h"
 #include "StaticMeshBatch.h"
@@ -2425,7 +2426,7 @@ void FProjectedShadowInfo::FinalizeAddSubjectPrimitive(
 	{
 		checkSlow(!PrimitiveSceneInfo->Proxy->IsMeshShapeOftenMoving() || CacheMode == SDCM_CSMScrolling || CacheMode == SDCM_MovablePrimitivesOnly);
 
-		FCachedShadowMapData& CachedShadowMapData = PrimitiveSceneInfo->Scene->GetCachedShadowMapDataRef(LightSceneInfo->Id, CascadeSettings.ShadowSplitIndex);
+		FCachedShadowMapData& CachedShadowMapData = PrimitiveSceneInfo->Scene->GetExtension<FCSMSceneExtension>().GetCachedShadowMapDataRef(LightSceneInfo->Id, CascadeSettings.ShadowSplitIndex);
 
 		if (CacheMode == SDCM_StaticPrimitivesOnly)
 		{
@@ -3893,10 +3894,12 @@ void ComputeWholeSceneShadowCacheModes(
 		break;
 	}
 
+	FCSMSceneExtension& CSMShadowScene = Scene->GetExtension<FCSMSceneExtension>();
+
 	if (GCacheWholeSceneShadows 
 		&& (!bCubeShadowMap || RHISupportsGeometryShaders(GShaderPlatformForFeatureLevel[Scene->GetFeatureLevel()]) || (RHISupportsVertexShaderLayer(Scene->GetShaderPlatform()) && GRHISupportsArrayIndexFromAnyShader)))
 	{
-		TArray<FCachedShadowMapData>* CachedShadowMapDatas = Scene->GetCachedShadowMapDatas(LightSceneInfo->Id);
+		TArray<FCachedShadowMapData>* CachedShadowMapDatas = CSMShadowScene.GetCachedShadowMapDatas(LightSceneInfo->Id);
 
 		if (CachedShadowMapDatas)
 		{
@@ -3977,7 +3980,7 @@ void ComputeWholeSceneShadowCacheModes(
 			TryToCacheShadowMap(Scene, CachedShadowMapsSize, OutNumShadowMaps, OutCacheModes, *NumCachesUpdatedThisFrame,
 			[&]()
 			{
-				Scene->CachedShadowMaps.Add(LightSceneInfo->Id).Add(FCachedShadowMapData(InOutProjectedShadowInitializer, RealTime));
+				CSMShadowScene.GetAllCachedShadowMaps().Add(LightSceneInfo->Id).Add(FCachedShadowMapData(InOutProjectedShadowInitializer, RealTime));
 			},
 			[&]()
 			{
@@ -3988,7 +3991,7 @@ void ComputeWholeSceneShadowCacheModes(
 	{
 		OutNumShadowMaps = 1;
 		OutCacheModes[0] = SDCM_Uncached;
-		Scene->CachedShadowMaps.Remove(LightSceneInfo->Id);
+		CSMShadowScene.GetAllCachedShadowMaps().Remove(LightSceneInfo->Id);
 	}
 
 	if (OutNumShadowMaps > 0)
@@ -4016,9 +4019,11 @@ void ComputeViewDependentWholeSceneShadowCacheModes(
 	int32& OutNumShadowMaps,
 	EShadowDepthCacheMode* OutCacheModes)
 {
+	FCSMSceneExtension& CSMShadowScene = Scene->GetExtension<FCSMSceneExtension>();
+
 	if (CVarCSMCaching.GetValueOnAnyThread() == 1 && !UseVirtualShadowMaps(Scene->GetShaderPlatform(), Scene->GetFeatureLevel()))
 	{
-		TArray<FCachedShadowMapData>* CachedShadowMapDatas = Scene->GetCachedShadowMapDatas(LightSceneInfo->Id);
+		TArray<FCachedShadowMapData>* CachedShadowMapDatas = CSMShadowScene.GetCachedShadowMapDatas(LightSceneInfo->Id);
 
 		checkSlow(ProjectedShadowInitializer.CascadeSettings.ShadowSplitIndex >= 0);
 		FCachedShadowMapData* CachedShadowMapData = (CachedShadowMapDatas && CachedShadowMapDatas->Num() > ProjectedShadowInitializer.CascadeSettings.ShadowSplitIndex) ? &(*CachedShadowMapDatas)[ProjectedShadowInitializer.CascadeSettings.ShadowSplitIndex] : nullptr;
@@ -4126,7 +4131,7 @@ void ComputeViewDependentWholeSceneShadowCacheModes(
 			{
 				if (CachedShadowMapDatas == nullptr)
 				{
-					CachedShadowMapDatas = &Scene->CachedShadowMaps.Add(LightSceneInfo->Id);
+					CachedShadowMapDatas = &CSMShadowScene.GetAllCachedShadowMaps().Add(LightSceneInfo->Id);
 				}
 				CachedShadowMapDatas->Add(FCachedShadowMapData(ProjectedShadowInitializer, RealTime));
 				CachedShadowMapData = &CachedShadowMapDatas->Last();
@@ -4142,7 +4147,7 @@ void ComputeViewDependentWholeSceneShadowCacheModes(
 	{
 		OutNumShadowMaps = 1;
 		OutCacheModes[0] = SDCM_Uncached;
-		Scene->CachedShadowMaps.Remove(LightSceneInfo->Id);
+		CSMShadowScene.GetAllCachedShadowMaps().Remove(LightSceneInfo->Id);
 	}
 
 	if (OutNumShadowMaps > 0)
@@ -4556,7 +4561,7 @@ void FSceneRenderer::CreateWholeSceneProjectedShadow(
 						{
 							const bool bHasStaticPrimitives = ProjectedShadowInfo->HasSubjectPrims();
 							bRenderShadow = bHasStaticPrimitives;
-							FCachedShadowMapData& CachedShadowMapData = Scene->GetCachedShadowMapDataRef(ProjectedShadowInfo->GetLightSceneInfo().Id);
+							FCachedShadowMapData& CachedShadowMapData = Scene->GetExtension<FCSMSceneExtension>().GetCachedShadowMapDataRef(ProjectedShadowInfo->GetLightSceneInfo().Id);
 							CachedShadowMapData.bCachedShadowMapHasPrimitives = bHasStaticPrimitives;
 							CachedShadowMapData.bCachedShadowMapHasNaniteGeometry = ProjectedShadowInfo->bNaniteGeometry;
 						}
@@ -5141,6 +5146,8 @@ struct FGatherShadowPrimitivesPacket
 			}
 		}
 
+		const FCSMSceneExtension& CSMShadowScene = TaskData.Scene->GetExtension<FCSMSceneExtension>();
+
 		for (int32 ShadowIndex = 0, Num = TaskData.ViewDependentWholeSceneShadows.Num();ShadowIndex < Num;ShadowIndex++)
 		{
 			const FProjectedShadowInfo* RESTRICT ProjectedShadowInfo = TaskData.ViewDependentWholeSceneShadows[ShadowIndex];
@@ -5181,7 +5188,7 @@ struct FGatherShadowPrimitivesPacket
 
 				if (ProjectedShadowInfo->CacheMode == SDCM_MovablePrimitivesOnly)
 				{
-					const FCachedShadowMapData& CachedShadowMapData = TaskData.Scene->GetCachedShadowMapDataRef(LightSceneInfo.Id, ProjectedShadowInfo->CascadeSettings.ShadowSplitIndex);
+					const FCachedShadowMapData& CachedShadowMapData = CSMShadowScene.GetCachedShadowMapDataRef(LightSceneInfo.Id, ProjectedShadowInfo->CascadeSettings.ShadowSplitIndex);
 
 					// if the mesh is dynamic, it should cast csm.
 					bShadowRelevance = PrimitiveProxy->IsMeshShapeOftenMoving();
@@ -5192,7 +5199,7 @@ struct FGatherShadowPrimitivesPacket
 				}
 				else if (ProjectedShadowInfo->CacheMode == SDCM_CSMScrolling)
 				{
-					const FCachedShadowMapData& CachedShadowMapData = TaskData.Scene->GetCachedShadowMapDataRef(LightSceneInfo.Id, ProjectedShadowInfo->CascadeSettings.ShadowSplitIndex);
+					const FCachedShadowMapData& CachedShadowMapData = CSMShadowScene.GetCachedShadowMapDataRef(LightSceneInfo.Id, ProjectedShadowInfo->CascadeSettings.ShadowSplitIndex);
 
 					// if the mesh is dynamic or a new added mesh, it should cast csm.
 					bShadowRelevance = PrimitiveProxy->IsMeshShapeOftenMoving() || !CachedShadowMapData.StaticShadowSubjectPersistentPrimitiveIdMap[PrimitiveSceneInfo->GetPersistentIndex().Index];
@@ -5281,7 +5288,7 @@ struct FGatherShadowPrimitivesPacket
 
 			if (ProjectedShadowInfo->CacheMode == SDCM_StaticPrimitivesOnly && ProjectedShadowInfo->IsWholeSceneDirectionalShadow())
 			{
-				FCachedShadowMapData* CachedShadowMapData = const_cast<FCachedShadowMapData*>(TaskData.Scene->GetCachedShadowMapData(ProjectedShadowInfo->GetLightSceneInfo().Id, ProjectedShadowInfo->CascadeSettings.ShadowSplitIndex));
+				FCachedShadowMapData* CachedShadowMapData = const_cast<FCachedShadowMapData*>(TaskData.Scene->GetExtension<FCSMSceneExtension>().GetCachedShadowMapData(ProjectedShadowInfo->GetLightSceneInfo().Id, ProjectedShadowInfo->CascadeSettings.ShadowSplitIndex));
 				checkSlow(CachedShadowMapData != nullptr);
 				CachedShadowMapData->bCachedShadowMapHasPrimitives = ProjectedShadowInfo->HasSubjectPrims();
 			}
@@ -5795,7 +5802,7 @@ void FSceneRenderer::AddViewDependentWholeSceneShadowsForView(
 
 								if (CacheMode[CacheModeIndex] != SDCM_Uncached)
 								{
-									FCachedShadowMapData& CachedShadowMapData = Scene->GetCachedShadowMapDataRef(LightSceneInfo.Id, ProjectedShadowInfo->CascadeSettings.ShadowSplitIndex);
+									FCachedShadowMapData& CachedShadowMapData = Scene->GetExtension<FCSMSceneExtension>().GetCachedShadowMapDataRef(LightSceneInfo.Id, ProjectedShadowInfo->CascadeSettings.ShadowSplitIndex);
 
 									CachedShadowMapData.LastFrameExtraStaticShadowSubjects = 0;
 
@@ -5957,7 +5964,7 @@ void FSceneRenderer::AllocateShadowDepthTargets(FDynamicShadowsTaskData& TaskDat
 		}
 	}
 
-	SET_MEMORY_STAT(STAT_CachedShadowmapMemory, Scene->GetCachedWholeSceneShadowMapsSize());
+	SET_MEMORY_STAT(STAT_CachedShadowmapMemory, Scene->GetExtension<FCSMSceneExtension>().GetCachedWholeSceneShadowMapsSize());
 	SET_MEMORY_STAT(STAT_ShadowmapAtlasMemory, SortedShadowsForShadowDepthPass.ComputeMemorySize());
 }
 
@@ -6008,7 +6015,7 @@ void FSceneRenderer::AllocateAtlasedShadowDepthTargets(
 
 		if (ProjectedShadowInfo->CacheMode == SDCM_MovablePrimitivesOnly && !ProjectedShadowInfo->HasSubjectPrims())
 		{
-			const FCachedShadowMapData& CachedShadowMapData = Scene->GetCachedShadowMapDataRef(ProjectedShadowInfo->GetLightSceneInfo().Id);
+			const FCachedShadowMapData& CachedShadowMapData = Scene->GetExtension<FCSMSceneExtension>().GetCachedShadowMapDataRef(ProjectedShadowInfo->GetLightSceneInfo().Id);
 			ProjectedShadowInfo->X = ProjectedShadowInfo->Y = 0;
 			ProjectedShadowInfo->bAllocated = true;
 			// Skip the shadow depth pass since there are no movable primitives to composite, project from the cached shadowmap directly which contains static primitive depths
@@ -6130,6 +6137,7 @@ public:
 
 void FSceneRenderer::AllocateCachedShadowDepthTargets(FRHICommandListBase& RHICmdList, TConstArrayView<FProjectedShadowInfo*> CachedShadows)
 {
+	FCSMSceneExtension& CSMShadowScene = Scene->GetExtension<FCSMSceneExtension>();
 	for (int32 ShadowIndex = 0; ShadowIndex < CachedShadows.Num(); ShadowIndex++)
 	{
 		FProjectedShadowInfo* ProjectedShadowInfo = CachedShadows[ShadowIndex];
@@ -6146,7 +6154,7 @@ void FSceneRenderer::AllocateCachedShadowDepthTargets(FRHICommandListBase& RHICm
 		GRenderTargetPool.FindFreeElement(RHICmdList, ShadowMapDesc2D, ShadowMap.RenderTargets.DepthTarget, RTDebugName);
 
 		check(ProjectedShadowInfo->CacheMode == SDCM_StaticPrimitivesOnly);
-		FCachedShadowMapData& CachedShadowMapData = Scene->GetCachedShadowMapDataRef(ProjectedShadowInfo->GetLightSceneInfo().Id, FMath::Max(ProjectedShadowInfo->CascadeSettings.ShadowSplitIndex, 0));
+		FCachedShadowMapData& CachedShadowMapData = CSMShadowScene.GetCachedShadowMapDataRef(ProjectedShadowInfo->GetLightSceneInfo().Id, FMath::Max(ProjectedShadowInfo->CascadeSettings.ShadowSplitIndex, 0));
 		CachedShadowMapData.ShadowMap = ShadowMap.RenderTargets;
 
 		ProjectedShadowInfo->X = ProjectedShadowInfo->Y = 0;
@@ -6222,6 +6230,7 @@ void FSceneRenderer::AllocateOnePassPointLightDepthTargets(FRHICommandListBase&
 {
 	if (FeatureLevel >= ERHIFeatureLevel::SM5 || DoesRuntimeSupportOnePassPointLightShadows(ShaderPlatform))
 	{
+		FCSMSceneExtension& CSMShadowScene = Scene->GetExtension<FCSMSceneExtension>();
 		for (int32 ShadowIndex = 0; ShadowIndex < WholeScenePointShadows.Num(); ShadowIndex++)
 		{
 			FProjectedShadowInfo* ProjectedShadowInfo = WholeScenePointShadows[ShadowIndex];
@@ -6229,7 +6238,7 @@ void FSceneRenderer::AllocateOnePassPointLightDepthTargets(FRHICommandListBase&
 
 			if (ProjectedShadowInfo->CacheMode == SDCM_MovablePrimitivesOnly && !ProjectedShadowInfo->HasSubjectPrims())
 			{
-				const FCachedShadowMapData& CachedShadowMapData = Scene->GetCachedShadowMapDataRef(ProjectedShadowInfo->GetLightSceneInfo().Id);
+				const FCachedShadowMapData& CachedShadowMapData = CSMShadowScene.GetCachedShadowMapDataRef(ProjectedShadowInfo->GetLightSceneInfo().Id);
 				ProjectedShadowInfo->X = ProjectedShadowInfo->Y = 0;
 				ProjectedShadowInfo->bAllocated = true;
 				// Skip the shadow depth pass since there are no movable primitives to composite, project from the cached shadowmap directly which contains static primitive depths
@@ -6247,7 +6256,7 @@ void FSceneRenderer::AllocateOnePassPointLightDepthTargets(FRHICommandListBase&
 
 				if (ProjectedShadowInfo->CacheMode == SDCM_StaticPrimitivesOnly)
 				{
-					FCachedShadowMapData& CachedShadowMapData = Scene->GetCachedShadowMapDataRef(ProjectedShadowInfo->GetLightSceneInfo().Id);
+					FCachedShadowMapData& CachedShadowMapData = CSMShadowScene.GetCachedShadowMapDataRef(ProjectedShadowInfo->GetLightSceneInfo().Id);
 					CachedShadowMapData.ShadowMap = ShadowMapCubemap.RenderTargets;
 				}
 
@@ -6376,7 +6385,7 @@ void FSceneRenderer::CreateDynamicShadows(FDynamicShadowsTaskData& TaskData)
 				if (CachedShadowMapsSize < 0)
 				{
 					// This call is quite expensive, so compute it lazily.
-					CachedShadowMapsSize = Scene->GetCachedWholeSceneShadowMapsSize();
+					CachedShadowMapsSize = Scene->GetExtension<FCSMSceneExtension>().GetCachedWholeSceneShadowMapsSize();
 				}
 				return CachedShadowMapsSize;
 			};
@@ -6560,21 +6569,9 @@ void FSceneRenderer::CreateDynamicShadows(FDynamicShadowsTaskData& TaskData)
 		}
 
 		// Remove cache entries that haven't been used in a while
-		for (TMap<FLightSceneId, TArray<FCachedShadowMapData>>::TIterator CachedShadowMapIt(Scene->CachedShadowMaps); CachedShadowMapIt; ++CachedShadowMapIt)
-		{
-			TArray<FCachedShadowMapData>& ShadowMapDatas = CachedShadowMapIt.Value();
-
-			for (auto& ShadowMapData : ShadowMapDatas)
-			{
-				if (ShadowMapData.ShadowMap.IsValid() && ViewFamily.Time.GetRealTimeSeconds() - ShadowMapData.LastUsedTime > 2.0f)
-				{
-					ShadowMapData.InvalidateCachedShadow();
-				}
-			}
-		}
+		Scene->GetExtension<FCSMSceneExtension>().RemoveExpiredCacheEntries(ViewFamily.Time.GetRealTimeSeconds());
 
 		CSV_CUSTOM_STAT(LightCount, UpdatedShadowMaps, float(NumPointShadowCachesUpdatedThisFrame + NumSpotShadowCachesUpdatedThisFrame + NumCSMCachesUpdatedThisFrame), ECsvCustomStatOp::Set);
-		CSV_CUSTOM_STAT_GLOBAL(ShadowCacheUsageMB, (float(Scene->GetCachedWholeSceneShadowMapsSize()) / 1024) / 1024, ECsvCustomStatOp::Set);
 
 		// Calculate visibility of the projected shadows.
 		InitProjectedShadowVisibility(TaskData);
@@ -6651,7 +6648,7 @@ void FSceneRenderer::FilterDynamicShadows(FDynamicShadowsTaskData& TaskData)
 			// Skip the shadow rendering if there is no primitive in the cached shadow map and the current shadow map either
 			else if ((ProjectedShadowInfo->CacheMode == SDCM_StaticPrimitivesOnly || (ProjectedShadowInfo->CacheMode == SDCM_MovablePrimitivesOnly && !ProjectedShadowInfo->IsWholeSceneDirectionalShadow())) && !ProjectedShadowInfo->HasSubjectPrims())
 			{
-				const FCachedShadowMapData& CachedShadowMapData = Scene->GetCachedShadowMapDataRef(ProjectedShadowInfo->GetLightSceneInfo().Id, FMath::Max(ProjectedShadowInfo->CascadeSettings.ShadowSplitIndex, 0));
+				const FCachedShadowMapData& CachedShadowMapData = Scene->GetExtension<FCSMSceneExtension>().GetCachedShadowMapDataRef(ProjectedShadowInfo->GetLightSceneInfo().Id, FMath::Max(ProjectedShadowInfo->CascadeSettings.ShadowSplitIndex, 0));
 
 				// A shadowmap for movable primitives when there are no movable primitives would normally read directly from the cached shadowmap
 				// However if the cached shadowmap also had no primitives then we need to skip rendering the shadow entirely
diff --git a/Engine/Source/Runtime/Renderer/Private/Shadows/CSMSceneExtension.cpp b/Engine/Source/Runtime/Renderer/Private/Shadows/CSMSceneExtension.cpp
new file mode 100644
--- /dev/null
+++ b/Engine/Source/Runtime/Renderer/Private/Shadows/CSMSceneExtension.cpp
@@ -0,0 +1,147 @@
+// Copyright Epic Games, Inc. All Rights Reserved.
+
+#include "CSMSceneExtension.h"
+#include "ScenePrivate.h"
+#include "ScenePrimitiveUpdates.h"
+
+IMPLEMENT_SCENE_EXTENSION(FCSMSceneExtension);
+
+FCachedShadowMapData::FCachedShadowMapData(const FWholeSceneProjectedShadowInitializer& InInitializer, float InLastUsedTime)
+	: Initializer(InInitializer)
+	, LastUsedTime(InLastUsedTime)
+	, bCachedShadowMapHasPrimitives(true)
+	, bCachedShadowMapHasNaniteGeometry(false)
+	, LastFrameExtraStaticShadowSubjects(0)
+{
+}
+
+void FCSMSceneExtension::InitExtension(FScene& InScene)
+{
+}
+
+ISceneExtensionUpdater* FCSMSceneExtension::CreateUpdater()
+{
+	return new FUpdater(*this);
+}
+
+void FCSMSceneExtension::RemoveExpiredCacheEntries(float CurrentRealTime)
+{
+	for (TMap<FLightSceneId, TArray<FCachedShadowMapData>>::TIterator CachedShadowMapIt(CachedShadowMaps); CachedShadowMapIt; ++CachedShadowMapIt)
+	{
+		TArray<FCachedShadowMapData>& ShadowMapDatas = CachedShadowMapIt.Value();
+
+		for (auto& ShadowMapData : ShadowMapDatas)
+		{
+			if (ShadowMapData.ShadowMap.IsValid() && CurrentRealTime - ShadowMapData.LastUsedTime > 2.0f)
+			{
+				ShadowMapData.InvalidateCachedShadow();
+			}
+		}
+	}
+
+	CSV_CUSTOM_STAT_GLOBAL(ShadowCacheUsageMB, (float(GetCachedWholeSceneShadowMapsSize()) / 1024) / 1024, ECsvCustomStatOp::Set);
+}
+
+int64 FCSMSceneExtension::GetCachedWholeSceneShadowMapsSize() const
+{
+	int64 CachedShadowmapMemory = 0;
+
+	for (TMap<FLightSceneId, TArray<FCachedShadowMapData>>::TConstIterator CachedShadowMapIt(CachedShadowMaps); CachedShadowMapIt; ++CachedShadowMapIt)
+	{
+		const TArray<FCachedShadowMapData>& ShadowMapDatas = CachedShadowMapIt.Value();
+
+		for (const auto& ShadowMapData : ShadowMapDatas)
+		{
+			if (ShadowMapData.ShadowMap.IsValid())
+			{
+				CachedShadowmapMemory += ShadowMapData.ShadowMap.ComputeMemorySize();
+			}
+		}
+	}
+
+	return CachedShadowmapMemory;
+}
+
+TArray<FCachedShadowMapData>* FCSMSceneExtension::GetCachedShadowMapDatas(FLightSceneId LightId)
+{
+	return CachedShadowMaps.Find(LightId);
+}
+
+const TArray<FCachedShadowMapData>* FCSMSceneExtension::GetCachedShadowMapDatas(FLightSceneId LightId) const
+{
+	return CachedShadowMaps.Find(LightId);
+}
+
+FCachedShadowMapData& FCSMSceneExtension::GetCachedShadowMapDataRef(FLightSceneId LightId, int32 ShadowMapIndex)
+{
+	TArray<FCachedShadowMapData>& CachedShadowMapDatas = CachedShadowMaps.FindChecked(LightId);
+	checkSlow(ShadowMapIndex >= 0 && ShadowMapIndex < CachedShadowMapDatas.Num());
+	return CachedShadowMapDatas[ShadowMapIndex];
+}
+
+const FCachedShadowMapData& FCSMSceneExtension::GetCachedShadowMapDataRef(FLightSceneId LightId, int32 ShadowMapIndex) const
+{
+	const TArray<FCachedShadowMapData>& CachedShadowMapDatas = CachedShadowMaps.FindChecked(LightId);
+	checkSlow(ShadowMapIndex >= 0 && ShadowMapIndex < CachedShadowMapDatas.Num());
+	return CachedShadowMapDatas[ShadowMapIndex];
+}
+
+const FCachedShadowMapData* FCSMSceneExtension::GetCachedShadowMapData(FLightSceneId LightId, int32 ShadowMapIndex) const
+{
+	const TArray<FCachedShadowMapData>* CachedShadowMapDatas = CachedShadowMaps.Find(LightId);
+	if (CachedShadowMapDatas && ShadowMapIndex >= 0 && ShadowMapIndex < CachedShadowMapDatas->Num())
+	{
+		return &(*CachedShadowMapDatas)[ShadowMapIndex];
+	}
+	return nullptr;
+}
+
+void FCSMSceneExtension::FUpdater::PreSceneUpdate(FRDGBuilder& GraphBuilder, const FScenePreUpdateChangeSet& ChangeSet)
+{
+	// Invalidate cached shadows whenever a shadow-casting primitive is removed
+	FScene* ScenePtr = static_cast<FScene*>(&CSMScene.Scene);
+
+	for (const FLightSceneInfo* LightSceneInfo : ScenePtr->DirectionalLights)
+	{
+		TArray<FCachedShadowMapData>* CachedShadowMapDatas = CSMScene.GetCachedShadowMapDatas(LightSceneInfo->Id);
+
+		if (CachedShadowMapDatas)
+		{
+			for (FCachedShadowMapData& CachedShadowMapData : *CachedShadowMapDatas)
+			{
+				for (FPersistentPrimitiveIndex PersistentPrimitiveIndex : ChangeSet.RemovedPrimitiveIds)
+				{
+					if (CachedShadowMapData.StaticShadowSubjectPersistentPrimitiveIdMap[PersistentPrimitiveIndex.Index] == true)
+					{
+						CachedShadowMapData.InvalidateCachedShadow();
+						break;
+					}
+				}
+			}
+		}
+	}
+}
+
+void FCSMSceneExtension::FUpdater::PostSceneUpdate(FRDGBuilder& GraphBuilder, const FScenePostUpdateChangeSet& ChangeSet)
+{
+	// Resize primitive tracking bitmaps to match new scene size
+	FScene* ScenePtr = static_cast<FScene*>(&CSMScene.Scene);
+	const int32 MaxPersistentPrimitiveIndex = ScenePtr->GetMaxPersistentPrimitiveIndex();
+
+	for (auto& Pair : CSMScene.CachedShadowMaps)
+	{
+		for (FCachedShadowMapData& CachedShadowMapData : Pair.Value)
+		{
+			CachedShadowMapData.StaticShadowSubjectPersistentPrimitiveIdMap.SetNum(MaxPersistentPrimitiveIndex, false);
+		}
+	}
+}
+
+void FCSMSceneExtension::FUpdater::PreLightsUpdate(FRDGBuilder& GraphBuilder, const FLightSceneChangeSet& LightSceneChangeSet)
+{
+	// Remove cached shadow maps for removed lights
+	for (FLightSceneId LightId : LightSceneChangeSet.RemovedLightIds)
+	{
+		CSMScene.CachedShadowMaps.Remove(LightId);
+	}
+}
diff --git a/Engine/Source/Runtime/Renderer/Private/Shadows/CSMSceneExtension.h b/Engine/Source/Runtime/Renderer/Private/Shadows/CSMSceneExtension.h
new file mode 100644
--- /dev/null
+++ b/Engine/Source/Runtime/Renderer/Private/Shadows/CSMSceneExtension.h
@@ -0,0 +1,120 @@
+// Copyright Epic Games, Inc. All Rights Reserved.
+
+#pragma once
+
+#include "CoreMinimal.h"
+#include "SceneExtensions.h"
+#include "SceneManagement.h"
+#include "SceneRendering.h"
+
+class FScene;
+struct FLightSceneChangeSet;
+class FScenePreUpdateChangeSet;
+class FScenePostUpdateChangeSet;
+class FSceneUniformBuffer;
+
+
+/**
+ * Cached shadow map data for whole-scene shadows (CSM, point, spot).
+ */
+class FCachedShadowMapData
+{
+public:
+	FWholeSceneProjectedShadowInitializer Initializer;
+	FShadowMapRenderTargetsRefCounted ShadowMap;
+	float LastUsedTime;
+	bool bCachedShadowMapHasPrimitives;
+	bool bCachedShadowMapHasNaniteGeometry;
+
+	/**
+	 * The static meshes cast shadow on this cached csm
+	 */
+	TBitArray<> StaticShadowSubjectPersistentPrimitiveIdMap;
+
+	FIntPoint ShadowBufferResolution;
+	FVector PreShadowTranslation;
+	float MaxSubjectZ;
+	float MinSubjectZ;
+
+	/**
+	 * The extra static meshes cast shadow in last frame, if it exceeds the r.Shadow.MaxCSMScrollingStaticShadowSubjects, the cached csm should be rebuilt.
+	 */
+	int32 LastFrameExtraStaticShadowSubjects;
+
+	void InvalidateCachedShadow()
+	{
+		ShadowMap.Release();
+		StaticShadowSubjectPersistentPrimitiveIdMap.SetRange(0, StaticShadowSubjectPersistentPrimitiveIdMap.Num(), false);
+	}
+
+	FCachedShadowMapData(const FWholeSceneProjectedShadowInitializer& InInitializer, float InLastUsedTime);
+};
+
+/**
+ * Scene extension for Cascaded Shadow Map (CSM) state management.
+ * This extension owns all persistent CSM-related state.
+ */
+class FCSMSceneExtension : public ISceneExtension
+{
+	DECLARE_SCENE_EXTENSION(RENDERER_API, FCSMSceneExtension);
+
+public:
+	using ISceneExtension::ISceneExtension;
+
+	//~ ISceneExtension interface
+	virtual void InitExtension(FScene& InScene) override;
+	virtual ISceneExtensionUpdater* CreateUpdater() override;
+
+	/**
+	 * Remove cached shadow map entries that haven't been used recently and report stats.
+	 */
+	void RemoveExpiredCacheEntries(float CurrentRealTime);
+
+	/**
+	 * Get the total memory used by all cached shadow maps.
+	 */
+	int64 GetCachedWholeSceneShadowMapsSize() const;
+
+	/**
+	 * Get cached shadow map data array for a light, or nullptr if not found.
+	 */
+	TArray<FCachedShadowMapData>* GetCachedShadowMapDatas(FLightSceneId LightID);
+	const TArray<FCachedShadowMapData>* GetCachedShadowMapDatas(FLightSceneId LightID) const;
+
+	/**
+	 * Get cached shadow map data for a light at a specific index (checked slow).
+	 */
+	FCachedShadowMapData& GetCachedShadowMapDataRef(FLightSceneId LightID, int32 ShadowMapIndex = 0);
+	const FCachedShadowMapData& GetCachedShadowMapDataRef(FLightSceneId LightID, int32 ShadowMapIndex = 0) const;
+
+	/**
+	 * Get cached shadow map data for a light at a specific index, or nullptr if invalid.
+	 */
+	const FCachedShadowMapData* GetCachedShadowMapData(FLightSceneId LightID, int32 ShadowMapIndex = 0) const;
+
+	/**
+	 * Get the map of all cached shadow maps (mutable).
+	 */
+	TMap<FLightSceneId, TArray<FCachedShadowMapData>>& GetAllCachedShadowMaps() { return CachedShadowMaps; }
+	const TMap<FLightSceneId, TArray<FCachedShadowMapData>>& GetAllCachedShadowMaps() const { return CachedShadowMaps; }
+
+private:
+	class FUpdater : public ISceneExtensionUpdater
+	{
+	public:
+		DECLARE_SCENE_EXTENSION_UPDATER(FUpdater, FCSMSceneExtension);
+
+		FUpdater(FCSMSceneExtension& InCSMScene) : CSMScene(InCSMScene) {}
+
+		virtual void PreSceneUpdate(FRDGBuilder& GraphBuilder, const FScenePreUpdateChangeSet& ChangeSet) override;
+		virtual void PostSceneUpdate(FRDGBuilder& GraphBuilder, const FScenePostUpdateChangeSet& ChangeSet) override;
+		virtual void PreLightsUpdate(FRDGBuilder& GraphBuilder, const FLightSceneChangeSet& LightSceneChangeSet) override;
+
+		FCSMSceneExtension& CSMScene;
+	};
+
+	friend class FUpdater;
+
+	/** Map from light id to the cached shadowmap data for that light. */
+	TMap<FLightSceneId, TArray<FCachedShadowMapData>> CachedShadowMaps;
+};
diff --git a/Engine/Source/Runtime/Renderer/Private/Shadows/ShadowScene.h b/Engine/Source/Runtime/Renderer/Private/Shadows/ShadowScene.h
--- a/Engine/Source/Runtime/Renderer/Private/Shadows/ShadowScene.h
+++ b/Engine/Source/Runtime/Renderer/Private/Shadows/ShadowScene.h
@@ -21,7 +21,7 @@ class FScenePreUpdateChangeSet;
 class FScenePostUpdateChangeSet;
 
 /**
- * Persistent scene-representation of for shadow rendering.
+ * Persistent scene-representation for VSM shadow rendering.
  */
 class FShadowScene : public ISceneExtension
 {
