Modified QuantizeViewRectMin and QuantizeSceneBufferSize to always quantize to 8x8.
- rendering features such as Nanite, Substrate, MegaLights and Lumen rely on it.

#Synced-CL 52158667

diff --git a/Engine/Source/Runtime/RenderCore/Private/RenderUtils.cpp b/Engine/Source/Runtime/RenderCore/Private/RenderUtils.cpp
--- a/Engine/Source/Runtime/RenderCore/Private/RenderUtils.cpp
+++ b/Engine/Source/Runtime/RenderCore/Private/RenderUtils.cpp
@@ -1184,15 +1184,14 @@ RENDERCORE_API FBufferRHIRef& GetUnitCubeAABBVertexBuffer()
 
 RENDERCORE_API void QuantizeSceneBufferSize(const FIntPoint& InBufferSize, FIntPoint& OutBufferSize, const uint32 SuggestedDivisor)
 {
-	// Ensure sizes are dividable by SUBSTRATE_TILE_SIZE (==8) 2d tiles to make it more convenient.
-	const uint32 SubstrateDividableBy = SUBSTRATE_TILE_SIZE;
-	static_assert(SubstrateDividableBy % 8 == 0, "A lot of graphic algorithms where previously assuming DividableBy >= 4");
+	// A lot of rendering features assume that buffer size is dividable by 8
+	const uint32 RenderingDividableBy = 8;
 
-	// Ensure sizes are dividable by the ideal group size for 2d tiles to make it more convenient.
-	const uint32 LegacyDividableBy = 4;
-	static_assert(LegacyDividableBy % 4 == 0, "A lot of graphic algorithms where previously assuming DividableBy == 4");
+	static_assert(RenderingDividableBy % SUBSTRATE_TILE_SIZE == 0, "Ensure sizes are dividable by SUBSTRATE_TILE_SIZE");
 
-	const uint32 DividableBy = FMath::Max(Substrate::IsSubstrateEnabled() ? SubstrateDividableBy : LegacyDividableBy, SuggestedDivisor);
+	const uint32 DividableBy = FMath::Max(RenderingDividableBy, SuggestedDivisor);
+
+	checkf(DividableBy % RenderingDividableBy == 0, TEXT("A lot of rendering features assume that buffer size is dividable by %u"), RenderingDividableBy);
 
 	const uint32 Mask = ~(DividableBy - 1);
 	OutBufferSize.X = (InBufferSize.X + DividableBy - 1) & Mask;
diff --git a/Engine/Source/Runtime/Renderer/Private/SceneRendering.cpp b/Engine/Source/Runtime/Renderer/Private/SceneRendering.cpp
--- a/Engine/Source/Runtime/Renderer/Private/SceneRendering.cpp
+++ b/Engine/Source/Runtime/Renderer/Private/SceneRendering.cpp
@@ -3319,13 +3319,8 @@ FIntPoint FSceneRenderer::ApplyResolutionFraction(const FIntPoint& UnscaledViewS
 FIntPoint FSceneRenderer::QuantizeViewRectMin(const FIntPoint& ViewRectMin)
 {
 	FIntPoint Out;
-
-	// Some code paths of Nanite require that view rect is aligned on 8x8 boundary.
-	static const auto EnableNaniteCVar = IConsoleManager::Get().FindConsoleVariable(TEXT("r.Nanite"));
-	const bool bNaniteEnabled = (EnableNaniteCVar != nullptr) ? (EnableNaniteCVar->GetInt() != 0) : true;
-	const int kMinimumNaniteDivisor = 8;	// HTILE size
-
-	QuantizeSceneBufferSize(ViewRectMin, Out, bNaniteEnabled ? kMinimumNaniteDivisor : 0);
+	// some features (e.g. Nanite, Substrate, Lumen, MegaLights) require that view rect is aligned on 8x8 boundary.
+	QuantizeSceneBufferSize(ViewRectMin, Out, 8);
 	return Out;
 }
 
@@ -3513,7 +3508,7 @@ void FSceneRenderer::PrepareViewRectsForRendering()
 	// If not supporting screen percentage, bypass all computation.
 	// This bypasses not only resolution scaling, but also quantizing the ViewRect and shifting letterboxed ViewRects to the top left for post-processing.
 	// The only configurations that do not support screen percentage are Forward Mobile non-HDR and Forward Mobile w/ Tonemap Subpass.
-	// These configurations do not support post-processing or any features that require ViewRect quantization (e.g. Nanite, Substrate), so it's safe to skip.
+	// These configurations do not support post-processing or any features that require ViewRect quantization (e.g. Nanite, Substrate, Lumen, MegaLights), so it's safe to skip.
 	if (!ViewFamily.SupportsScreenPercentage())
 	{
 		DynamicResolutionFractions[GDynamicPrimaryResolutionFraction] = 1.0f;
