StochasticLighting pass: don't copy full res depth and normal for LumenIrradianceFieldGather, which doesn't use it
* Saves .29ms on Switch 2 docked at 1440p
* Saves .1ms on Switch 2 docked at 900p

#Synced-CL 52184307

diff --git a/Engine/Source/Runtime/Renderer/Private/StochasticLighting/StochasticLighting.cpp b/Engine/Source/Runtime/Renderer/Private/StochasticLighting/StochasticLighting.cpp
--- a/Engine/Source/Runtime/Renderer/Private/StochasticLighting/StochasticLighting.cpp
+++ b/Engine/Source/Runtime/Renderer/Private/StochasticLighting/StochasticLighting.cpp
@@ -849,7 +849,7 @@ void FDeferredShadingSceneRenderer::StochasticLightingTileClassificationMark(FRD
 			const bool bTileClassifySubstrate = Substrate::UsesStochasticLightingClassification(View.GetShaderPlatform());
 			const bool bNeedsReprojection = bLumenScreenProbeGather || bTileClassifyMegaLights;
 			// Always copy if MegaLights is enabled. The spatial denoiser needs the normal copy
-			const bool bCopyDepthAndNormal = (View.ViewState && !View.bStatePrevViewInfoIsReadOnly) || bTileClassifyMegaLights;
+			const bool bCopyDepthAndNormal = ((View.ViewState && !View.bStatePrevViewInfoIsReadOnly) && bTileClassifyLumen) || bTileClassifyMegaLights;
 			const bool bDownsampleDepthAndNormal2x1 = bTileClassifyMegaLights && MegaLightsDownsampleFactor == FIntPoint(2, 1);
 			const bool bDownsampleDepthAndNormal2x2 = (bLumenScreenProbeGather && LumenScreenProbeGather::IsUsingDownsampledDepthAndNormal(View))
 				|| (bLumenIrradianceFieldGather && LumenIrradianceFieldGather::IsUsingDownsampledDepthAndNormal(View))
