Move ACES calculations to a secondary LUT for speed

#Synced-CL 52207145

diff --git a/Engine/Config/BaseScalability.ini b/Engine/Config/BaseScalability.ini
--- a/Engine/Config/BaseScalability.ini
+++ b/Engine/Config/BaseScalability.ini
@@ -517,6 +517,7 @@ r.FastBlurThreshold=0
 r.Upscale.Quality=1
 r.LightShaftQuality=0
 r.Filter.SizeScale=0.6
+r.LUT.Size=32
 r.Tonemapper.Quality=0
 
 [PostProcessQuality@1]
@@ -537,6 +538,7 @@ r.FastBlurThreshold=2
 r.Upscale.Quality=2
 r.LightShaftQuality=0
 r.Filter.SizeScale=0.7
+r.LUT.Size=32
 r.Tonemapper.Quality=2
 
 ; DOF settings.
@@ -570,6 +572,7 @@ r.FastBlurThreshold=3
 r.Upscale.Quality=2
 r.LightShaftQuality=1
 r.Filter.SizeScale=0.8
+r.LUT.Size=32
 r.Tonemapper.Quality=5
 
 ; DOF settings.
@@ -605,6 +608,7 @@ r.FastBlurThreshold=100
 r.Upscale.Quality=3
 r.LightShaftQuality=1
 r.Filter.SizeScale=1
+r.LUT.Size=32
 r.Tonemapper.Quality=5
 
 ; DOF settings.
@@ -643,6 +647,7 @@ r.FastBlurThreshold=100
 r.Upscale.Quality=3
 r.LightShaftQuality=1
 r.Filter.SizeScale=1
+r.LUT.Size=64
 r.Tonemapper.Quality=5
 
 ; DOF settings:
diff --git a/Engine/Shaders/Private/PostProcessCombineLUTs.usf b/Engine/Shaders/Private/PostProcessCombineLUTs.usf
--- a/Engine/Shaders/Private/PostProcessCombineLUTs.usf
+++ b/Engine/Shaders/Private/PostProcessCombineLUTs.usf
@@ -7,10 +7,6 @@
 	also be made in PostProcessDeviceEncodingOnly.usf.
 =============================================================================*/
 
-#define USE_COLOR_MATRIX	1
-#define USE_SHADOW_TINT		1
-#define USE_CONTRAST		1
-
 #include "Common.ush"
 #include "PostProcessCommon.ush"
 #include "TonemapCommon.ush"
@@ -50,6 +46,7 @@ float3x3 GetLimitingXYZToRgb()
 uint bIsTemperatureWhiteBalance;
 
 float LUTSize;
+float InnerLUTSize;
 
 float WhiteTemp;
 float WhiteTint;
@@ -82,25 +79,6 @@ float  ColorCorrectionShadowsMax;
 float  ColorCorrectionHighlightsMin;
 float  ColorCorrectionHighlightsMax;
 
-float4 ACESMinMaxData;
-float4 ACESMidData;
-float4 ACESCoefsLow_0;
-float4 ACESCoefsHigh_0;
-float ACESCoefsLow_4;
-float ACESCoefsHigh_4;
-float ACESSceneColorMultiplier;
-float ACESGamutCompression;
-
-// ACES 2.0
-Texture2D<float> ACESReachTable;
-Texture2D<float> ACESGamutTable;
-Texture2D<float> ACESGammaTable;
-Texture2D<float> ACESReachTableSDR;
-Texture2D<float> ACESGamutTableSDR;
-Texture2D<float> ACESGammaTableSDR;
-uint bScaleWhite;
-uint bScaleWhiteSDR;
-
 static const float DefaultTemperature = 6500.0;
 static const float DefaultTint = 0.0;
 static const float Epsilon = 1e-8;
@@ -164,9 +142,6 @@ uint OutputDevice;
 uint OutputGamut;
 float OutputMaxLuminance;
 
-float HDRPaperWhiteRatio;
-float InvHDRPaperWhiteRatio;
-
 uint GetOutputDevice()
 {
 #if OUTPUT_DEVICE_SRGB
@@ -267,68 +242,60 @@ float3 ComputeGradedColor(float3 ColorAP1)
 	return GradedColor;
 }
 
-float3 ComputeOutputACES(
-	float3 GradedColor,
-	float PeakLuminance,
-	float3x3 LimitingRgbToXYZ,
-	float3x3 LimitingXYZToRgb,
-	float3x3 AP1_2_Output)
+#define INVEOTF_LINEAR 0
+#define INVEOTF_PQ 1
+
+float3 ApplyInvEOTF(float3 Color, int InvEOTF)
 {
-#if USE_ACES_2
-	float3 ODTColor = ACESOutputTransform(
-		GradedColor * ACESSceneColorMultiplier,
-		(float3x3)WorkingColorSpace.ToAP0,
-		PeakLuminance,
-		ACESReachTable,
-		ACESGamutTable,
-		ACESGammaTable,
-		bScaleWhite,
-		LimitingRgbToXYZ,
-		LimitingXYZToRgb,
-		GetColorSpace(OutputGamut));
-	return ODTColor;
-#else
-	FACESTonemapParams AcesParams = ComputeACESTonemapParams(ACESMinMaxData, ACESMidData, ACESCoefsLow_0, ACESCoefsHigh_0, ACESCoefsLow_4, ACESCoefsHigh_4, ACESSceneColorMultiplier, ACESGamutCompression);
-	float3 ODTColor = ACESOutputTransform( GradedColor, (float3x3)WorkingColorSpace.ToAP0, AcesParams);
-	float3 OutDeviceColor = mul(AP1_2_Output, ODTColor);
-	return OutDeviceColor;
-#endif
+	if (InvEOTF == INVEOTF_PQ)
+	{
+		// Apply conversion to ST-2084 (Dolby PQ)
+		Color = LinearToST2084(Color * LinearToNitsScale);
+	}
+
+	return Color;
 }
 
-float3 ComputeInvOutputACES(float3 ODTColor, float3x3 LimitingRgbToXYZ, float3x3 LimitingXYZToRgb)
+Texture3D InnerLUT;
+SamplerState InnerLUTSampler;
+
+float3 InnerLookupTable(float3 LUTEncodedColor, int InvEOTF)
 {
-#if USE_ACES_2
-	float3 GradedColorMultiplied = ACESInvOutputTransform(
-		ODTColor,
-		(float3x3)WorkingColorSpace.FromAP0,
-		LinearToNitsScale,
-		ACESReachTableSDR,
-		ACESGamutTableSDR,
-		ACESGammaTableSDR,
-		bScaleWhiteSDR,
-		LimitingRgbToXYZ,
-		LimitingXYZToRgb,
-		GetColorSpace(OutputGamut));
-	return GradedColorMultiplied / ACESSceneColorMultiplier;
-#else
-	// Magenta makes it obvious this isn't implemented.
-	return float3(1., 0., 1.);
-#endif
-}
+	float InvInnerLUTSize = 1. / InnerLUTSize;
+	float InnerLUTScale = 1. - InvInnerLUTSize;
+	float InnerLUTOffset = .5 * InvInnerLUTSize;
+	float3 UVW = mad(LUTEncodedColor, InnerLUTScale, InnerLUTOffset);
 
-/** Values of output device (matching C++'s ETonemappingMethod) */
-#define TONEMAPPING_METHOD_Filmic 0
-#define TONEMAPPING_METHOD_ACES 1
+	float3 Color = Texture3DSample(InnerLUT, InnerLUTSampler, UVW).rgb;
+
+	if (OutputMaxLuminance > LinearToNitsScale)
+	{
+		if (InvEOTF != INVEOTF_PQ)
+		{
+			Color = ST2084ToLinear(Color) * LinearToNitsScaleInverse;
+			Color = ApplyInvEOTF(Color, InvEOTF);
+		}
+	}
+	else
+	{
+		Color = sRGBToLinear(Color);
+		Color = ApplyInvEOTF(Color, InvEOTF);
+	}
+
+	return Color;
+}
 
 #define TONEMAPPING_SDR_INVEOTF_SRGB 0
 #define TONEMAPPING_SDR_INVEOTF_BT1886 1
 uint TonemappingSDRInvEOTF;
 
-float3 ApplyToneCurve(float3 ColorAP1, float PeakLuminance, float3x3 AP1_2_Output)
+float3 ApplyToneCurve(float3 ColorAP1, float PeakLuminance, float3x3 AP1_2_Output, int InvEOTF)
 {
 #if TONEMAPPING_METHOD == TONEMAPPING_METHOD_ACES
 	float3 GradedColor = ComputeGradedColor(ColorAP1);
-	float3 OutputGamutColor = ComputeOutputACES(GradedColor, PeakLuminance, GetLimitingRgbToXYZ(), GetLimitingXYZToRgb(), AP1_2_Output) * LinearToNitsScaleInverse;
+	float3 GradedColorAP1 = mul( (float3x3)WorkingColorSpace.ToAP1, GradedColor );
+	float3 LUTEncodedColor = LinearToST2084(clamp(GradedColorAP1 * LinearToNitsScale, 0.0, 10000.0));
+	float3 OutputGamutColor = InnerLookupTable(LUTEncodedColor, InvEOTF);
 #elif TONEMAPPING_METHOD == TONEMAPPING_METHOD_Filmic
 	float3 FilmColorNoGamma = ComputeFilmColorNoGamma(ColorAP1);
 
@@ -356,10 +323,15 @@ float3 ApplyToneCurve(float3 ColorAP1, float PeakLuminance, float3x3 AP1_2_Outpu
 			OutputGamutColor = pow(OutputGamutColor, Bt1886Gamma);
 		}
 
-		float3 GradedColor = ComputeInvOutputACES(OutputGamutColor * LinearToNitsScale, (float3x3)WorkingColorSpace.ToXYZ, (float3x3)WorkingColorSpace.FromXYZ);
-		OutputGamutColor = ComputeOutputACES(GradedColor, PeakLuminance, GetLimitingRgbToXYZ(), GetLimitingXYZToRgb(), AP1_2_Output) * LinearToNitsScaleInverse;
+		// Use PQ shaper to preserve shadow detail and normalize for SDR by using the reciprocal of ST2084 at 100 nits.
+		float3 LUTEncodedColor = LinearToST2084(clamp(OutputGamutColor * LinearToNitsScale, 0.0, 10000.0)) * 1.9682001;
+		OutputGamutColor = InnerLookupTable(LUTEncodedColor, InvEOTF);
 	}
+	else
 #endif
+	{
+		OutputGamutColor = ApplyInvEOTF(OutputGamutColor, InvEOTF);
+	}
 #else
 #error Unknown TONEMAPPING_METHOD
 #endif
@@ -435,7 +407,7 @@ float4 CombineLUTsCommon(float3 LUTEncodedColor)
 	// sRGB, user specified gamut
 	BRANCH if( GetOutputDevice() == TONEMAPPER_OUTPUT_sRGB)
 	{
-		float3 OutputGamutColor = ApplyToneCurve(ColorAP1, LinearToNitsScale, AP1_2_Output);
+		float3 OutputGamutColor = ApplyToneCurve(ColorAP1, LinearToNitsScale, AP1_2_Output, INVEOTF_LINEAR);
 
 		BRANCH if (TonemappingSDRInvEOTF == TONEMAPPING_SDR_INVEOTF_SRGB)
 		{
@@ -450,7 +422,7 @@ float4 CombineLUTsCommon(float3 LUTEncodedColor)
 	// Rec 709, user specified gamut
 	else BRANCH if( GetOutputDevice() == TONEMAPPER_OUTPUT_Rec709)
 	{
-		float3 OutputGamutColor = ApplyToneCurve(ColorAP1, LinearToNitsScale, AP1_2_Output);
+		float3 OutputGamutColor = ApplyToneCurve(ColorAP1, LinearToNitsScale, AP1_2_Output, INVEOTF_LINEAR);
 
 		// Didn't profile yet if the branching version would be faster (different linear segment).
 		OutDeviceColor = LinearTo709Branchless(OutputGamutColor);
@@ -459,18 +431,13 @@ float4 CombineLUTsCommon(float3 LUTEncodedColor)
 	// Transform with PQ/2084 encoding, user specified gamut
 	else BRANCH if( GetOutputDevice() == TONEMAPPER_OUTPUT_ACES1000nitST2084 || GetOutputDevice() == TONEMAPPER_OUTPUT_ACES2000nitST2084)
 	{
-		float3 OutputGamutColor = ApplyToneCurve(ColorAP1, OutputMaxLuminance * InvHDRPaperWhiteRatio, AP1_2_Output);
-		OutputGamutColor *= HDRPaperWhiteRatio;
-
-		// Apply conversion to ST-2084 (Dolby PQ)
-		OutDeviceColor = LinearToST2084(OutputGamutColor * LinearToNitsScale);
+		OutDeviceColor = ApplyToneCurve(ColorAP1, OutputMaxLuminance, AP1_2_Output, INVEOTF_PQ);
 	}
 
 	// Transform to linear ScRGB
 	else BRANCH if( GetOutputDevice() == TONEMAPPER_OUTPUT_ACES1000nitScRGB || GetOutputDevice() == TONEMAPPER_OUTPUT_ACES2000nitScRGB)
 	{
-		float3 OutputGamutColor = ApplyToneCurve(ColorAP1, OutputMaxLuminance * InvHDRPaperWhiteRatio, AP1_2_Output);
-		OutputGamutColor *= HDRPaperWhiteRatio;
+		float3 OutputGamutColor = ApplyToneCurve(ColorAP1, OutputMaxLuminance, AP1_2_Output, INVEOTF_LINEAR);
 		OutDeviceColor = (LinearToNitsScale / UE_SCRGB_WHITE_NITS) * OutputGamutColor;
 	}
 
@@ -489,14 +456,14 @@ float4 CombineLUTsCommon(float3 LUTEncodedColor)
 	// "Linear" including all color correction and the tone curve, but no device gamma
 	else BRANCH if (GetOutputDevice() == TONEMAPPER_OUTPUT_WithToneCurve)
 	{
-		OutDeviceColor = ApplyToneCurve(ColorAP1, LinearToNitsScale, AP1_2_Output);
+		OutDeviceColor = ApplyToneCurve(ColorAP1, LinearToNitsScale, AP1_2_Output, INVEOTF_LINEAR);
 	}
 
 	// OutputDevice == TONEMAPPER_OUTPUT_ExplicitGammaMapping
 	// Gamma 2.2, user specified gamut
 	else
 	{
-		float3 OutputGamutColor = ApplyToneCurve(ColorAP1, LinearToNitsScale, AP1_2_Output);
+		float3 OutputGamutColor = ApplyToneCurve(ColorAP1, LinearToNitsScale, AP1_2_Output, INVEOTF_LINEAR);
 
 		// This is different than the prior "gamma" curve adjustment (but reusing the variable).
 		// For displays set to a gamma colorspace.
diff --git a/Engine/Shaders/Private/PostProcessCombineLUTsInner.usf b/Engine/Shaders/Private/PostProcessCombineLUTsInner.usf
new file mode 100644
--- /dev/null
+++ b/Engine/Shaders/Private/PostProcessCombineLUTsInner.usf
@@ -0,0 +1,159 @@
+// Copyright Epic Games, Inc. All Rights Reserved.
+
+#include "Common.ush"
+#include "PostProcessCommon.ush"
+#include "TonemapCommon.ush"
+
+// ---------------------------------------------------
+
+float3 LimitingRgbToXYZ0;
+float3 LimitingRgbToXYZ1;
+float3 LimitingRgbToXYZ2;
+float3x3 GetLimitingRgbToXYZ()
+{
+	return float3x3(LimitingRgbToXYZ0, LimitingRgbToXYZ1, LimitingRgbToXYZ2);
+}
+
+float3 LimitingXYZToRgb0;
+float3 LimitingXYZToRgb1;
+float3 LimitingXYZToRgb2;
+float3x3 GetLimitingXYZToRgb()
+{
+	return float3x3(LimitingXYZToRgb0, LimitingXYZToRgb1, LimitingXYZToRgb2);
+}
+
+float4 ACESMinMaxData;
+float4 ACESMidData;
+float4 ACESCoefsLow_0;
+float4 ACESCoefsHigh_0;
+float ACESCoefsLow_4;
+float ACESCoefsHigh_4;
+float ACESSceneColorMultiplier;
+float ACESGamutCompression;
+
+// ACES 2.0
+Texture2D<float> ACESReachTable;
+Texture2D<float> ACESGamutTable;
+Texture2D<float> ACESGammaTable;
+Texture2D<float> ACESReachTableSDR;
+Texture2D<float> ACESGamutTableSDR;
+Texture2D<float> ACESGammaTableSDR;
+uint bScaleWhite;
+uint bScaleWhiteSDR;
+
+uint OutputDevice;
+uint OutputGamut;
+float OutputMaxLuminance;
+
+float HDRPaperWhiteRatio;
+float InvHDRPaperWhiteRatio;
+
+uint GetOutputDevice()
+{
+#if OUTPUT_DEVICE_SRGB
+	return TONEMAPPER_OUTPUT_sRGB;
+#else
+	return OutputDevice;
+#endif
+}
+
+float3 ComputeOutputACES(
+	float3 GradedColor,
+	float3x3 LimitingRgbToXYZ,
+	float3x3 LimitingXYZToRgb,
+	float3x3 AP1_2_Output,
+	float bSDR)
+{
+#if USE_ACES_2
+	float3 PaperMaxLuminance = lerp(OutputMaxLuminance * InvHDRPaperWhiteRatio, OutputMaxLuminance, bSDR);
+	float3 ODTColor = ACESOutputTransform(
+		GradedColor * ACESSceneColorMultiplier,
+		(float3x3)WorkingColorSpace.ToAP0,
+		PaperMaxLuminance,
+		ACESReachTable,
+		ACESGamutTable,
+		ACESGammaTable,
+		bScaleWhite,
+		LimitingRgbToXYZ,
+		LimitingXYZToRgb,
+		GetColorSpace(OutputGamut));
+	ODTColor = lerp(ODTColor * HDRPaperWhiteRatio, ODTColor, bSDR);
+	return ODTColor;
+#else
+	FACESTonemapParams AcesParams = ComputeACESTonemapParams(ACESMinMaxData, ACESMidData, ACESCoefsLow_0, ACESCoefsHigh_0, ACESCoefsLow_4, ACESCoefsHigh_4, ACESSceneColorMultiplier, ACESGamutCompression);
+	float3 ODTColor = ACESOutputTransform( GradedColor, (float3x3)WorkingColorSpace.ToAP0, AcesParams);
+	float3 OutDeviceColor = mul(AP1_2_Output, ODTColor);
+	return OutDeviceColor;
+#endif
+}
+
+float3 ComputeInvOutputACES(float3 ODTColor, float3x3 LimitingRgbToXYZ, float3x3 LimitingXYZToRgb)
+{
+#if USE_ACES_2
+	float3 GradedColorMultiplied = ACESInvOutputTransform(
+		ODTColor,
+		(float3x3)WorkingColorSpace.FromAP0,
+		LinearToNitsScale,
+		ACESReachTableSDR,
+		ACESGamutTableSDR,
+		ACESGammaTableSDR,
+		bScaleWhiteSDR,
+		LimitingRgbToXYZ,
+		LimitingXYZToRgb,
+		GetColorSpace(OutputGamut));
+	return GradedColorMultiplied / ACESSceneColorMultiplier;
+#else
+	// Magenta makes it obvious this isn't implemented.
+	return float3(1., 0., 1.);
+#endif
+}
+
+float4 MethodLUTsCommon(float3 LUTEncodedColor)
+{
+	float3 GradedColor;
+	BRANCH
+	if (TonemappingMethod == TONEMAPPING_METHOD_ACES)
+	{
+		float3 GradedColorAP1 = ST2084ToLinear(LUTEncodedColor) * LinearToNitsScaleInverse;
+		GradedColor = mul( (float3x3)WorkingColorSpace.FromAP1, GradedColorAP1 );
+	}
+	else // TonemappingMethod == TONEMAPPING_METHOD_Filmic
+	{
+		// "Automatic HDR" using ACES to remap filmic.
+		// Undo multiplication performed by outer LUT.
+		float3 FilmNits = ST2084ToLinear(LUTEncodedColor * .5080784215);
+		GradedColor = ComputeInvOutputACES(FilmNits, (float3x3)WorkingColorSpace.ToXYZ, (float3x3)WorkingColorSpace.FromXYZ);
+	}
+
+	const float3x3 AP1_2_Output = OutputGamutMappingMatrix(OutputGamut);
+	float bSDR = step(OutputMaxLuminance, LinearToNitsScale);
+	float3 OutputGamutColor = ComputeOutputACES(GradedColor, GetLimitingRgbToXYZ(), GetLimitingXYZToRgb(), AP1_2_Output, bSDR);
+
+	BRANCH
+	if (bSDR)
+	{
+		OutputGamutColor = saturate(OutputGamutColor * LinearToNitsScaleInverse);
+		OutputGamutColor = LinearToSrgb(OutputGamutColor);
+	}
+	else
+	{
+		OutputGamutColor = clamp(OutputGamutColor, 0., 10000.);
+		OutputGamutColor = LinearToST2084(OutputGamutColor);
+	}
+
+	return float4(OutputGamutColor, 0.);
+}
+
+float InvLUTSizeMinusOne;
+
+RWTexture3D<float4> RWOutputTexture;
+
+[numthreads(THREADGROUP_SIZE, THREADGROUP_SIZE, THREADGROUP_SIZE)]
+void MainCS(uint3 DispatchThreadId : SV_DispatchThreadID) 
+{
+	float3 LUTEncodedColor = (float3)DispatchThreadId * InvLUTSizeMinusOne;
+	float4 OutColor = MethodLUTsCommon(LUTEncodedColor);
+
+	uint3 PixelPos = DispatchThreadId;
+	RWOutputTexture[PixelPos] = OutColor;
+}
diff --git a/Engine/Shaders/Private/TonemapCommon.ush b/Engine/Shaders/Private/TonemapCommon.ush
--- a/Engine/Shaders/Private/TonemapCommon.ush
+++ b/Engine/Shaders/Private/TonemapCommon.ush
@@ -19,6 +19,10 @@
 	#include "ACES/ACES_v1.3.ush"
 #endif
 
+/** Values of output device (matching C++'s ETonemappingMethod) */
+#define TONEMAPPING_METHOD_Filmic 0
+#define TONEMAPPING_METHOD_ACES 1
+
 /** Values of output device (matching C++'s EDisplayOutputFormat) */
 #define TONEMAPPER_OUTPUT_sRGB 0
 #define TONEMAPPER_OUTPUT_Rec709 1
@@ -64,6 +68,7 @@ half3 TonemapAndGammaCorrect(half3 LinearColor)
 	return GammaColor;
 }
 
+uint TonemappingMethod;
 
 /*
 ============================================
diff --git a/Engine/Source/Runtime/Renderer/Private/PostProcess/PostProcessCombineLUTs.cpp b/Engine/Source/Runtime/Renderer/Private/PostProcess/PostProcessCombineLUTs.cpp
--- a/Engine/Source/Runtime/Renderer/Private/PostProcess/PostProcessCombineLUTs.cpp
+++ b/Engine/Source/Runtime/Renderer/Private/PostProcess/PostProcessCombineLUTs.cpp
@@ -42,6 +42,12 @@
 	TEXT("Size of film LUT"),
 	ECVF_RenderThreadSafe | ECVF_Scalability);
 
+TAutoConsoleVariable<int32> CVarLUTInnerSize(
+	TEXT("r.LUT.Inner.Size"),
+	64,
+	TEXT("Size of inner LUT"),
+	ECVF_RenderThreadSafe | ECVF_Scalability);
+
 TAutoConsoleVariable<bool> CVarLUTAsyncCompute(
 	TEXT("r.LUT.AsyncCompute"),
 	1,
@@ -58,6 +64,11 @@
 	TEXT("Controls whether the tonemapping LUT pass is executed every frame."),
 	ECVF_RenderThreadSafe);
 
+TAutoConsoleVariable<int32> CVarLUTInnerUpdateEveryFrame(
+	TEXT("r.LUT.Inner.UpdateEveryFrame"), 0,
+	TEXT("Controls whether the inner tonemapping LUT pass is executed every frame."),
+	ECVF_RenderThreadSafe);
+
 TAutoConsoleVariable<int32> CVarLUTTonemappingMethod(
 	TEXT("r.LUT.TonemappingMethod"), -1,
 	TEXT("Set the tonemapping method to transform colors for the output device.\n")
@@ -122,7 +133,7 @@
 	return Parameters;
 }
 
-FPooledRenderTargetDesc FSceneViewState::CreateLUTRenderTarget(const int32 LUTSize, const bool bNeedFloatOutput)
+FPooledRenderTargetDesc FSceneViewState::CreateLUTRenderTarget(const int32 LUTSize, const bool bNeedFloatOutput, const TCHAR* const DebugName)
 {
 	// Create the texture needed for the tonemapping LUT in one place
 	EPixelFormat LUTPixelFormat = PF_A2B10G10R10;
@@ -148,7 +159,7 @@
 		GFastVRamConfig.CombineLUTs,
 		TexCreate_ShaderResource | TexCreate_UAV,
 		false);
-	Desc.DebugName = TEXT("CombineLUTs");
+		Desc.DebugName = DebugName;
 
 	return Desc;
 }
@@ -160,7 +171,7 @@
 		(CombinedLUTRenderTarget->GetDesc().Format == PF_FloatRGBA) != bNeedFloatOutput)
 	{
 		// Create the texture needed for the tonemapping LUT
-		FPooledRenderTargetDesc Desc = CreateLUTRenderTarget(LUTSize, bNeedFloatOutput);
+		FPooledRenderTargetDesc Desc = CreateLUTRenderTarget(LUTSize, bNeedFloatOutput, TEXT("CombineLUTs"));
 
 		GRenderTargetPool.FindFreeElement(RHICmdList, Desc, CombinedLUTRenderTarget, Desc.DebugName);
 	}
@@ -168,6 +179,21 @@
 	return CombinedLUTRenderTarget.GetReference();
 }
 
+IPooledRenderTarget* FSceneViewState::GetInnerTonemappingLUT(FRHICommandList& RHICmdList, const int32 LUTSize, const bool bNeedFloatOutput)
+{
+	if (InnerLUTRenderTarget.IsValid() == false ||
+		InnerLUTRenderTarget->GetDesc().Extent.Y != LUTSize ||
+		(InnerLUTRenderTarget->GetDesc().Format == PF_FloatRGBA) != bNeedFloatOutput)
+	{
+		// Create the texture needed for the tonemapping LUT
+		FPooledRenderTargetDesc Desc = CreateLUTRenderTarget(LUTSize, bNeedFloatOutput, TEXT("CombineLUTs_Inner"));
+
+		GRenderTargetPool.FindFreeElement(RHICmdList, Desc, InnerLUTRenderTarget, Desc.DebugName);
+	}
+
+	return InnerLUTRenderTarget.GetReference();
+}
+
 BEGIN_SHADER_PARAMETER_STRUCT(FACESTonemapShaderParameters, )
 	SHADER_PARAMETER(FVector4f, ACESMinMaxData) // xy = min ACES/luminance, zw = max ACES/luminance
 	SHADER_PARAMETER(FVector4f, ACESMidData) // x = mid ACES, y = mid luminance, z = mid slope
@@ -187,13 +213,11 @@
 	SHADER_PARAMETER_SRV(Texture2D<float>, ACESGammaTableSDR)
 END_SHADER_PARAMETER_STRUCT()
 
-
 BEGIN_SHADER_PARAMETER_STRUCT(FCombineLUTParameters, )
 	SHADER_PARAMETER_TEXTURE_ARRAY(Texture2D, Textures, [GMaxLUTBlendCount])
 	SHADER_PARAMETER_SAMPLER_ARRAY(SamplerState, Samplers, [GMaxLUTBlendCount])
 	SHADER_PARAMETER_SCALAR_ARRAY(float, LUTWeights, [GMaxLUTBlendCount])
 	SHADER_PARAMETER_STRUCT_REF(FWorkingColorSpaceShaderParameters, WorkingColorSpace)
-	SHADER_PARAMETER_STRUCT_INCLUDE(FACESTonemapShaderParameters, ACESTonemapParameters)
 	SHADER_PARAMETER(FVector3f, LimitingRgbToXYZ0)
 	SHADER_PARAMETER(FVector3f, LimitingRgbToXYZ1)
 	SHADER_PARAMETER(FVector3f, LimitingRgbToXYZ2)
@@ -223,6 +247,7 @@
 	SHADER_PARAMETER(FVector4f, ColorGainHighlights)
 	SHADER_PARAMETER(FVector4f, ColorOffsetHighlights)
 	SHADER_PARAMETER(float, LUTSize)
+	SHADER_PARAMETER(float, InnerLUTSize)
 	SHADER_PARAMETER(float, WhiteTemp)
 	SHADER_PARAMETER(float, WhiteTint)
 	SHADER_PARAMETER(float, ColorCorrectionShadowsMax)
@@ -237,13 +262,26 @@
 	SHADER_PARAMETER(float, FilmShoulder)
 	SHADER_PARAMETER(float, FilmBlackClip)
 	SHADER_PARAMETER(float, FilmWhiteClip)
-	SHADER_PARAMETER(float, HDRPaperWhiteRatio)
-	SHADER_PARAMETER(float, InvHDRPaperWhiteRatio)
 	SHADER_PARAMETER(uint32, bIsTemperatureWhiteBalance)
 	SHADER_PARAMETER_STRUCT_INCLUDE(FColorRemapParameters, ColorRemap)
 	SHADER_PARAMETER_STRUCT_INCLUDE(FTonemapperOutputDeviceParameters, OutputDevice)
 END_SHADER_PARAMETER_STRUCT()
 
+BEGIN_SHADER_PARAMETER_STRUCT(FInnerLUTParameters, )
+	SHADER_PARAMETER_STRUCT_REF(FWorkingColorSpaceShaderParameters, WorkingColorSpace)
+	SHADER_PARAMETER_STRUCT_INCLUDE(FACESTonemapShaderParameters, ACESTonemapParameters)
+	SHADER_PARAMETER(FVector3f, LimitingRgbToXYZ0)
+	SHADER_PARAMETER(FVector3f, LimitingRgbToXYZ1)
+	SHADER_PARAMETER(FVector3f, LimitingRgbToXYZ2)
+	SHADER_PARAMETER(FVector3f, LimitingXYZToRgb0)
+	SHADER_PARAMETER(FVector3f, LimitingXYZToRgb1)
+	SHADER_PARAMETER(FVector3f, LimitingXYZToRgb2)
+	SHADER_PARAMETER(float, HDRPaperWhiteRatio)
+	SHADER_PARAMETER(float, InvHDRPaperWhiteRatio)
+	SHADER_PARAMETER(uint32, TonemappingMethod)
+	SHADER_PARAMETER_STRUCT_INCLUDE(FTonemapperOutputDeviceParameters, OutputDevice)
+END_SHADER_PARAMETER_STRUCT()
+
 #define UPDATE_CACHE_SETTINGS(DestParameters, ParamValue, bOutHasChanged) \
 if(DestParameters != (ParamValue)) \
 { \
@@ -351,47 +389,6 @@
 		UPDATE_CACHE_SETTINGS(Parameters.OutputDevice.OutputGamut, TonemapperOutputDeviceParameters.OutputGamut, bHasChanged);
 		UPDATE_CACHE_SETTINGS(Parameters.OutputDevice.OutputMaxLuminance, TonemapperOutputDeviceParameters.OutputMaxLuminance, bHasChanged);
 
-		const UE::Color::FColorSpace& LimitingColorSpace = ViewFamily.RenderTarget->GetLimitingColorSpace();
-		const FMatrix44d& LimitingRgbToXYZ = LimitingColorSpace.GetRgbToXYZ();
-		UPDATE_CACHE_SETTINGS(Parameters.LimitingRgbToXYZ0, FVector3f(LimitingRgbToXYZ.GetColumn(0)), bHasChanged);
-		UPDATE_CACHE_SETTINGS(Parameters.LimitingRgbToXYZ1, FVector3f(LimitingRgbToXYZ.GetColumn(1)), bHasChanged);
-		UPDATE_CACHE_SETTINGS(Parameters.LimitingRgbToXYZ2, FVector3f(LimitingRgbToXYZ.GetColumn(2)), bHasChanged);
-		const FMatrix44d& LimitingXYZToRgb = LimitingColorSpace.GetXYZToRgb();
-		UPDATE_CACHE_SETTINGS(Parameters.LimitingXYZToRgb0, FVector3f(LimitingXYZToRgb.GetColumn(0)), bHasChanged);
-		UPDATE_CACHE_SETTINGS(Parameters.LimitingXYZToRgb1, FVector3f(LimitingXYZToRgb.GetColumn(1)), bHasChanged);
-		UPDATE_CACHE_SETTINGS(Parameters.LimitingXYZToRgb2, FVector3f(LimitingXYZToRgb.GetColumn(2)), bHasChanged);
-
-		const float HDRPaperWhiteRatio = ViewFamily.RenderTarget->GetHDRPaperWhiteInNits() / 203.f;
-		UPDATE_CACHE_SETTINGS(Parameters.HDRPaperWhiteRatio, HDRPaperWhiteRatio, bHasChanged);
-		UPDATE_CACHE_SETTINGS(Parameters.InvHDRPaperWhiteRatio, 1.f / HDRPaperWhiteRatio, bHasChanged);
-
-		FACESTonemapParams TonemapperParams;
-		GetACESTonemapParameters(TonemapperParams);
-		UPDATE_CACHE_SETTINGS(Parameters.ACESTonemapParameters.ACESMinMaxData, TonemapperParams.ACESMinMaxData, bHasChanged);
-		UPDATE_CACHE_SETTINGS(Parameters.ACESTonemapParameters.ACESMidData, TonemapperParams.ACESMidData, bHasChanged);
-		UPDATE_CACHE_SETTINGS(Parameters.ACESTonemapParameters.ACESCoefsLow_0, TonemapperParams.ACESCoefsLow_0, bHasChanged);
-		UPDATE_CACHE_SETTINGS(Parameters.ACESTonemapParameters.ACESCoefsHigh_0, TonemapperParams.ACESCoefsHigh_0, bHasChanged);
-		UPDATE_CACHE_SETTINGS(Parameters.ACESTonemapParameters.ACESCoefsLow_4, TonemapperParams.ACESCoefsLow_4, bHasChanged);
-		UPDATE_CACHE_SETTINGS(Parameters.ACESTonemapParameters.ACESCoefsHigh_4, TonemapperParams.ACESCoefsHigh_4, bHasChanged);
-		UPDATE_CACHE_SETTINGS(Parameters.ACESTonemapParameters.ACESSceneColorMultiplier, TonemapperParams.ACESSceneColorMultiplier, bHasChanged);
-		UPDATE_CACHE_SETTINGS(Parameters.ACESTonemapParameters.ACESGamutCompression, TonemapperParams.ACESGamutCompression, bHasChanged);
-		UPDATE_CACHE_SETTINGS(Parameters.ACESTonemapParameters.bScaleWhite, !DoesGamutWhitePointMatch(TonemapperOutputDeviceParameters.OutputGamut, LimitingColorSpace), bHasChanged);
-		UPDATE_CACHE_SETTINGS(Parameters.ACESTonemapParameters.bScaleWhiteSDR, !DoesGamutWhitePointMatch(TonemapperOutputDeviceParameters.OutputGamut, UE::Color::FColorSpace::GetWorking()), bHasChanged);
-		UE::Color::ACES::GetTransformResources(
-			GraphBuilder,
-			TonemapperOutputDeviceParameters.OutputMaxLuminance,
-			LimitingColorSpace,
-			Parameters.ACESTonemapParameters.ACESReachTable,
-			Parameters.ACESTonemapParameters.ACESGamutTable,
-			Parameters.ACESTonemapParameters.ACESGammaTable
-		);
-		UE::Color::ACES::GetTransformResourcesSDR(
-			GraphBuilder,
-			Parameters.ACESTonemapParameters.ACESReachTableSDR,
-			Parameters.ACESTonemapParameters.ACESGamutTableSDR,
-			Parameters.ACESTonemapParameters.ACESGammaTableSDR
-		);
-
 		UPDATE_CACHE_SETTINGS(Parameters.ColorScale, FVector3f(View.ColorScale), bHasChanged);
 		UPDATE_CACHE_SETTINGS(Parameters.OverlayColor, FVector4f(View.OverlayColor), bHasChanged);
 		UPDATE_CACHE_SETTINGS(Parameters.ColorRemap.MappingPolynomial, GetMappingPolynomial(), bHasChanged);
@@ -399,6 +396,7 @@
 		// White balance
 		UPDATE_CACHE_SETTINGS(Parameters.bIsTemperatureWhiteBalance, uint32(Settings.TemperatureType == ETemperatureMethod::TEMP_WhiteBalance), bHasChanged);
 		UPDATE_CACHE_SETTINGS(Parameters.LUTSize, (float)CVarLUTSize->GetInt(), bHasChanged);
+		UPDATE_CACHE_SETTINGS(Parameters.InnerLUTSize, (float)CVarLUTInnerSize->GetInt(), bHasChanged);
 		UPDATE_CACHE_SETTINGS(Parameters.WhiteTemp, Settings.WhiteTemp, bHasChanged);
 		UPDATE_CACHE_SETTINGS(Parameters.WhiteTint, Settings.WhiteTint, bHasChanged);
 
@@ -456,9 +454,8 @@
 
 	class FBlendCount : SHADER_PERMUTATION_RANGE_INT("BLENDCOUNT", 1, 5);
 	class FOutputDeviceSRGB : SHADER_PERMUTATION_BOOL("OUTPUT_DEVICE_SRGB");
-	class FUseACES2 : SHADER_PERMUTATION_BOOL("USE_ACES_2");
 	class FTonemappingMethod : SHADER_PERMUTATION_ENUM_CLASS("TONEMAPPING_METHOD", ETonemappingMethod);
-	using FPermutationDomain = TShaderPermutationDomain<FBlendCount, FOutputDeviceSRGB, FUseACES2, FTonemappingMethod>;
+	using FPermutationDomain = TShaderPermutationDomain<FBlendCount, FOutputDeviceSRGB, FTonemappingMethod>;
 
 	DECLARE_GLOBAL_SHADER(FLUTBlenderCS);
 	SHADER_USE_PARAMETER_STRUCT(FLUTBlenderCS, FGlobalShader);
@@ -467,6 +464,8 @@
 		SHADER_PARAMETER_STRUCT_INCLUDE(FCombineLUTParameters, CombineLUT)
 		SHADER_PARAMETER(float, InvLUTSizeMinusOne)
 		SHADER_PARAMETER_RDG_TEXTURE_UAV(RWTexture3D<float4>, RWOutputTexture)
+		SHADER_PARAMETER_RDG_TEXTURE(Texture3D, InnerLUT)
+		SHADER_PARAMETER_SAMPLER(SamplerState, InnerLUTSampler)
 	END_SHADER_PARAMETER_STRUCT()
 
 	static EShaderPermutationPrecacheRequest ShouldPrecachePermutation(const FShaderPermutationParameters& Parameters)
@@ -484,6 +483,137 @@
 
 IMPLEMENT_GLOBAL_SHADER(FLUTBlenderCS, "/Engine/Private/PostProcessCombineLUTs.usf", "MainCS", SF_Compute);
 
+class FLUTInnerCS : public FGlobalShader
+{
+public:
+	static const int32 GroupSize = 4;
+
+	class FOutputDeviceSRGB : SHADER_PERMUTATION_BOOL("OUTPUT_DEVICE_SRGB");
+	class FUseACES2 : SHADER_PERMUTATION_BOOL("USE_ACES_2");
+	using FPermutationDomain = TShaderPermutationDomain<FOutputDeviceSRGB, FUseACES2>;
+
+	DECLARE_GLOBAL_SHADER(FLUTInnerCS);
+	SHADER_USE_PARAMETER_STRUCT(FLUTInnerCS, FGlobalShader);
+
+	BEGIN_SHADER_PARAMETER_STRUCT(FParameters, )
+		SHADER_PARAMETER_STRUCT_INCLUDE(FInnerLUTParameters, InnerLUT)
+		SHADER_PARAMETER(float, InvLUTSizeMinusOne)
+		SHADER_PARAMETER_RDG_TEXTURE_UAV(RWTexture3D<float4>, RWOutputTexture)
+	END_SHADER_PARAMETER_STRUCT()
+
+	static void ModifyCompilationEnvironment(const FGlobalShaderPermutationParameters& Parameters, FShaderCompilerEnvironment& OutEnvironment)
+	{
+		FGlobalShader::ModifyCompilationEnvironment(Parameters, OutEnvironment);
+
+		OutEnvironment.SetDefine(TEXT("THREADGROUP_SIZE"), GroupSize);
+	}
+};
+
+IMPLEMENT_GLOBAL_SHADER(FLUTInnerCS, "/Engine/Private/PostProcessCombineLUTsInner.usf", "MainCS", SF_Compute);
+
+struct FCachedInnerLUTSettings
+{
+	uint32 UniqueID = 0;
+	EShaderPlatform ShaderPlatform = GMaxRHIShaderPlatform;
+	FInnerLUTParameters Parameters;
+	FWorkingColorSpaceShaderParameters WorkingColorSpaceShaderParameters;
+	bool bUseCompute = false;
+	int32 AcesVersion = 0;
+	int32 InnerLUTSize = 0;
+
+	bool UpdateCachedValues(FRDGBuilder& GraphBuilder, const FViewInfo& View, int32 InInnerLUTSize, const FTexture* const* Textures, const float* Weights, uint32 BlendCount)
+	{
+		bool bHasChanged = false;
+		GetInnerLUTParameters(GraphBuilder, View, bHasChanged);
+		UPDATE_CACHE_SETTINGS(UniqueID, View.State ? View.State->GetViewKey() : 0, bHasChanged);
+		UPDATE_CACHE_SETTINGS(ShaderPlatform, View.GetShaderPlatform(), bHasChanged);
+		UPDATE_CACHE_SETTINGS(bUseCompute, View.bUseComputePasses, bHasChanged);
+		static const auto CVarAcesVersion = IConsoleManager::Get().FindTConsoleVariableDataInt(TEXT("r.HDR.Aces.Version"));
+		UPDATE_CACHE_SETTINGS(AcesVersion, CVarAcesVersion->GetValueOnRenderThread(), bHasChanged);
+		UPDATE_CACHE_SETTINGS(InnerLUTSize, InInnerLUTSize, bHasChanged);
+
+		const FWorkingColorSpaceShaderParameters* InWorkingColorSpaceShaderParameters = reinterpret_cast<const FWorkingColorSpaceShaderParameters*>(GDefaultWorkingColorSpaceUniformBuffer.GetContents());
+		if (InWorkingColorSpaceShaderParameters)
+		{
+			UPDATE_CACHE_SETTINGS(WorkingColorSpaceShaderParameters.ToXYZ, InWorkingColorSpaceShaderParameters->ToXYZ, bHasChanged);
+			UPDATE_CACHE_SETTINGS(WorkingColorSpaceShaderParameters.FromXYZ, InWorkingColorSpaceShaderParameters->FromXYZ, bHasChanged);
+			UPDATE_CACHE_SETTINGS(WorkingColorSpaceShaderParameters.ToAP1, InWorkingColorSpaceShaderParameters->ToAP1, bHasChanged);
+			UPDATE_CACHE_SETTINGS(WorkingColorSpaceShaderParameters.FromAP1, InWorkingColorSpaceShaderParameters->FromAP1, bHasChanged);
+			UPDATE_CACHE_SETTINGS(WorkingColorSpaceShaderParameters.ToAP0, InWorkingColorSpaceShaderParameters->ToAP0, bHasChanged);
+			UPDATE_CACHE_SETTINGS(WorkingColorSpaceShaderParameters.bIsSRGB, InWorkingColorSpaceShaderParameters->bIsSRGB, bHasChanged);
+		}
+
+		return bHasChanged;
+	}
+	
+	void GetInnerLUTParameters(
+		FRDGBuilder& GraphBuilder,
+		const FViewInfo& View,
+		bool& bHasChanged)
+	{
+		const FSceneViewFamily& ViewFamily = *(View.Family);
+		const FPostProcessSettings& Settings = (
+			ViewFamily.EngineShowFlags.ColorGrading && CVarColorGrading.GetValueOnRenderThread() != 0)
+			? View.FinalPostProcessSettings
+			: FPostProcessSettings::GetDefault();
+
+		Parameters.WorkingColorSpace = GDefaultWorkingColorSpaceUniformBuffer.GetUniformBufferRef();
+
+		FTonemapperOutputDeviceParameters TonemapperOutputDeviceParameters = GetTonemapperOutputDeviceParameters(ViewFamily);
+		UPDATE_CACHE_SETTINGS(Parameters.OutputDevice.InverseGamma, TonemapperOutputDeviceParameters.InverseGamma, bHasChanged);
+		UPDATE_CACHE_SETTINGS(Parameters.OutputDevice.OutputDevice, TonemapperOutputDeviceParameters.OutputDevice, bHasChanged);
+		UPDATE_CACHE_SETTINGS(Parameters.OutputDevice.OutputGamut, TonemapperOutputDeviceParameters.OutputGamut, bHasChanged);
+		UPDATE_CACHE_SETTINGS(Parameters.OutputDevice.OutputMaxLuminance, TonemapperOutputDeviceParameters.OutputMaxLuminance, bHasChanged);
+
+		const UE::Color::FColorSpace& LimitingColorSpace = ViewFamily.RenderTarget->GetLimitingColorSpace();
+		const FMatrix44d& LimitingRgbToXYZ = LimitingColorSpace.GetRgbToXYZ();
+		UPDATE_CACHE_SETTINGS(Parameters.LimitingRgbToXYZ0, FVector3f(LimitingRgbToXYZ.GetColumn(0)), bHasChanged);
+		UPDATE_CACHE_SETTINGS(Parameters.LimitingRgbToXYZ1, FVector3f(LimitingRgbToXYZ.GetColumn(1)), bHasChanged);
+		UPDATE_CACHE_SETTINGS(Parameters.LimitingRgbToXYZ2, FVector3f(LimitingRgbToXYZ.GetColumn(2)), bHasChanged);
+		const FMatrix44d& LimitingXYZToRgb = LimitingColorSpace.GetXYZToRgb();
+		UPDATE_CACHE_SETTINGS(Parameters.LimitingXYZToRgb0, FVector3f(LimitingXYZToRgb.GetColumn(0)), bHasChanged);
+		UPDATE_CACHE_SETTINGS(Parameters.LimitingXYZToRgb1, FVector3f(LimitingXYZToRgb.GetColumn(1)), bHasChanged);
+		UPDATE_CACHE_SETTINGS(Parameters.LimitingXYZToRgb2, FVector3f(LimitingXYZToRgb.GetColumn(2)), bHasChanged);
+
+		FACESTonemapParams TonemapperParams;
+		GetACESTonemapParameters(TonemapperParams);
+		UPDATE_CACHE_SETTINGS(Parameters.ACESTonemapParameters.ACESMinMaxData, TonemapperParams.ACESMinMaxData, bHasChanged);
+		UPDATE_CACHE_SETTINGS(Parameters.ACESTonemapParameters.ACESMidData, TonemapperParams.ACESMidData, bHasChanged);
+		UPDATE_CACHE_SETTINGS(Parameters.ACESTonemapParameters.ACESCoefsLow_0, TonemapperParams.ACESCoefsLow_0, bHasChanged);
+		UPDATE_CACHE_SETTINGS(Parameters.ACESTonemapParameters.ACESCoefsHigh_0, TonemapperParams.ACESCoefsHigh_0, bHasChanged);
+		UPDATE_CACHE_SETTINGS(Parameters.ACESTonemapParameters.ACESCoefsLow_4, TonemapperParams.ACESCoefsLow_4, bHasChanged);
+		UPDATE_CACHE_SETTINGS(Parameters.ACESTonemapParameters.ACESCoefsHigh_4, TonemapperParams.ACESCoefsHigh_4, bHasChanged);
+		UPDATE_CACHE_SETTINGS(Parameters.ACESTonemapParameters.ACESSceneColorMultiplier, TonemapperParams.ACESSceneColorMultiplier, bHasChanged);
+		UPDATE_CACHE_SETTINGS(Parameters.ACESTonemapParameters.ACESGamutCompression, TonemapperParams.ACESGamutCompression, bHasChanged);
+		UPDATE_CACHE_SETTINGS(Parameters.ACESTonemapParameters.bScaleWhite, !DoesGamutWhitePointMatch(TonemapperOutputDeviceParameters.OutputGamut, LimitingColorSpace), bHasChanged);
+		UPDATE_CACHE_SETTINGS(Parameters.ACESTonemapParameters.bScaleWhiteSDR, !DoesGamutWhitePointMatch(TonemapperOutputDeviceParameters.OutputGamut, UE::Color::FColorSpace::GetWorking()), bHasChanged);
+		UE::Color::ACES::GetTransformResources(
+			GraphBuilder,
+			TonemapperOutputDeviceParameters.OutputMaxLuminance,
+			LimitingColorSpace,
+			Parameters.ACESTonemapParameters.ACESReachTable,
+			Parameters.ACESTonemapParameters.ACESGamutTable,
+			Parameters.ACESTonemapParameters.ACESGammaTable
+		);
+		UE::Color::ACES::GetTransformResourcesSDR(
+			GraphBuilder,
+			Parameters.ACESTonemapParameters.ACESReachTableSDR,
+			Parameters.ACESTonemapParameters.ACESGamutTableSDR,
+			Parameters.ACESTonemapParameters.ACESGammaTableSDR
+		);
+
+		const float HDRPaperWhiteRatio = ViewFamily.RenderTarget->GetHDRPaperWhiteInNits() / 203.f;
+		UPDATE_CACHE_SETTINGS(Parameters.HDRPaperWhiteRatio, HDRPaperWhiteRatio, bHasChanged);
+		UPDATE_CACHE_SETTINGS(Parameters.InvHDRPaperWhiteRatio, 1.f / HDRPaperWhiteRatio, bHasChanged);
+
+		const int32 TonemappingMethodOverride = CVarLUTTonemappingMethod.GetValueOnRenderThread();
+		const uint32 TonemappingMethod = (TonemappingMethodOverride >= 0) ?
+			static_cast<uint32>(TonemappingMethodOverride) :
+			static_cast<uint32>(Settings.TonemappingMethod);
+		UPDATE_CACHE_SETTINGS(Parameters.TonemappingMethod, TonemappingMethod, bHasChanged);
+	}
+};
+
 uint32 GenerateFinalTable(const FFinalPostProcessSettings& Settings, const FTexture* OutTextures[], float OutWeights[], uint32 MaxCount)
 {
 	// Find the n strongest contributors, drop small contributors.
@@ -586,6 +716,7 @@
 FRDGTextureRef AddCombineLUTPass(FRDGBuilder& GraphBuilder, const FViewInfo& View)
 {
 	static FCachedLUTSettings CachedLUTSettings;
+	static FCachedInnerLUTSettings CachedInnerLUTSettings;
 
 	const FSceneViewFamily& ViewFamily = *(View.Family);
 	const FTexture* LocalTextures[GMaxLUTBlendCount];
@@ -620,8 +751,10 @@
 		? View.FinalPostProcessSettings
 		: FPostProcessSettings::GetDefault();
 
+	const int32 InnerLUTSize = CVarLUTInnerSize->GetInt();
+	const bool bHasInnerChanged = CachedInnerLUTSettings.UpdateCachedValues(GraphBuilder, View, InnerLUTSize, LocalTextures, LocalWeights, LocalCount);
 	const bool bHasChanged = CachedLUTSettings.UpdateCachedValues(GraphBuilder, View, Settings, LocalTextures, LocalWeights, LocalCount);
-	if (!bHasChanged && OutputTexture && CVarLUTUpdateEveryFrame.GetValueOnRenderThread() == 0)
+	if (!bHasInnerChanged && !bHasChanged && OutputTexture && CVarLUTUpdateEveryFrame.GetValueOnRenderThread() == 0)
 	{
 		return OutputTexture;
 	}
@@ -630,24 +763,68 @@
 	if (!OutputTexture)
 	{
 		OutputTexture = GraphBuilder.CreateTexture(
-			Translate(FSceneViewState::CreateLUTRenderTarget(LUTSize, bUseFloatOutput)),
-			TEXT("CombineLUT"));
+			Translate(FSceneViewState::CreateLUTRenderTarget(LUTSize, bUseFloatOutput, TEXT("CombineLUTs"))),
+			TEXT("CombineLUTs"));
 	}
 
 	check(OutputTexture);
 
+	const ETonemappingMethod TonemappingMethod = static_cast<ETonemappingMethod>(CachedInnerLUTSettings.Parameters.TonemappingMethod);
+	const bool bNeedsInnerTonemappingLUT =
+		(TonemappingMethod == ETonemappingMethod::StandardACES) ||
+		((TonemappingMethod == ETonemappingMethod::Filmic) && CachedInnerLUTSettings.Parameters.OutputDevice.OutputMaxLuminance > 100.f);
+	const int32 InnerLUTTextureSize = bNeedsInnerTonemappingLUT ? InnerLUTSize : 1;
+	FRDGTextureRef InnerTexture = TryRegisterExternalTexture(GraphBuilder,
+		View.GetInnerTonemappingLUT(GraphBuilder.RHICmdList, InnerLUTTextureSize, bUseFloatOutput));
+
+	const bool bOutputDeviceSRGB = (CachedLUTSettings.Parameters.OutputDevice.OutputDevice == (uint32)EDisplayOutputFormat::SDR_sRGB);
+	if (bHasInnerChanged || !InnerTexture || CVarLUTInnerUpdateEveryFrame.GetValueOnRenderThread())
+	{
+		// View doesn't support a persistent LUT, so create a temporary one.
+		if (!InnerTexture)
+		{
+			InnerTexture = GraphBuilder.CreateTexture(
+				Translate(FSceneViewState::CreateLUTRenderTarget(InnerLUTTextureSize, bUseFloatOutput, TEXT("CombineLUTs_Inner"))),
+				TEXT("CombineLUTs_inner"));
+		}
+
+		check(InnerTexture);
+
+		if (bNeedsInnerTonemappingLUT)
+		{
+			FLUTInnerCS::FPermutationDomain PermutationVector;
+			PermutationVector.Set<FLUTInnerCS::FUseACES2>(CachedInnerLUTSettings.AcesVersion > 1);
+
+			FLUTInnerCS::FParameters* PassParameters = GraphBuilder.AllocParameters<FLUTInnerCS::FParameters>();
+			PassParameters->InnerLUT = CachedInnerLUTSettings.Parameters;
+			PassParameters->InvLUTSizeMinusOne = 1.f / (InnerLUTSize - 1);
+			PassParameters->RWOutputTexture = GraphBuilder.CreateUAV(InnerTexture);
+
+			PermutationVector.Set<FLUTInnerCS::FOutputDeviceSRGB>(bOutputDeviceSRGB);
+			TShaderMapRef<FLUTInnerCS> ComputeShader(View.ShaderMap, PermutationVector);
+
+			const uint32 GroupSize = FMath::DivideAndRoundUp(InnerLUTSize, FLUTInnerCS::GroupSize);
+			FComputeShaderUtils::AddPass(
+				GraphBuilder,
+				RDG_EVENT_NAME("CombineLUTs_inner %d (CS)", InnerLUTSize),
+				CVarLUTAsyncCompute.GetValueOnRenderThread() ? ERDGPassFlags::AsyncCompute : ERDGPassFlags::Compute,
+				ComputeShader,
+				PassParameters,
+				FIntVector(GroupSize, GroupSize, GroupSize));
+		}
+	}
+
 	FLUTBlenderCS::FPermutationDomain PermutationVector;
 	PermutationVector.Set<FLUTBlenderCS::FBlendCount>(LocalCount);
-	PermutationVector.Set<FLUTBlenderCS::FUseACES2>(CachedLUTSettings.AcesVersion > 1);
 	PermutationVector.Set<FLUTBlenderCS::FTonemappingMethod>(CachedLUTSettings.TonemappingMethod);
 
 	FLUTBlenderCS::FParameters* PassParameters = GraphBuilder.AllocParameters<FLUTBlenderCS::FParameters>();
 	PassParameters->CombineLUT = CachedLUTSettings.Parameters;
+	PassParameters->InnerLUT = InnerTexture;
+	PassParameters->InnerLUTSampler = TStaticSamplerState<SF_Bilinear, AM_Clamp, AM_Clamp, AM_Clamp>::GetRHI();
 	PassParameters->InvLUTSizeMinusOne = 1.f / (LUTSize - 1);
 	PassParameters->RWOutputTexture = GraphBuilder.CreateUAV(OutputTexture);
 
-	const bool bOutputDeviceSRGB = (PassParameters->CombineLUT.OutputDevice.OutputDevice == (uint32)EDisplayOutputFormat::SDR_sRGB);
-
 	PermutationVector.Set<FLUTBlenderCS::FOutputDeviceSRGB>(bOutputDeviceSRGB);
 	TShaderMapRef<FLUTBlenderCS> ComputeShader(View.ShaderMap, PermutationVector);
 
diff --git a/Engine/Source/Runtime/Renderer/Private/PostProcess/PostProcessVisualizeHDR.cpp b/Engine/Source/Runtime/Renderer/Private/PostProcess/PostProcessVisualizeHDR.cpp
--- a/Engine/Source/Runtime/Renderer/Private/PostProcess/PostProcessVisualizeHDR.cpp
+++ b/Engine/Source/Runtime/Renderer/Private/PostProcess/PostProcessVisualizeHDR.cpp
@@ -57,10 +57,6 @@ public:
 	static void ModifyCompilationEnvironment(const FGlobalShaderPermutationParameters& Parameters, FShaderCompilerEnvironment& OutEnvironment)
 	{
 		FGlobalShader::ModifyCompilationEnvironment(Parameters, OutEnvironment);
-		OutEnvironment.SetDefine(TEXT("USE_COLOR_MATRIX"), 1);
-		OutEnvironment.SetDefine(TEXT("USE_SHADOW_TINT"), 1);
-		OutEnvironment.SetDefine(TEXT("USE_CONTRAST"), 1);
-		OutEnvironment.SetDefine(TEXT("USE_APPROXIMATE_SRGB"), (uint32)0);
 
 		// This shader takes a very long time to compile with FXC, so we pre-compile it with DXC first and then forward the optimized HLSL to FXC.
 		OutEnvironment.CompilerFlags.Add(CFLAG_PrecompileWithDXC);
diff --git a/Engine/Source/Runtime/Renderer/Private/SceneRendering.cpp b/Engine/Source/Runtime/Renderer/Private/SceneRendering.cpp
--- a/Engine/Source/Runtime/Renderer/Private/SceneRendering.cpp
+++ b/Engine/Source/Runtime/Renderer/Private/SceneRendering.cpp
@@ -2646,6 +2646,16 @@ IPooledRenderTarget* FViewInfo::GetTonemappingLUT(FRHICommandList& RHICmdList, c
 	return nullptr;
 }
 
+IPooledRenderTarget* FViewInfo::GetInnerTonemappingLUT(FRHICommandList& RHICmdList, const int32 LUTSize, const bool bNeedFloatOutput) const
+{
+	FSceneViewState* EffectiveViewState = GetEyeAdaptationViewState();
+	if (EffectiveViewState)
+	{
+		return EffectiveViewState->GetInnerTonemappingLUT(RHICmdList, LUTSize, bNeedFloatOutput);
+	}
+	return nullptr;
+}
+
 bool FViewInfo::RequiresDebugMaterials() const
 {
 	// We can only use debug materials in ODSC environments.
diff --git a/Engine/Source/Runtime/Renderer/Private/SceneRendering.h b/Engine/Source/Runtime/Renderer/Private/SceneRendering.h
--- a/Engine/Source/Runtime/Renderer/Private/SceneRendering.h
+++ b/Engine/Source/Runtime/Renderer/Private/SceneRendering.h
@@ -1833,6 +1833,9 @@ public:
 	* for stereo rendering, this will force the post-processing to use the same render target for both eyes*/
 	IPooledRenderTarget* GetTonemappingLUT(FRHICommandList& RHICmdList, const int32 LUTSize, const bool bNeedFloatOutput) const;
 
+	/** Gets the other rendertarget that will be populated by CombineLUTS post process */
+	IPooledRenderTarget* GetInnerTonemappingLUT(FRHICommandList& RHICmdList, const int32 LUTSize, const bool bNeedFloatOutput) const;
+
 	bool IsFirstInFamily() const
 	{
 		return Family->Views[0] == this;
diff --git a/Engine/Source/Runtime/Renderer/Private/SceneViewState.h b/Engine/Source/Runtime/Renderer/Private/SceneViewState.h
--- a/Engine/Source/Runtime/Renderer/Private/SceneViewState.h
+++ b/Engine/Source/Runtime/Renderer/Private/SceneViewState.h
@@ -250,9 +250,12 @@ private:
 
 	FSubstrateViewDebugData SubstrateViewDebugData;
 
-	// The LUT used by tonemapping.  In stereo this is only computed and stored by the Left Eye.
+	// The LUT used by tonemapping. In stereo this is only computed and stored by the Left Eye.
 	TRefCountPtr<IPooledRenderTarget> CombinedLUTRenderTarget;
 
+	// The inner LUT used by the LUT used by tonemapping.
+	TRefCountPtr<IPooledRenderTarget> InnerLUTRenderTarget;
+
 	// LUT is only valid after it has been computed, not on allocation of the RT
 	bool bValidTonemappingLUT = false;
 
@@ -665,11 +668,14 @@ public:
 		bValidTonemappingLUT = bValid;
 	}
 
-	static FPooledRenderTargetDesc CreateLUTRenderTarget(const int32 LUTSize, const bool bNeedFloatOutput);
+	static FPooledRenderTargetDesc CreateLUTRenderTarget(const int32 LUTSize, const bool bNeedFloatOutput, const TCHAR* const DebugName);
 
-	// Returns a reference to the render target used for the LUT.  Allocated on the first request.
+	// Returns a reference to the render target used for the LUT. Allocated on the first request.
 	IPooledRenderTarget* GetTonemappingLUT(FRHICommandList& RHICmdList, const int32 LUTSize, const bool bNeedFloatOutput);
 
+	// Returns a reference to the render target used for the inner LUT. Allocated on the first request.
+	IPooledRenderTarget* GetInnerTonemappingLUT(FRHICommandList& RHICmdList, const int32 LUTSize, const bool bNeedFloatOutput);
+
 	IPooledRenderTarget* GetTonemappingLUT() const
 	{
 		return CombinedLUTRenderTarget.GetReference();
