Tie shaper quality for LUT to r.LUT.Shaper rather than EDisplayOutputFormat

Default to Log2 for speed over PQ, even in HDR.

#Synced-CL 52237297

diff --git a/Engine/Config/BaseScalability.ini b/Engine/Config/BaseScalability.ini
--- a/Engine/Config/BaseScalability.ini
+++ b/Engine/Config/BaseScalability.ini
@@ -517,6 +517,7 @@ r.FastBlurThreshold=0
 r.Upscale.Quality=1
 r.LightShaftQuality=0
 r.Filter.SizeScale=0.6
+r.LUT.Shaper=0
 r.LUT.Size=32
 r.Tonemapper.Quality=0
 
@@ -538,6 +539,7 @@ r.FastBlurThreshold=2
 r.Upscale.Quality=2
 r.LightShaftQuality=0
 r.Filter.SizeScale=0.7
+r.LUT.Shaper=0
 r.LUT.Size=32
 r.Tonemapper.Quality=2
 
@@ -572,6 +574,7 @@ r.FastBlurThreshold=3
 r.Upscale.Quality=2
 r.LightShaftQuality=1
 r.Filter.SizeScale=0.8
+r.LUT.Shaper=0
 r.LUT.Size=32
 r.Tonemapper.Quality=5
 
@@ -608,6 +611,7 @@ r.FastBlurThreshold=100
 r.Upscale.Quality=3
 r.LightShaftQuality=1
 r.Filter.SizeScale=1
+r.LUT.Shaper=0
 r.LUT.Size=32
 r.Tonemapper.Quality=5
 
@@ -647,6 +651,7 @@ r.FastBlurThreshold=100
 r.Upscale.Quality=3
 r.LightShaftQuality=1
 r.Filter.SizeScale=1
+r.LUT.Shaper=1
 r.LUT.Size=64
 r.Tonemapper.Quality=5
 
diff --git a/Engine/Shaders/Private/PostProcessCombineLUTs.usf b/Engine/Shaders/Private/PostProcessCombineLUTs.usf
--- a/Engine/Shaders/Private/PostProcessCombineLUTs.usf
+++ b/Engine/Shaders/Private/PostProcessCombineLUTs.usf
@@ -141,6 +141,7 @@ float3 ColorCorrectAll( float3 WorkingColor )
 uint OutputDevice;
 uint OutputGamut;
 float OutputMaxLuminance;
+uint LUTShaper;
 
 uint GetOutputDevice()
 {
@@ -349,7 +350,7 @@ float4 CombineLUTsCommon(float3 LUTEncodedColor)
 	float3 LinearColor;
 	// Decode texture values as ST-2084 (Dolby PQ)
 	BRANCH
-	if (GetOutputDevice() >= TONEMAPPER_OUTPUT_ACES1000nitST2084)
+	if (LUTShaper)
 	{
 		// Since ST2084 returns linear values in nits, divide by a scale factor to convert
 		// the reference nit result to be 1.0 in linear.
diff --git a/Engine/Shaders/Private/PostProcessTonemap.usf b/Engine/Shaders/Private/PostProcessTonemap.usf
--- a/Engine/Shaders/Private/PostProcessTonemap.usf
+++ b/Engine/Shaders/Private/PostProcessTonemap.usf
@@ -107,6 +107,7 @@ float EditorNITLevel;
 uint OutputDevice;
 uint OutputGamut;
 float OutputMaxLuminance;
+uint LUTShaper;
 
 float4 BloomDirtMaskTint;
 Texture2D BloomDirtMaskTexture;
@@ -190,7 +191,7 @@ half3 ColorLookupTable( half3 LinearColor )
 	float3 LUTEncodedColor;
 	// Encode as ST-2084 (Dolby PQ) values
 	BRANCH
-	if(GetOutputDevice() >= TONEMAPPER_OUTPUT_ACES1000nitST2084)
+	if (LUTShaper)
 	{
 		// ST2084 expects to receive linear values 0-10000 in nits.
 		// So the linear value must be multiplied by a scale factor to convert to nits.
diff --git a/Engine/Shaders/Private/StochasticLighting/StochasticLightingVisualize.ush b/Engine/Shaders/Private/StochasticLighting/StochasticLightingVisualize.ush
--- a/Engine/Shaders/Private/StochasticLighting/StochasticLightingVisualize.ush
+++ b/Engine/Shaders/Private/StochasticLighting/StochasticLightingVisualize.ush
@@ -13,7 +13,7 @@ SamplerState ColorGradingLUTSampler;
 
 float3 ColorLookupTable(float3 LinearColor)
 {
-	//@todo - reuse with PostProcessTonemap.usf and handle OutputDevice
+	//@todo - reuse with PostProcessTonemap.usf and handle LUTShaper
 	float3 LUTEncodedColor = LinToLog( LinearColor + LogToLin( 0 ) );
 	float3 UVW = LUTEncodedColor * ((LUT_SIZE - 1) / LUT_SIZE) + (0.5f / LUT_SIZE);
 	return ColorGradingLUT.SampleLevel(ColorGradingLUTSampler, UVW, 0).rgb * LUT_SCALING;
diff --git a/Engine/Source/Runtime/Engine/Private/SceneView.cpp b/Engine/Source/Runtime/Engine/Private/SceneView.cpp
--- a/Engine/Source/Runtime/Engine/Private/SceneView.cpp
+++ b/Engine/Source/Runtime/Engine/Private/SceneView.cpp
@@ -283,7 +283,7 @@ static TAutoConsoleVariable<int32> CVarTonemapperQuality(
 	TEXT(" 0: basic tonemapper only, lowest quality\n")
 	TEXT(" 2: + Vignette\n")
 	TEXT(" 4: + Grain\n")
-	TEXT(" 5: + GrainJitter = full quality (default)"),
+	TEXT(" 5: (default)"),
 	ECVF_Scalability | ECVF_RenderThreadSafe);
 
 // should be changed to BaseColor and Metallic, since some time now UE is not using DiffuseColor and SpecularColor any more
diff --git a/Engine/Source/Runtime/Renderer/Private/PostProcess/PostProcessCombineLUTs.cpp b/Engine/Source/Runtime/Renderer/Private/PostProcess/PostProcessCombineLUTs.cpp
--- a/Engine/Source/Runtime/Renderer/Private/PostProcess/PostProcessCombineLUTs.cpp
+++ b/Engine/Source/Runtime/Renderer/Private/PostProcess/PostProcessCombineLUTs.cpp
@@ -388,6 +388,7 @@ struct FCachedLUTSettings
 		UPDATE_CACHE_SETTINGS(Parameters.OutputDevice.OutputDevice, TonemapperOutputDeviceParameters.OutputDevice, bHasChanged);
 		UPDATE_CACHE_SETTINGS(Parameters.OutputDevice.OutputGamut, TonemapperOutputDeviceParameters.OutputGamut, bHasChanged);
 		UPDATE_CACHE_SETTINGS(Parameters.OutputDevice.OutputMaxLuminance, TonemapperOutputDeviceParameters.OutputMaxLuminance, bHasChanged);
+		UPDATE_CACHE_SETTINGS(Parameters.OutputDevice.LUTShaper, TonemapperOutputDeviceParameters.LUTShaper, bHasChanged);
 
 		UPDATE_CACHE_SETTINGS(Parameters.ColorScale, FVector3f(View.ColorScale), bHasChanged);
 		UPDATE_CACHE_SETTINGS(Parameters.OverlayColor, FVector4f(View.OverlayColor), bHasChanged);
diff --git a/Engine/Source/Runtime/Renderer/Private/PostProcess/PostProcessTonemap.cpp b/Engine/Source/Runtime/Renderer/Private/PostProcess/PostProcessTonemap.cpp
--- a/Engine/Source/Runtime/Renderer/Private/PostProcess/PostProcessTonemap.cpp
+++ b/Engine/Source/Runtime/Renderer/Private/PostProcess/PostProcessTonemap.cpp
@@ -70,6 +70,13 @@ static TAutoConsoleVariable<int32> CVarBackbufferQuantizationDitheringOverride(
 	TEXT("Disabled by default. Instead is automatically found out by FSceneViewFamily::RenderTarget's pixel format of the backbuffer."),
 	ECVF_RenderThreadSafe);
 
+TAutoConsoleVariable<int32> CVarLUTShaper(
+	TEXT("r.LUT.Shaper"), 0,
+	TEXT("Which shaper to use for addressing the LUT\n")
+	TEXT(" 0: Log2 (fast)\n")
+	TEXT(" 1: PQ"),
+	ECVF_Scalability | ECVF_RenderThreadSafe);
+
 const int32 GTonemapComputeTileSizeX = 8;
 const int32 GTonemapComputeTileSizeY = 8;
 
@@ -238,8 +245,6 @@ RDG_REGISTER_BLACKBOARD_STRUCT(FOutputLuminance);
 
 FTonemapperOutputDeviceParameters GetTonemapperOutputDeviceParameters(const FSceneViewFamily& Family)
 {
-	static TConsoleVariableData<float>* CVarOutputGamma = IConsoleManager::Get().FindTConsoleVariableDataFloat(TEXT("r.TonemapperGamma"));
-
 	EDisplayOutputFormat OutputDeviceValue;
 	bool bIsHDR = false;
 
@@ -257,7 +262,7 @@ FTonemapperOutputDeviceParameters GetTonemapperOutputDeviceParameters(const FSce
 		bIsHDR = Family.RenderTarget->GetSceneHDREnabled();
 	}
 
-	float Gamma = CVarOutputGamma->GetValueOnRenderThread();
+	float Gamma = CVarTonemapperGamma.GetValueOnRenderThread();
 
 	// In case gamma is unspecified, fall back to 2.2 which is the most common case
 	if ((PLATFORM_APPLE || OutputDeviceValue == EDisplayOutputFormat::SDR_ExplicitGammaMapping) && Gamma == 0.0f)
@@ -284,6 +289,7 @@ FTonemapperOutputDeviceParameters GetTonemapperOutputDeviceParameters(const FSce
 	Parameters.OutputMaxLuminance = CVarHDRDisplayOverrideOSMaxLuminance->GetBool() ?
 		(bIsHDR ? HDRGetDisplayMaximumLuminance() : 100.f) :
 		Family.RenderTarget->GetMaximumLuminanceInNits();
+	Parameters.LUTShaper = CVarLUTShaper.GetValueOnRenderThread();
 	return Parameters;
 }
 
@@ -1150,6 +1156,7 @@ class FMobileCustomResolvePS : public FGlobalShader
 		SHADER_PARAMETER(float, InvLUTSize)
 		SHADER_PARAMETER(float, LUTScale)
 		SHADER_PARAMETER(float, LUTOffset)
+		SHADER_PARAMETER(uint32, LUTShaper)
 	END_SHADER_PARAMETER_STRUCT()
 
 	class FTonemapperSubpassMsaaDim : SHADER_PERMUTATION_SPARSE_INT("SUBPASS_MSAA_SAMPLES", 0, 1, 2, 4, 8);
@@ -1195,6 +1202,7 @@ void RenderMobileCustomResolve(FRHICommandList& RHICmdList, const FViewInfo& Vie
 	PSShaderParameters.InvLUTSize = 1.0f / LUTSize;
 	PSShaderParameters.LUTScale = (LUTSize - 1.0f) / LUTSize;
 	PSShaderParameters.LUTOffset = 0.5f / LUTSize;
+	PSShaderParameters.LUTShaper = CVarLUTShaper.GetValueOnRenderThread();
 	if (SubpassMSAASamples == 0u)
 	{
 		PSShaderParameters.ColorTexture = SceneTextures.Color.Resolve;
diff --git a/Engine/Source/Runtime/Renderer/Private/PostProcess/PostProcessTonemap.h b/Engine/Source/Runtime/Renderer/Private/PostProcess/PostProcessTonemap.h
--- a/Engine/Source/Runtime/Renderer/Private/PostProcess/PostProcessTonemap.h
+++ b/Engine/Source/Runtime/Renderer/Private/PostProcess/PostProcessTonemap.h
@@ -17,6 +17,7 @@ BEGIN_SHADER_PARAMETER_STRUCT(FTonemapperOutputDeviceParameters, )
 	SHADER_PARAMETER(uint32, OutputDevice)
 	SHADER_PARAMETER(uint32, OutputGamut)
 	SHADER_PARAMETER(float, OutputMaxLuminance)
+	SHADER_PARAMETER(uint32, LUTShaper)
 END_SHADER_PARAMETER_STRUCT()
 
 RENDERER_API FTonemapperOutputDeviceParameters GetTonemapperOutputDeviceParameters(const FSceneViewFamily& Family);
