Bypass 1.05 division in the tonemapping LUT when it's not needed

Filmic SDR may exceed 1.0, but Filmic HDR and ACES do not.

#Synced-CL 52240169

diff --git a/Engine/Shaders/Private/CompositeUIPixelShader.usf b/Engine/Shaders/Private/CompositeUIPixelShader.usf
--- a/Engine/Shaders/Private/CompositeUIPixelShader.usf
+++ b/Engine/Shaders/Private/CompositeUIPixelShader.usf
@@ -7,9 +7,6 @@ CompositeUIPixelShader.usf: Filter pixel shader source.
 #include "TonemapCommon.ush"
 #include "ColorDeficiency.ush"
 
-#define LUT_SCALING	1.05f
-#define LUT_SIZE	32.f
-
 ///////////////////////////////////////////////////////////////////////////////////////
 
 #if APPLY_COLOR_DEFICIENCY
diff --git a/Engine/Shaders/Private/PostProcessCombineLUTs.usf b/Engine/Shaders/Private/PostProcessCombineLUTs.usf
--- a/Engine/Shaders/Private/PostProcessCombineLUTs.usf
+++ b/Engine/Shaders/Private/PostProcessCombineLUTs.usf
@@ -45,6 +45,7 @@ float3x3 GetLimitingXYZToRgb()
 
 uint bIsTemperatureWhiteBalance;
 
+float LUTInvMax;
 float LUTSize;
 float InnerLUTSize;
 
@@ -472,8 +473,8 @@ float4 CombineLUTsCommon(float3 LUTEncodedColor)
 		OutDeviceColor = pow( OutputGamutColor, InverseGamma.z );
 	}
 
-	// Better to saturate(lerp(a,b,t)) than lerp(saturate(a),saturate(b),t)
-	OutColor.rgb = OutDeviceColor / 1.05;
+	// Some functions output values slightly above 1. Undo this squishing in PostProcessTonemap.usf
+	OutColor.rgb = OutDeviceColor * LUTInvMax;
 	OutColor.a = 0;
 
 	return OutColor;
diff --git a/Engine/Shaders/Private/PostProcessTonemap.usf b/Engine/Shaders/Private/PostProcessTonemap.usf
--- a/Engine/Shaders/Private/PostProcessTonemap.usf
+++ b/Engine/Shaders/Private/PostProcessTonemap.usf
@@ -181,6 +181,7 @@ bool IsFilmGrainEnabled()
 Texture3D ColorGradingLUT;
 SamplerState ColorGradingLUTSampler;
 
+float LUTMax;
 float LUTSize;
 float InvLUTSize; // 1 / LUTSize
 float LUTScale; // (LUTSize - 1) / LUTSize
@@ -206,7 +207,7 @@ half3 ColorLookupTable( half3 LinearColor )
 	float3 UVW = LUTEncodedColor * LUTScale + LUTOffset;
 
 	half3 OutDeviceColor = Texture3DSample( ColorGradingLUT, ColorGradingLUTSampler, UVW ).rgb;
-	return OutDeviceColor * 1.05;
+	return OutDeviceColor * LUTMax;
 }
 
 // can be optimized
diff --git a/Engine/Source/Runtime/Renderer/Private/PostProcess/PostProcessCombineLUTs.cpp b/Engine/Source/Runtime/Renderer/Private/PostProcess/PostProcessCombineLUTs.cpp
--- a/Engine/Source/Runtime/Renderer/Private/PostProcess/PostProcessCombineLUTs.cpp
+++ b/Engine/Source/Runtime/Renderer/Private/PostProcess/PostProcessCombineLUTs.cpp
@@ -246,6 +246,7 @@ BEGIN_SHADER_PARAMETER_STRUCT(FCombineLUTParameters, )
 	SHADER_PARAMETER(FVector4f, ColorGammaHighlights)
 	SHADER_PARAMETER(FVector4f, ColorGainHighlights)
 	SHADER_PARAMETER(FVector4f, ColorOffsetHighlights)
+	SHADER_PARAMETER(float, LUTInvMax)
 	SHADER_PARAMETER(float, LUTSize)
 	SHADER_PARAMETER(float, InnerLUTSize)
 	SHADER_PARAMETER(float, WhiteTemp)
@@ -316,28 +317,48 @@ uint32 DoesGamutWhitePointMatch(uint32 OutputGamut, const UE::Color::FColorSpace
 	return bMatch;
 }
 
+ETonemappingMethod GetTonemappingMethod(const FPostProcessSettings& Settings)
+{
+	const int32 TonemappingMethodOverride = CVarLUTTonemappingMethod.GetValueOnRenderThread();
+	return (TonemappingMethodOverride >= 0) ?
+		static_cast<ETonemappingMethod>(TonemappingMethodOverride) :
+		Settings.TonemappingMethod;
+}
+
+float GetTonemappingMax(const FTonemapperOutputDeviceParameters& Parameters, ETonemappingMethod TonemappingMethod)
+{
+	float LUTMax = 1.f;
+	if (TonemappingMethod == ETonemappingMethod::Filmic)
+	{
+		switch (static_cast<EDisplayOutputFormat>(Parameters.OutputDevice))
+		{
+		case EDisplayOutputFormat::SDR_sRGB:
+		case EDisplayOutputFormat::SDR_Rec709:
+		case EDisplayOutputFormat::HDR_LinearWithToneCurve:
+				LUTMax =  1.05f;
+		}
+	}
+	return LUTMax;
+}
+
 struct FCachedLUTSettings
 {
 	uint32 UniqueID = 0;
 	EShaderPlatform ShaderPlatform = GMaxRHIShaderPlatform;
+	ETonemappingMethod TonemappingMethod{};
 	FCombineLUTParameters Parameters;
 	FWorkingColorSpaceShaderParameters WorkingColorSpaceShaderParameters;
 	int32 AcesVersion = 0;
-	ETonemappingMethod TonemappingMethod{};
 
 	bool UpdateCachedValues(FRDGBuilder& GraphBuilder, const FViewInfo& View, const FPostProcessSettings& Settings, const FTexture* const* Textures, const float* Weights, uint32 BlendCount)
 	{
 		bool bHasChanged = false;
+		UPDATE_CACHE_SETTINGS(TonemappingMethod, GetTonemappingMethod(Settings), bHasChanged);
 		GetCombineLUTParameters(GraphBuilder, View, Settings, Textures, Weights, BlendCount, bHasChanged);
 		UPDATE_CACHE_SETTINGS(UniqueID, View.State ? View.State->GetViewKey() : 0, bHasChanged);
 		UPDATE_CACHE_SETTINGS(ShaderPlatform, View.GetShaderPlatform(), bHasChanged);
 		static const auto CVarAcesVersion = IConsoleManager::Get().FindTConsoleVariableDataInt(TEXT("r.HDR.Aces.Version"));
 		UPDATE_CACHE_SETTINGS(AcesVersion, CVarAcesVersion->GetValueOnRenderThread(), bHasChanged);
-		const int32 TonemappingMethodOverride = CVarLUTTonemappingMethod.GetValueOnRenderThread();
-		const ETonemappingMethod NextTonemappingMethod = (TonemappingMethodOverride >= 0) ?
-			static_cast<ETonemappingMethod>(TonemappingMethodOverride) :
-			Settings.TonemappingMethod;
-		UPDATE_CACHE_SETTINGS(TonemappingMethod, NextTonemappingMethod, bHasChanged);
 
 		const FWorkingColorSpaceShaderParameters* InWorkingColorSpaceShaderParameters = reinterpret_cast<const FWorkingColorSpaceShaderParameters*>(GDefaultWorkingColorSpaceUniformBuffer.GetContents());
 		if (InWorkingColorSpaceShaderParameters)
@@ -396,6 +417,7 @@ struct FCachedLUTSettings
 
 		// White balance
 		UPDATE_CACHE_SETTINGS(Parameters.bIsTemperatureWhiteBalance, uint32(Settings.TemperatureType == ETemperatureMethod::TEMP_WhiteBalance), bHasChanged);
+		UPDATE_CACHE_SETTINGS(Parameters.LUTInvMax, 1.f / GetTonemappingMax(TonemapperOutputDeviceParameters, TonemappingMethod), bHasChanged);
 		UPDATE_CACHE_SETTINGS(Parameters.LUTSize, (float)CVarLUTSize->GetInt(), bHasChanged);
 		UPDATE_CACHE_SETTINGS(Parameters.InnerLUTSize, (float)CVarLUTInnerSize->GetInt(), bHasChanged);
 		UPDATE_CACHE_SETTINGS(Parameters.WhiteTemp, Settings.WhiteTemp, bHasChanged);
@@ -607,11 +629,7 @@ struct FCachedInnerLUTSettings
 		UPDATE_CACHE_SETTINGS(Parameters.HDRPaperWhiteRatio, HDRPaperWhiteRatio, bHasChanged);
 		UPDATE_CACHE_SETTINGS(Parameters.InvHDRPaperWhiteRatio, 1.f / HDRPaperWhiteRatio, bHasChanged);
 
-		const int32 TonemappingMethodOverride = CVarLUTTonemappingMethod.GetValueOnRenderThread();
-		const uint32 TonemappingMethod = (TonemappingMethodOverride >= 0) ?
-			static_cast<uint32>(TonemappingMethodOverride) :
-			static_cast<uint32>(Settings.TonemappingMethod);
-		UPDATE_CACHE_SETTINGS(Parameters.TonemappingMethod, TonemappingMethod, bHasChanged);
+		UPDATE_CACHE_SETTINGS(Parameters.TonemappingMethod, static_cast<uint32>(GetTonemappingMethod(Settings)), bHasChanged);
 	}
 };
 
diff --git a/Engine/Source/Runtime/Renderer/Private/PostProcess/PostProcessTonemap.cpp b/Engine/Source/Runtime/Renderer/Private/PostProcess/PostProcessTonemap.cpp
--- a/Engine/Source/Runtime/Renderer/Private/PostProcess/PostProcessTonemap.cpp
+++ b/Engine/Source/Runtime/Renderer/Private/PostProcess/PostProcessTonemap.cpp
@@ -354,19 +354,20 @@ BEGIN_SHADER_PARAMETER_STRUCT(FTonemapParameters, )
 	SHADER_PARAMETER(FVector4f, ChromaticAberrationParams)
 	SHADER_PARAMETER(FVector4f, LensPrincipalPointOffsetScale)
 	SHADER_PARAMETER(FVector4f, LensPrincipalPointOffsetScaleInverse)
-	SHADER_PARAMETER(uint32,	Flags)
-	SHADER_PARAMETER(float, 	SharpenDiv6)
-	SHADER_PARAMETER(float, 	LUTSize)
-	SHADER_PARAMETER(float, 	InvLUTSize)
-	SHADER_PARAMETER(float, 	LUTScale)
-	SHADER_PARAMETER(float, 	LUTOffset)
-	SHADER_PARAMETER(float, 	EditorNITLevel)
-	SHADER_PARAMETER(float, 	BackbufferQuantizationDithering)
+	SHADER_PARAMETER(uint32, Flags)
+	SHADER_PARAMETER(float, SharpenDiv6)
+	SHADER_PARAMETER(float, LUTMax)
+	SHADER_PARAMETER(float, LUTSize)
+	SHADER_PARAMETER(float, InvLUTSize)
+	SHADER_PARAMETER(float, LUTScale)
+	SHADER_PARAMETER(float, LUTOffset)
+	SHADER_PARAMETER(float, EditorNITLevel)
+	SHADER_PARAMETER(float, BackbufferQuantizationDithering)
 
 	SHADER_PARAMETER(FVector4f, VignetteParams0)
 	SHADER_PARAMETER(FVector4f, VignetteParams1)
-	SHADER_PARAMETER(float, 	VignetteParams2)
-	SHADER_PARAMETER(uint32,	VignetteType)
+	SHADER_PARAMETER(float, VignetteParams2)
+	SHADER_PARAMETER(uint32, VignetteType)
 	SHADER_PARAMETER_RDG_TEXTURE(Texture2D, VignetteTexture)
 	SHADER_PARAMETER_SAMPLER(SamplerState, VignetteTextureSampler)
 END_SHADER_PARAMETER_STRUCT()
@@ -605,6 +606,8 @@ bool ShouldOverrideOutputLoadActionToClear(const FScreenPassRenderTarget& Output
 	return bShouldOverride;
 }
 
+extern ETonemappingMethod GetTonemappingMethod(const FPostProcessSettings& Settings);
+extern float GetTonemappingMax(const FTonemapperOutputDeviceParameters& Parameters, ETonemappingMethod TonemappingMethod);
 FScreenPassTexture AddTonemapPass(FRDGBuilder& GraphBuilder, const FViewInfo& View, const FTonemapInputs& Inputs)
 {
 	if (!Inputs.bGammaOnly)
@@ -936,6 +939,7 @@ FScreenPassTexture AddTonemapPass(FRDGBuilder& GraphBuilder, const FViewInfo& Vi
 			}
 		}
 	}
+	CommonParameters.LUTMax = GetTonemappingMax(CommonParameters.OutputDevice, GetTonemappingMethod(PostProcessSettings));
 	CommonParameters.LUTSize = LUTSize;
 	CommonParameters.InvLUTSize = 1.0f / LUTSize;
 	CommonParameters.LUTScale = (LUTSize - 1.0f) / LUTSize;
@@ -1152,6 +1156,7 @@ class FMobileCustomResolvePS : public FGlobalShader
 		SHADER_PARAMETER(FVector4f, ColorScale0)
 		SHADER_PARAMETER_TEXTURE(Texture3D, ColorGradingLUT)
 		SHADER_PARAMETER_SAMPLER(SamplerState, ColorGradingLUTSampler)
+		SHADER_PARAMETER(float, LUTMax)
 		SHADER_PARAMETER(float, LUTSize)
 		SHADER_PARAMETER(float, InvLUTSize)
 		SHADER_PARAMETER(float, LUTScale)
@@ -1198,6 +1203,8 @@ void RenderMobileCustomResolve(FRHICommandList& RHICmdList, const FViewInfo& Vie
 	PSShaderParameters.ColorScale0 = FVector4f(Settings.SceneColorTint.R, Settings.SceneColorTint.G, Settings.SceneColorTint.B, 0);
 	PSShaderParameters.ColorGradingLUT = ColorGradingLUT ? ColorGradingLUT->GetRHI() : GBlackVolumeTexture->TextureRHI.GetReference();
 	PSShaderParameters.ColorGradingLUTSampler = TStaticSamplerState<SF_Bilinear, AM_Clamp, AM_Clamp, AM_Clamp>::GetRHI();
+	const FTonemapperOutputDeviceParameters OutputDeviceParameters = GetTonemapperOutputDeviceParameters(*View.Family);
+	PSShaderParameters.LUTMax = GetTonemappingMax(OutputDeviceParameters, GetTonemappingMethod(Settings));
 	PSShaderParameters.LUTSize = LUTSize;
 	PSShaderParameters.InvLUTSize = 1.0f / LUTSize;
 	PSShaderParameters.LUTScale = (LUTSize - 1.0f) / LUTSize;
