Fix for static analysis warning.

#Synced-CL 52261411

diff --git a/Engine/Source/Runtime/Renderer/Private/Skinning/SkinningSceneExtension.cpp b/Engine/Source/Runtime/Renderer/Private/Skinning/SkinningSceneExtension.cpp
--- a/Engine/Source/Runtime/Renderer/Private/Skinning/SkinningSceneExtension.cpp
+++ b/Engine/Source/Runtime/Renderer/Private/Skinning/SkinningSceneExtension.cpp
@@ -2141,7 +2141,9 @@ void FSkinningSceneExtension::ResolveAttachments(FRDGBuilder& GraphBuilder, cons
 				BlockHeader.LocalOffsetRotation[2] = Socket.Rotation.Z;
 				BlockHeader.LocalOffsetRotation[3] = Socket.Rotation.W;
 
-				FMemory::Memcpy(&BlockHeader.BoneRefPoseRow0, Socket.RefPoseMatrix.M, sizeof(float) * 12);
+				FMemory::Memcpy(BlockHeader.BoneRefPoseRow0, Socket.RefPoseMatrix.M[0], sizeof(float) * 4);
+				FMemory::Memcpy(BlockHeader.BoneRefPoseRow1, Socket.RefPoseMatrix.M[1], sizeof(float) * 4);
+				FMemory::Memcpy(BlockHeader.BoneRefPoseRow2, Socket.RefPoseMatrix.M[2], sizeof(float) * 4);
 			}
 		}
 
