[Visualize light shape] Fix simple lights not being added to the sort keys array correctly

#tests tested in editor
#Synced-CL 52238062

diff --git a/Engine/Source/Runtime/Renderer/Private/LightGridInjection.cpp b/Engine/Source/Runtime/Renderer/Private/LightGridInjection.cpp
--- a/Engine/Source/Runtime/Renderer/Private/LightGridInjection.cpp
+++ b/Engine/Source/Runtime/Renderer/Private/LightGridInjection.cpp
@@ -1617,7 +1617,7 @@ FComputeLightGridOutput FSceneRenderer::ComputeLightGrid(FRDGBuilder& GraphBuild
 
 					if (bVisualizeLightShape)
 					{
-						VisualizeLightShape::FSortKey& SortKey = VisualizeLightShapeSortKeys[VisualizeLightShapeSortKeys.AddUninitialized()];
+						VisualizeLightShape::FSortKey& SortKey = VisualizeLightShapeSortKeys[SortedIndex];
 						SortKey.LightData.ForwardLightIndex = IndexInBuffer;
 						SortKey.LightData.LightType = LightType_Point;
 						SortKey.Distance = FVector3f(ViewSpacePosAndRadius).SizeSquared();
