VolumetricCloud - made local light sampling possible to use with scalability.
Also removed unused function.

#tests PC editor, toggling scalability after setting up engine BaseScalability locally.
#rnx
#Synced-CL 52290745

diff --git a/Engine/Source/Runtime/Renderer/Private/VolumetricCloudRendering.cpp b/Engine/Source/Runtime/Renderer/Private/VolumetricCloudRendering.cpp
--- a/Engine/Source/Runtime/Renderer/Private/VolumetricCloudRendering.cpp
+++ b/Engine/Source/Runtime/Renderer/Private/VolumetricCloudRendering.cpp
@@ -242,12 +242,12 @@ static TAutoConsoleVariable<int32> CVarVolumetricCloudEnableAtmosphericLightsSam
 static TAutoConsoleVariable<int32> CVarVolumetricCloudEnableLocalLightsSampling(
 	TEXT("r.VolumetricCloud.EnableLocalLightsSampling"), 0,
 	TEXT("[EXPERIMENTAL] Enable/disable local lights contribution on clouds. Expenssive! Use for cinematics if needed."),
-	ECVF_RenderThreadSafe);
+	ECVF_RenderThreadSafe | ECVF_Scalability);
 
 static TAutoConsoleVariable<float> CVarVolumetricCloudLocalLightsShadowSampleCount(
 	TEXT("r.VolumetricCloud.LocalLights.ShadowSampleCount"), 12,
 	TEXT("[EXPERIMENTAL] Set the volumetric shadow sample count when evaluating local lights. Expenssive! Use for cinematics if needed."),
-	ECVF_RenderThreadSafe);
+	ECVF_RenderThreadSafe | ECVF_Scalability);
 
 ////////////////////////////////////////////////////////////////////////// 
 
@@ -340,11 +340,6 @@ bool ShouldViewVisualizeVolumetricCloudConservativeDensity(const FViewInfo& View
 		&& GIsEditor;
 }
 
-bool VolumetricCloudWantsToSampleLocalLights(const FScene* Scene, const FEngineShowFlags& EngineShowFlags)
-{
-	return ShouldRenderVolumetricCloud(Scene, EngineShowFlags) && CVarVolumetricCloudEnableLocalLightsSampling.GetValueOnRenderThread();
-}
-
 static bool ShouldRenderCloudShadowmap(const FLightSceneProxy* AtmosphericLight)
 {
 	return CVarVolumetricCloudShadowMap.GetValueOnRenderThread() > 0 && AtmosphericLight && AtmosphericLight->GetCastCloudShadows() && AtmosphericLight->GetCloudShadowStrength() > 0;
diff --git a/Engine/Source/Runtime/Renderer/Private/VolumetricCloudRendering.h b/Engine/Source/Runtime/Renderer/Private/VolumetricCloudRendering.h
--- a/Engine/Source/Runtime/Renderer/Private/VolumetricCloudRendering.h
+++ b/Engine/Source/Runtime/Renderer/Private/VolumetricCloudRendering.h
@@ -110,7 +110,6 @@ bool ShouldRenderVolumetricCloud(const FScene* Scene, const FEngineShowFlags& En
 bool ShouldRenderVolumetricCloudWithBlueNoise_GameThread(const FScene* Scene, const FSceneView& View);
 
 bool ShouldViewVisualizeVolumetricCloudConservativeDensity(const FViewInfo& ViewInfo, const FEngineShowFlags& EngineShowFlags);
-bool VolumetricCloudWantsToSampleLocalLights(const FScene* Scene, const FEngineShowFlags& EngineShowFlags);
 uint32 GetVolumetricCloudDebugViewMode(const FEngineShowFlags& ShowFlags);
 bool ShouldVolumetricCloudTraceWithMinMaxDepth(const FViewInfo& ViewInfo);
 bool ShouldVolumetricCloudTraceWithMinMaxDepth(const TArray<FViewInfo>& Views);
