Removed old void RenderAtmosphere unused declaration and stat group.

#tests Pc editor
#rnx
#Synced-CL 52290745

diff --git a/Engine/Source/Runtime/Renderer/Private/DeferredShadingRenderer.cpp b/Engine/Source/Runtime/Renderer/Private/DeferredShadingRenderer.cpp
--- a/Engine/Source/Runtime/Renderer/Private/DeferredShadingRenderer.cpp
+++ b/Engine/Source/Runtime/Renderer/Private/DeferredShadingRenderer.cpp
@@ -410,7 +410,6 @@ DECLARE_CYCLE_STAT(TEXT("DeferredShadingSceneRenderer FXSystem PostRenderOpaque"
 DECLARE_CYCLE_STAT(TEXT("DeferredShadingSceneRenderer AfterBasePass"), STAT_FDeferredShadingSceneRenderer_AfterBasePass, STATGROUP_SceneRendering);
 DECLARE_CYCLE_STAT(TEXT("DeferredShadingSceneRenderer Lighting"), STAT_FDeferredShadingSceneRenderer_Lighting, STATGROUP_SceneRendering);
 DECLARE_CYCLE_STAT(TEXT("DeferredShadingSceneRenderer RenderLightShaftOcclusion"), STAT_FDeferredShadingSceneRenderer_RenderLightShaftOcclusion, STATGROUP_SceneRendering);
-DECLARE_CYCLE_STAT(TEXT("DeferredShadingSceneRenderer RenderAtmosphere"), STAT_FDeferredShadingSceneRenderer_RenderAtmosphere, STATGROUP_SceneRendering);
 DECLARE_CYCLE_STAT(TEXT("DeferredShadingSceneRenderer RenderSkyAtmosphere"), STAT_FDeferredShadingSceneRenderer_RenderSkyAtmosphere, STATGROUP_SceneRendering);
 DECLARE_CYCLE_STAT(TEXT("DeferredShadingSceneRenderer RenderFog"), STAT_FDeferredShadingSceneRenderer_RenderFog, STATGROUP_SceneRendering);
 DECLARE_CYCLE_STAT(TEXT("DeferredShadingSceneRenderer RenderLocalFogVolume"), STAT_FDeferredShadingSceneRenderer_RenderLocalFogVolume, STATGROUP_SceneRendering);
diff --git a/Engine/Source/Runtime/Renderer/Private/DeferredShadingRenderer.h b/Engine/Source/Runtime/Renderer/Private/DeferredShadingRenderer.h
--- a/Engine/Source/Runtime/Renderer/Private/DeferredShadingRenderer.h
+++ b/Engine/Source/Runtime/Renderer/Private/DeferredShadingRenderer.h
@@ -598,11 +598,6 @@ private:
 		const FSceneWithoutWaterTextures& SceneWithoutWaterTextures,
 		TRDGUniformBufferRef<FSceneTextureUniformParameters> SceneTexturesWithDepth);
 
-	void RenderAtmosphere(
-		FRDGBuilder& GraphBuilder,
-		const FMinimalSceneTextures& SceneTextures,
-		FRDGTextureRef LightShaftOcclusionTexture);
-
 	// TODO: Address tech debt to that directly in RenderDiffuseIndirectAndAmbientOcclusion()
 	void SetupCommonDiffuseIndirectParameters(
 		FRDGBuilder& GraphBuilder,
