Fix CIS

#Synced-CL 52288206

diff --git a/Engine/Source/Runtime/Renderer/Private/DebugViewModeRendering.h b/Engine/Source/Runtime/Renderer/Private/DebugViewModeRendering.h
--- a/Engine/Source/Runtime/Renderer/Private/DebugViewModeRendering.h
+++ b/Engine/Source/Runtime/Renderer/Private/DebugViewModeRendering.h
@@ -187,11 +187,9 @@ private:
 
 	const FDebugViewModeInterface* DebugViewModeInterface;
 
-#if !UE_BUILD_SHIPPING
 	/** Per-proxy cache for texture streaming deficit. Keyed on proxy (not FMaterial*) because
 	    material instances sharing the same compiled material can have different texture overrides. */
 	TMap<const FMaterialRenderProxy*, int32> TextureStreamingDeficitCache;
-#endif
 };
 
 class FDebugViewModeImplementation : public FDebugViewModeInterface
