VolumetricCloud - make sure the secondary atmospehric light can be ignored when rendering cloud and cloud scattering scale is set to black.

#tests PC editor cloud with two lights
#rnx
#Synced-CL 52290745

diff --git a/Engine/Source/Runtime/Engine/Classes/Components/DirectionalLightComponent.h b/Engine/Source/Runtime/Engine/Classes/Components/DirectionalLightComponent.h
--- a/Engine/Source/Runtime/Engine/Classes/Components/DirectionalLightComponent.h
+++ b/Engine/Source/Runtime/Engine/Classes/Components/DirectionalLightComponent.h
@@ -235,6 +235,7 @@ class UDirectionalLightComponent : public ULightComponent
 
 	/**
 	 * Scales the lights contribution when scattered in cloud participating media. This can help counter balance the fact that our multiple scattering solution is only an approximation.
+	 * Setting this to 0 on the secondary atmospheric light will optimise the cloud tracing to only use a single light.
 	 */
 	UPROPERTY(EditAnywhere, BlueprintReadOnly, interp, Category = AtmosphereAndCloud, meta = (HideAlphaChannel))
 	FLinearColor CloudScatteredLuminanceScale;
diff --git a/Engine/Source/Runtime/Renderer/Private/ReflectionEnvironmentRealTimeCapture.cpp b/Engine/Source/Runtime/Renderer/Private/ReflectionEnvironmentRealTimeCapture.cpp
--- a/Engine/Source/Runtime/Renderer/Private/ReflectionEnvironmentRealTimeCapture.cpp
+++ b/Engine/Source/Runtime/Renderer/Private/ReflectionEnvironmentRealTimeCapture.cpp
@@ -509,7 +509,7 @@ void FScene::AllocateAndCaptureFrameSkyEnvMap(
 					CloudRC.bShouldViewRenderVolumetricRenderTarget = false;
 					CloudRC.bIsReflectionRendering = true;
 					CloudRC.bIsSkyRealTimeReflectionRendering = true;
-					CloudRC.bSecondAtmosphereLightEnabled = IsSecondAtmosphereLightEnabled();
+					CloudRC.bSecondAtmosphereLightEnabled = IsSecondAtmosphereLightEnabledForCloud();
 
 					CloudRC.bSkipAtmosphericLightShadowmap = !CaptureShadowFromOpaque;
 					if (CaptureShadowFromOpaque)
diff --git a/Engine/Source/Runtime/Renderer/Private/ScenePrivate.h b/Engine/Source/Runtime/Renderer/Private/ScenePrivate.h
--- a/Engine/Source/Runtime/Renderer/Private/ScenePrivate.h
+++ b/Engine/Source/Runtime/Renderer/Private/ScenePrivate.h
@@ -2265,6 +2265,26 @@ public:
 		return AtmosphereLights[1] != nullptr;
 	}
 
+	bool IsSecondAtmosphereLightEnabledForCloud()
+	{
+		if (IsSecondAtmosphereLightEnabled())
+		{
+			// Check the scattering scale of the second light to make sure it can contribute to the clouds before enabling it, as the second atmosphere light is an expensive feature.
+			FLightSceneInfo* AtmosphericLight1Info = AtmosphereLights[1];
+			FLightSceneProxy* AtmosphericLight1 = AtmosphericLight1Info ? AtmosphericLight1Info->Proxy : nullptr;
+			if (AtmosphericLight1)
+			{
+				FLinearColor ColorScale = AtmosphericLight1->GetCloudScatteredLuminanceScale();
+				if (ColorScale.R <= 0.0f && ColorScale.G <= 0.0f && ColorScale.B <= 0.0f)
+				{
+					return false;
+				}
+				return true;
+			}
+		}
+		return false;
+	}
+
 	// Reset all the light to default state "not being affected by atmosphere". Should only be called from render side.
 	void ResetAtmosphereLightsProperties();
 
diff --git a/Engine/Source/Runtime/Renderer/Private/VolumetricCloudRendering.cpp b/Engine/Source/Runtime/Renderer/Private/VolumetricCloudRendering.cpp
--- a/Engine/Source/Runtime/Renderer/Private/VolumetricCloudRendering.cpp
+++ b/Engine/Source/Runtime/Renderer/Private/VolumetricCloudRendering.cpp
@@ -2729,7 +2729,7 @@ bool FSceneRenderer::RenderVolumetricCloud(
 			CloudRC.CloudInfo = &CloudInfo;
 			CloudRC.CloudVolumeMaterialProxy= CloudVolumeMaterialProxy;
 			CloudRC.bSkipAtmosphericLightShadowmap = !GetVolumetricCloudReceiveAtmosphericLightShadowmap(AtmosphericLight0);
-			CloudRC.bSecondAtmosphereLightEnabled = Scene->IsSecondAtmosphereLightEnabled();
+			CloudRC.bSecondAtmosphereLightEnabled = Scene->IsSecondAtmosphereLightEnabledForCloud();
 
 			const bool DebugCloudShadowMap = CVarVolumetricCloudShadowMapDebug.GetValueOnRenderThread() && ShouldRenderCloudShadowmap(AtmosphericLight0);
 			const bool DebugCloudSkyAO = CVarVolumetricCloudSkyAODebug.GetValueOnRenderThread() && ShouldRenderCloudSkyAO(SkyLight);
