Exclude RectLightAtlas from mobile platforms, as rect light source textures are not currently supported by the mobile renderer.

#Synced-CL 52285798

diff --git a/Engine/Source/Runtime/Renderer/Private/RectLightTextureManager.cpp b/Engine/Source/Runtime/Renderer/Private/RectLightTextureManager.cpp
--- a/Engine/Source/Runtime/Renderer/Private/RectLightTextureManager.cpp
+++ b/Engine/Source/Runtime/Renderer/Private/RectLightTextureManager.cpp
@@ -82,6 +82,13 @@ static TAutoConsoleVariable<int32> CVarRectLighAtlasFormat(
 	TEXT("Define the atlas format used.\n - 0: R11G11B10 (faster but can caused hue shift). \n - 1: RGA16F (more accurate but slower)"),
 	ECVF_RenderThreadSafe);
 
+static bool PlatformSupportsRectLightAtlas(FStaticShaderPlatform Platform)
+{
+	// Mobile platforms currently do not suppport rect light source texture
+	// See USE_SOURCE_TEXTURE in RectLight.ush
+	return !IsMobilePlatform(Platform);
+}
+
 namespace RectLightAtlas
 {
 
@@ -363,7 +370,7 @@ class FRectLightAtlasClearCS : public FGlobalShader
 	END_SHADER_PARAMETER_STRUCT()
 
 public:
-		static bool ShouldCompilePermutation(const FGlobalShaderPermutationParameters& Parameters) { return true; }
+		static bool ShouldCompilePermutation(const FGlobalShaderPermutationParameters& Parameters) { return PlatformSupportsRectLightAtlas(Parameters.Platform); }
 	static void ModifyCompilationEnvironment(const FGlobalShaderPermutationParameters& Parameters, FShaderCompilerEnvironment& OutEnvironment)
 	{
 		FGlobalShader::ModifyCompilationEnvironment(Parameters, OutEnvironment);
@@ -432,7 +439,7 @@ class FRectLightAtlasDebugInfoCS : public FGlobalShader
 	END_SHADER_PARAMETER_STRUCT()
 
 public:
-	static bool ShouldCompilePermutation(const FGlobalShaderPermutationParameters& Parameters) { return ShaderPrint::IsSupported(Parameters.Platform); }
+	static bool ShouldCompilePermutation(const FGlobalShaderPermutationParameters& Parameters) { return ShaderPrint::IsSupported(Parameters.Platform) && PlatformSupportsRectLightAtlas(Parameters.Platform); }
 	static void ModifyCompilationEnvironment(const FGlobalShaderPermutationParameters& Parameters, FShaderCompilerEnvironment& OutEnvironment)
 	{
 		FGlobalShader::ModifyCompilationEnvironment(Parameters, OutEnvironment);
@@ -544,7 +551,7 @@ class FRectLightAtlasVS : public FGlobalShader
 	END_SHADER_PARAMETER_STRUCT()
 
 public:
-	static bool ShouldCompilePermutation(const FGlobalShaderPermutationParameters& Parameters) { return true; }
+	static bool ShouldCompilePermutation(const FGlobalShaderPermutationParameters& Parameters) { return PlatformSupportsRectLightAtlas(Parameters.Platform); }
 	static void ModifyCompilationEnvironment(const FGlobalShaderPermutationParameters& Parameters, FShaderCompilerEnvironment& OutEnvironment)
 	{
 		FGlobalShader::ModifyCompilationEnvironment(Parameters, OutEnvironment);
@@ -586,7 +593,7 @@ class FRectAtlasAddTexturePS : public FGlobalShader
 	END_SHADER_PARAMETER_STRUCT()
 
 public:
-	static bool ShouldCompilePermutation(const FGlobalShaderPermutationParameters& Parameters) { return true; }
+	static bool ShouldCompilePermutation(const FGlobalShaderPermutationParameters& Parameters) { return PlatformSupportsRectLightAtlas(Parameters.Platform); }
 	static void ModifyCompilationEnvironment(const FGlobalShaderPermutationParameters& Parameters, FShaderCompilerEnvironment& OutEnvironment)
 	{
 		FGlobalShader::ModifyCompilationEnvironment(Parameters, OutEnvironment);
@@ -720,7 +727,7 @@ class FRectAtlasCopyTexturePS : public FGlobalShader
 	END_SHADER_PARAMETER_STRUCT()
 
 public:
-	static bool ShouldCompilePermutation(const FGlobalShaderPermutationParameters& Parameters) { return true; }
+	static bool ShouldCompilePermutation(const FGlobalShaderPermutationParameters& Parameters) { return PlatformSupportsRectLightAtlas(Parameters.Platform); }
 	static void ModifyCompilationEnvironment(const FGlobalShaderPermutationParameters& Parameters, FShaderCompilerEnvironment& OutEnvironment)
 	{
 		FGlobalShader::ModifyCompilationEnvironment(Parameters, OutEnvironment);
@@ -844,7 +851,7 @@ class FRectAtlasFilterTexturePS : public FGlobalShader
 	END_SHADER_PARAMETER_STRUCT()
 
 public:
-	static bool ShouldCompilePermutation(const FGlobalShaderPermutationParameters& Parameters) { return true; }
+	static bool ShouldCompilePermutation(const FGlobalShaderPermutationParameters& Parameters) { return PlatformSupportsRectLightAtlas(Parameters.Platform); }
 	static void ModifyCompilationEnvironment(const FGlobalShaderPermutationParameters& Parameters, FShaderCompilerEnvironment& OutEnvironment)
 	{
 		FGlobalShader::ModifyCompilationEnvironment(Parameters, OutEnvironment);
@@ -1612,12 +1619,12 @@ static void PackAtlas(
 ///////////////////////////////////////////////////////////////////////////////////////////////////
 // API
 
-uint32 AddTexture(UTexture* In, const FVector4f& ScaleOffset)
+uint32 AddTexture(FStaticShaderPlatform Platform, UTexture* In, const FVector4f& ScaleOffset)
 {
 	check(IsInRenderingThread());
 
 	uint32 SlotIndex = InvalidSlotIndex;
-	if (In)
+	if (In && PlatformSupportsRectLightAtlas(Platform))
 	{
 		// 1. Find if the texture already exist (simple linear search assuming a low number of textures)
 		bool bFound = false;
@@ -1667,6 +1674,11 @@ uint32 AddTexture(UTexture* In, const FVector4f& ScaleOffset)
 	return SlotIndex;
 }
 
+uint32 AddTexture(UTexture* In, const FVector4f& ScaleOffset)
+{
+	return AddTexture(GMaxRHIShaderPlatform, In, ScaleOffset);
+}
+
 void RemoveTexture(uint32 InSlotIndex)
 {
 	check(IsInRenderingThread());
@@ -1733,7 +1745,7 @@ static FRDGTextureRef CreateRectLightAtlasTexture(FRDGBuilder& GraphBuilder, con
 
 void UpdateAtlasTexture(FRDGBuilder& GraphBuilder, const ERHIFeatureLevel::Type FeatureLevel)
 {
-	if (GRectLightTextureManager.bLock)
+	if (GRectLightTextureManager.bLock || !PlatformSupportsRectLightAtlas(GetFeatureLevelShaderPlatform(FeatureLevel)))
 	{
 		return;
 	}
@@ -1930,7 +1942,9 @@ void UpdateAtlasTexture(FRDGBuilder& GraphBuilder, const ERHIFeatureLevel::Type
 
 void AddDebugPass(FRDGBuilder& GraphBuilder, const FViewInfo& View, FRDGTextureRef OutputTexture)
 {
-	if (CVarRectLighTextureDebug.GetValueOnRenderThread() > 0 && ShaderPrint::IsSupported(View.Family->GetShaderPlatform()))
+	if (CVarRectLighTextureDebug.GetValueOnRenderThread() > 0 
+		&& ShaderPrint::IsSupported(View.Family->GetShaderPlatform()) 
+		&& PlatformSupportsRectLightAtlas(View.Family->GetShaderPlatform()))
 	{
 		if (FRDGTextureRef DebugOutput = AddRectLightDebugInfoPass(GraphBuilder, View, OutputTexture->Desc))
 		{
diff --git a/Engine/Source/Runtime/Renderer/Private/RectLightTextureManager.h b/Engine/Source/Runtime/Renderer/Private/RectLightTextureManager.h
--- a/Engine/Source/Runtime/Renderer/Private/RectLightTextureManager.h
+++ b/Engine/Source/Runtime/Renderer/Private/RectLightTextureManager.h
@@ -20,6 +20,7 @@ struct FAtlasSlotDesc
 };
 	
 // Add a rect light source texture to the atlas
+RENDERER_API uint32 AddTexture(FStaticShaderPlatform Platform, UTexture* Texture, const FVector4f& ScaleOffset = FVector4f::One());
 RENDERER_API uint32 AddTexture(UTexture* Texture, const FVector4f& ScaleOffset = FVector4f::One());
 
 // Remove a rect light source texture from the atlas
diff --git a/Engine/Source/Runtime/Renderer/Private/RendererScene.cpp b/Engine/Source/Runtime/Renderer/Private/RendererScene.cpp
--- a/Engine/Source/Runtime/Renderer/Private/RendererScene.cpp
+++ b/Engine/Source/Runtime/Renderer/Private/RendererScene.cpp
@@ -2303,7 +2303,7 @@ void FScene::AddLightSceneInfo_RenderThread(FLightSceneInfo* LightSceneInfo)
 	if (LightType == LightType_Rect)
 	{
 		FRectLightSceneProxy* RectProxy = (FRectLightSceneProxy*)LightSceneInfo->Proxy;
-		RectProxy->RectAtlasId = RectLightAtlas::AddTexture(RectProxy->SourceTexture, RectProxy->SourceTextureScaleOffset); 
+		RectProxy->RectAtlasId = RectLightAtlas::AddTexture(ShaderPlatform, RectProxy->SourceTexture, RectProxy->SourceTextureScaleOffset);
 	}
 
 	// Register IES texture
