Removed r.Mobile.SupportsSM5MaterialNodes from volumetric fog rendering and eliminated the CVar
SM5 capabilities for volumetric fog on mobile are now conditionally enabled based on explicit graphics feature support
Also fixed volumetric fog in the multi-pass deferred path, ensuring correct behavior when full SceneDepth is required
#Synced-CL 52334746

diff --git a/Engine/Source/Runtime/Core/Private/HAL/ConsoleManager.cpp b/Engine/Source/Runtime/Core/Private/HAL/ConsoleManager.cpp
--- a/Engine/Source/Runtime/Core/Private/HAL/ConsoleManager.cpp
+++ b/Engine/Source/Runtime/Core/Private/HAL/ConsoleManager.cpp
@@ -4398,12 +4398,6 @@ static TAutoConsoleVariable<int32> CVarMobileDesiredResY(
 	TEXT("Desired mobile Y resolution (shortest axis) (non-zero == use for Y, calculate X to retain aspect ratio)"),
 	ECVF_Default);
 
-static TAutoConsoleVariable<int32> CVarMobileSupportsSM5MaterialNodes(
-	TEXT("r.Mobile.SupportsSM5MaterialNodes"),
-	0,
-	TEXT("Whether SM5 material nodes can be utilized or not by the mobile renderer on a specific platform (disabled by default)."),
-	ECVF_ReadOnly);
-
 static TAutoConsoleVariable<int32> CVarLWCTruncateMode(
 	TEXT("r.MaterialEditor.LWCTruncateMode"),
 	2,
diff --git a/Engine/Source/Runtime/RenderCore/Private/RenderUtils.cpp b/Engine/Source/Runtime/RenderCore/Private/RenderUtils.cpp
--- a/Engine/Source/Runtime/RenderCore/Private/RenderUtils.cpp
+++ b/Engine/Source/Runtime/RenderCore/Private/RenderUtils.cpp
@@ -678,12 +678,6 @@ RENDERCORE_API bool MobileUsesFullDepthPrepass(const FStaticShaderPlatform Platf
 	return bForcedDepthPrepass;
 }
 
-RENDERCORE_API bool MobileSupportsSM5MaterialNodes(const FStaticShaderPlatform Platform)
-{
-	static FShaderPlatformCachedIniValue<bool> MobileSupportsSM5MaterialNodesIniValue(TEXT("r.Mobile.SupportsSM5MaterialNodes"));
-	return IsMobilePlatform(Platform) && (MobileSupportsSM5MaterialNodesIniValue.Get(Platform) != 0);
-}
-
 RENDERCORE_API bool ShouldForceFullDepthPass(const FStaticShaderPlatform Platform)
 {
 	if (IsMobilePlatform(Platform))
diff --git a/Engine/Source/Runtime/RenderCore/Private/Shader.cpp b/Engine/Source/Runtime/RenderCore/Private/Shader.cpp
--- a/Engine/Source/Runtime/RenderCore/Private/Shader.cpp
+++ b/Engine/Source/Runtime/RenderCore/Private/Shader.cpp
@@ -2632,12 +2632,6 @@ void ShaderMapAppendKey(EShaderPlatform Platform, FShaderKeyGenerator& KeyGen)
 					break;
 			}
 		}
-
-		if (MobileSupportsSM5MaterialNodes(Platform))
-		{
-			KeyGen.AppendSeparator();
-			KeyGen.Append(TEXT("MATMOBSM5"));
-		}
 	}
 
 	{
diff --git a/Engine/Source/Runtime/RenderCore/Public/RenderUtils.h b/Engine/Source/Runtime/RenderCore/Public/RenderUtils.h
--- a/Engine/Source/Runtime/RenderCore/Public/RenderUtils.h
+++ b/Engine/Source/Runtime/RenderCore/Public/RenderUtils.h
@@ -328,8 +328,6 @@ RENDERCORE_API bool MobileBasePassAlwaysUsesCSM(const FStaticShaderPlatform Plat
 
 RENDERCORE_API bool MobileUsesFullDepthPrepass(const FStaticShaderPlatform Platform);
 
-RENDERCORE_API bool MobileSupportsSM5MaterialNodes(const FStaticShaderPlatform Platform);
-
 RENDERCORE_API bool ShouldForceFullDepthPass(const FStaticShaderPlatform Platform);
 
 RENDERCORE_API bool SupportsGen4TAA(const FStaticShaderPlatform Platform);
diff --git a/Engine/Source/Runtime/Renderer/Private/GenerateConservativeDepthBuffer.cpp b/Engine/Source/Runtime/Renderer/Private/GenerateConservativeDepthBuffer.cpp
--- a/Engine/Source/Runtime/Renderer/Private/GenerateConservativeDepthBuffer.cpp
+++ b/Engine/Source/Runtime/Renderer/Private/GenerateConservativeDepthBuffer.cpp
@@ -24,7 +24,6 @@ class FGenerateConservativeDepthBufferCS : public FGlobalShader
 	END_SHADER_PARAMETER_STRUCT()
 
 public:
-	static bool ShouldCompilePermutation(const FGlobalShaderPermutationParameters& Parameters) { return IsFeatureLevelSupported(Parameters.Platform, ERHIFeatureLevel::SM5); }
 	static void ModifyCompilationEnvironment(const FGlobalShaderPermutationParameters& Parameters, FShaderCompilerEnvironment& OutEnvironment)
 	{
 		FGlobalShader::ModifyCompilationEnvironment(Parameters, OutEnvironment);
diff --git a/Engine/Source/Runtime/Renderer/Private/MobileFogRendering.cpp b/Engine/Source/Runtime/Renderer/Private/MobileFogRendering.cpp
--- a/Engine/Source/Runtime/Renderer/Private/MobileFogRendering.cpp
+++ b/Engine/Source/Runtime/Renderer/Private/MobileFogRendering.cpp
@@ -192,7 +192,6 @@ void FMobileSceneRenderer::RenderFog(FRHICommandList& RHICmdList, const FViewInf
 	const bool bUseFogSecondTerm = !bSimpleFog && (View.ExponentialFogParameters2.Z > 0); // Fog Density in Z
 	const bool bUseFogDirectionalInscatering = !bSimpleFog && !bUseFogInscatteringColorCubemap && (View.DirectionalInscatteringColor.GetLuminance() > 0 || View.bUseDirectionalInscattering);
 	
-	const bool bShouldRenderVolumetricFog = ShouldRenderVolumetricFog();
 	bFogHasComposedLocalFogVolumes = bShouldRenderVolumetricFog && ShouldRenderLocalFogVolume(Scene, *View.Family) && ShouldRenderLocalFogVolumeInVolumetricFog(Scene, *View.Family, bShouldRenderVolumetricFog);
 	const bool bComposeHalfResLFV = View.LocalFogVolumeViewData.bUseHalfResLocalFogVolume;
 	bFogHasComposedLocalFogVolumes |= bComposeHalfResLFV;
diff --git a/Engine/Source/Runtime/Renderer/Private/MobileShadingRenderer.cpp b/Engine/Source/Runtime/Renderer/Private/MobileShadingRenderer.cpp
--- a/Engine/Source/Runtime/Renderer/Private/MobileShadingRenderer.cpp
+++ b/Engine/Source/Runtime/Renderer/Private/MobileShadingRenderer.cpp
@@ -392,6 +392,7 @@ FMobileSceneRenderer::FMobileSceneRenderer(const FSceneViewFamily* InViewFamily,
 	bShouldRenderCustomDepth = false;
 	bRequiresShadowProjections = false;
 	bEnableDistanceFieldAO = false;
+	bShouldRenderVolumetricFog = false;
 	bIsFullDepthPrepassEnabled = (Scene->EarlyZPassMode == DDM_AllOpaque || Scene->EarlyZPassMode == DDM_AllOpaqueNoVelocity);
 	bIsMaskedOnlyDepthPrepassEnabled = Scene->EarlyZPassMode == DDM_MaskedOnly;
 	bEnableClusteredLocalLights = MobileForwardEnableLocalLights(ShaderPlatform);
@@ -580,7 +581,9 @@ void FMobileSceneRenderer::InitViews(
 	TaskDatas.VisibilityTaskData->ProcessRenderThreadTasks();
 	TaskDatas.VisibilityTaskData->FinishGatherDynamicMeshElements(BasePassDepthStencilAccess, InstanceCullingManager, VirtualTextureUpdater);
 
-	if (ShouldRenderVolumetricFog() && bRendererOutputFinalSceneColor)
+	// Volumetric fog requires a full depth before lighting
+	bShouldRenderVolumetricFog = ShouldRenderVolumetricFog() && (bIsFullDepthPrepassEnabled || (bDeferredShading && !MobileAllowFramebufferFetch(ShaderPlatform)));
+	if (bShouldRenderVolumetricFog && bRendererOutputFinalSceneColor)
 	{
 		SetupVolumetricFog();
 	}
@@ -1374,7 +1377,7 @@ void FMobileSceneRenderer::Render(FRDGBuilder& GraphBuilder, const FSceneRenderU
 		}
 
 		// Run local fog volume initialization before base pass and volumetric fog for all the culled instance instance data to be ready.
-		InitLocalFogVolumesForViews(Scene, Views, ViewFamily, GraphBuilder, ShouldRenderVolumetricFog(), bUseHalfResLocalFogVolume);
+		InitLocalFogVolumesForViews(Scene, Views, ViewFamily, GraphBuilder, bShouldRenderVolumetricFog, bUseHalfResLocalFogVolume);
 
 		if (bRendererOutputFinalSceneColor && Views.Num() > 0)
 		{
@@ -1678,14 +1681,8 @@ void FMobileSceneRenderer::Render(FRDGBuilder& GraphBuilder, const FSceneRenderU
 			GraphBuilder.AddDispatchHint();
 			ExternalAccessQueue.Submit(GraphBuilder);
 		}
-
-		if (ShouldRenderVolumetricFog())
-		{
-			ComputeVolumetricFog(GraphBuilder, SceneTextures);
-		}
 	}
 
-
 	if (bIsFullDepthPrepassEnabled)
 	{
 		// When renderer is in ERendererOutput::DepthPrepassOnly mode, bShouldRenderHZB is set to false in InitViews()
@@ -1694,6 +1691,11 @@ void FMobileSceneRenderer::Render(FRDGBuilder& GraphBuilder, const FSceneRenderU
 			RenderHZB(GraphBuilder, SceneTextures.Depth.Resolve);
 		}
 
+		if (bShouldRenderVolumetricFog && bRendererOutputFinalSceneColor)
+		{
+			ComputeVolumetricFog(GraphBuilder, SceneTextures);
+		}
+
 		// When renderer is in ERendererOutput::DepthPrepassOnly mode, bRequiresAmbientOcclusionPass is set to false in InitViews()
 		// Render SSAO after base pass when possible, for GBuffer access. AsyncCompute is suspposed to run with a BasePass
 		if (bRendererOutputFinalSceneColor && ((!bDeferredShading || !bRequiresMultiPass) || MobileSSAOUsesCompute()))
@@ -2876,6 +2878,11 @@ void FMobileSceneRenderer::RenderDeferredMultiPass(FRDGBuilder& GraphBuilder, FS
 		RenderHZB(GraphBuilder, SceneTextures.Depth.Target);
 	}
 
+	if (bShouldRenderVolumetricFog && !bIsFullDepthPrepassEnabled)
+	{
+		ComputeVolumetricFog(GraphBuilder, SceneTextures);
+	}
+
 	BasePassRenderTargets.Enumerate([](FRenderTargetBinding& RenderTarget) {
 		RenderTarget.SetLoadAction(ERenderTargetLoadAction::ELoad);
 		});
diff --git a/Engine/Source/Runtime/Renderer/Private/SceneRendering.h b/Engine/Source/Runtime/Renderer/Private/SceneRendering.h
--- a/Engine/Source/Runtime/Renderer/Private/SceneRendering.h
+++ b/Engine/Source/Runtime/Renderer/Private/SceneRendering.h
@@ -3083,6 +3083,7 @@ private:
 	bool bRequiresShadowProjections;
 	bool bAdrenoOcclusionMode;
 	bool bEnableDistanceFieldAO;
+	bool bShouldRenderVolumetricFog;
 	ETranslucencyPass::Type StandardTranslucencyPass;
 	EMeshPass::Type StandardTranslucencyMeshPass;
 	FDecalVisibilityTaskData* DecalVisibilityTaskData = nullptr;
diff --git a/Engine/Source/Runtime/Renderer/Private/VolumetricFog.cpp b/Engine/Source/Runtime/Renderer/Private/VolumetricFog.cpp
--- a/Engine/Source/Runtime/Renderer/Private/VolumetricFog.cpp
+++ b/Engine/Source/Runtime/Renderer/Private/VolumetricFog.cpp
@@ -244,11 +244,6 @@ static int32 GetVolumetricFogGridSizeZ()
 	return FMath::Max(1, GVolumetricFogGridSizeZ);
 }
 
-static bool IsMobileVolumetricFog(const FStaticShaderPlatform Platform)
-{
-	return IsMobilePlatform(Platform) && !MobileSupportsSM5MaterialNodes(Platform);
-}
-
 static FIntPoint GetVolumetricFogTextureResourceRes(const FViewInfo& View)
 {
 	// Allocate texture using scene render targets size so we do not reallocate every frame when dynamic resolution is used in order to avoid resources allocation hitches.
@@ -418,6 +413,15 @@ static bool SetupInjectShadowedLocalLightCommonParameters(
 	return true;
 }
 
+static bool PlatformSupportsShadowedLocalLights(FStaticShaderPlatform Platform)
+{
+	if (IsMobilePlatform(Platform))
+	{
+		return (IsMobileMovableSpotlightShadowsEnabled(Platform) || FReadOnlyCVARCache::EnablePointLightShadows(Platform));
+	}
+	return true;
+}
+
 /** Shader that adds direct lighting contribution from the given light to the current volume lighting cascade. */
 class FInjectShadowedLocalLightPS : public FGlobalShader
 {
@@ -459,7 +463,13 @@ class FInjectShadowedLocalLightPS : public FGlobalShader
 	
 	static bool ShouldCompilePermutation(const FGlobalShaderPermutationParameters& Parameters)
 	{
-		return !IsMobileVolumetricFog(Parameters.Platform);
+		FPermutationDomain PermutationVector(Parameters.PermutationId);
+		if (PermutationVector.Get<FVirtualShadowMap>() && !DoesPlatformSupportVirtualShadowMaps(Parameters.Platform))
+		{
+			return false;
+		}
+
+		return PlatformSupportsShadowedLocalLights(Parameters.Platform);
 	}
 
 	static EShaderPermutationPrecacheRequest ShouldPrecachePermutation(const FGlobalShaderPermutationParameters& Parameters)
@@ -1206,7 +1216,7 @@ class FVolumetricFogLightScatteringCS : public FGlobalShader
 			PermutationVector.Set<FMegaLights>(false);
 		}
 
-		if (PermutationVector.Get<FLumenGI>())
+		if (PermutationVector.Get<FLumenGI>() || !DoesPlatformSupportDistanceFieldAO(ShaderPlatform))
 		{
 			PermutationVector.Set<FDistanceFieldSkyOcclusion>(false);
 		}
@@ -1219,25 +1229,16 @@ class FVolumetricFogLightScatteringCS : public FGlobalShader
 		if (PermutationVector.Get< FUbershader>())
 		{
 			PermutationVector.Set<FTemporalReprojection>(true);
-			PermutationVector.Set<FDistanceFieldSkyOcclusion>(true);
+			PermutationVector.Set<FDistanceFieldSkyOcclusion>(DoesPlatformSupportDistanceFieldAO(ShaderPlatform));
 			PermutationVector.Set<FLumenGI>(FDataDrivenShaderPlatformInfo::GetSupportsLumenGI(ShaderPlatform));
 			PermutationVector.Set<FVirtualShadowMap>(DoesPlatformSupportVirtualShadowMaps(ShaderPlatform));
-			PermutationVector.Set<FRaytracedShadowsVolume>(true);
+			PermutationVector.Set<FRaytracedShadowsVolume>(ShouldCompileRayTracingShadersForProject(ShaderPlatform));
 			PermutationVector.Set<FSampleLightFunctionAtlas>(true);
 			PermutationVector.Set<FMegaLights>(MegaLights::ShouldCompileShaders(ShaderPlatform));
 			PermutationVector.Set<FLightSoftFading>(true);
 			PermutationVector.Set<FSuperSampling>(true);
 		}
 
-		if (IsMobileVolumetricFog(FStaticShaderPlatform(ShaderPlatform)))
-		{
-			PermutationVector.Set<FDistanceFieldSkyOcclusion>(false);
-			PermutationVector.Set<FTemporalReprojection>(false);
-			PermutationVector.Set<FSampleLightFunctionAtlas>(false);
-			PermutationVector.Set<FMegaLights>(false);
-			PermutationVector.Set<FRaytracedShadowsVolume>(false);
-			PermutationVector.Set<FVirtualShadowMap>(false);
-		}
 		return PermutationVector;
 	}
 
@@ -1252,32 +1253,16 @@ class FVolumetricFogLightScatteringCS : public FGlobalShader
 		FPermutationDomain PermutationVector(Parameters.PermutationId);
 		if (PermutationVector.Get<FUbershader>())
 		{
-			if (IsMobileVolumetricFog(FStaticShaderPlatform(Parameters.Platform)))
-			{
-				return 
-					GetIfSupported<FLumenGI>(PermutationVector, FDataDrivenShaderPlatformInfo::GetSupportsLumenGI(Parameters.Platform)) &&
-					PermutationVector.Get<FTemporalReprojection>()==0 &&
-					PermutationVector.Get<FDistanceFieldSkyOcclusion>()==0 &&
-					PermutationVector.Get<FVirtualShadowMap>()==0 &&
-					PermutationVector.Get<FRaytracedShadowsVolume>()==0 &&
-					PermutationVector.Get<FSampleLightFunctionAtlas>()==0 &&
-					PermutationVector.Get<FMegaLights>()==0 &&
-					PermutationVector.Get<FLightSoftFading>() &&
-					PermutationVector.Get<FSuperSampling>();
-			}
-			else
-			{
-				return 
-					GetIfSupported<FLumenGI>(PermutationVector, FDataDrivenShaderPlatformInfo::GetSupportsLumenGI(Parameters.Platform)) &&
-					GetIfSupported<FMegaLights>(PermutationVector, MegaLights::ShouldCompileShaders(Parameters.Platform)) &&
-					GetIfSupported<FVirtualShadowMap>(PermutationVector, DoesPlatformSupportVirtualShadowMaps(Parameters.Platform)) &&
-					PermutationVector.Get<FRaytracedShadowsVolume>() &&
-					PermutationVector.Get<FTemporalReprojection>() &&
-					PermutationVector.Get<FDistanceFieldSkyOcclusion>() &&
-					PermutationVector.Get<FSampleLightFunctionAtlas>() &&
-					PermutationVector.Get<FLightSoftFading>() &&
-					PermutationVector.Get<FSuperSampling>();
-			}
+			return 
+				GetIfSupported<FLumenGI>(PermutationVector, FDataDrivenShaderPlatformInfo::GetSupportsLumenGI(Parameters.Platform)) &&
+				GetIfSupported<FMegaLights>(PermutationVector, MegaLights::ShouldCompileShaders(Parameters.Platform)) &&
+				GetIfSupported<FVirtualShadowMap>(PermutationVector, DoesPlatformSupportVirtualShadowMaps(Parameters.Platform)) &&
+				GetIfSupported<FDistanceFieldSkyOcclusion>(PermutationVector, DoesPlatformSupportDistanceFieldAO(Parameters.Platform)) &&
+				GetIfSupported<FRaytracedShadowsVolume>(PermutationVector, ShouldCompileRayTracingShadersForProject(Parameters.Platform)) &&
+				PermutationVector.Get<FTemporalReprojection>() &&
+				PermutationVector.Get<FSampleLightFunctionAtlas>() &&
+				PermutationVector.Get<FLightSoftFading>() &&
+				PermutationVector.Get<FSuperSampling>();
 		}
 		
 		if (RemapPermutation(PermutationVector, Parameters.Platform) != PermutationVector)
@@ -1357,10 +1342,6 @@ class FVolumetricFogFinalIntegrationCS : public FGlobalShader
 
 IMPLEMENT_GLOBAL_SHADER(FVolumetricFogFinalIntegrationCS, "/Engine/Private/VolumetricFog.usf", "FinalIntegrationCS", SF_Compute);
 
-bool DoesPlatformSupportVolumetricFogVoxelization(const FStaticShaderPlatform Platform)
-{
-	return !IsMobileVolumetricFog(Platform);
-}
 bool ShouldRenderVolumetricFog(const FScene* Scene, const FSceneViewFamily& ViewFamily)
 {
 	return ShouldRenderFog(ViewFamily)
@@ -1512,14 +1493,7 @@ void FSceneRenderer::ComputeVolumetricFog(FRDGBuilder& GraphBuilder,
 	{
 		return;
 	}
-	const bool bIsMobilePlatform = IsMobileVolumetricFog(Scene->GetShaderPlatform());
-	// FVolumetricFogLightScatteringCS::FTemporalReprojection permutation is disabled on mobile
-	if (bIsMobilePlatform)
-	{
-		GVolumetricFogJitter = 0;
-		GVolumetricFogTemporalReprojection = 0;
-		GVolumetricFogInjectShadowedLightsSeparately = 0;
-	}
+	
 	const FExponentialHeightFogSceneInfo& FogInfo = Scene->ExponentialFogs[0];
 
 	TRACE_CPUPROFILER_EVENT_SCOPE(FSceneRenderer::ComputeVolumetricFog);
@@ -1629,9 +1603,7 @@ void FSceneRenderer::ComputeVolumetricFog(FRDGBuilder& GraphBuilder,
 			IntegrationData.FrameJitterOffsetValues[FrameOffsetIndex] = VolumetricFogTemporalRandom(View.Family->FrameNumber - FrameOffsetIndex);
 		}
 
-
-		// Mobile has limited capacities with SRV binding so do not enable atlas sampling on there.
-		const bool bUseLightFunctionAtlasEnabledAndSupported = LightFunctionAtlas::IsEnabled(*Scene, ELightFunctionAtlasSystem::VolumetricFog) && !bIsMobilePlatform;
+		const bool bUseLightFunctionAtlasEnabledAndSupported = LightFunctionAtlas::IsEnabled(*Scene, ELightFunctionAtlasSystem::VolumetricFog);
 
 		const bool bUseTemporalReprojection =
 			GVolumetricFogTemporalReprojection
@@ -1916,7 +1888,7 @@ void FSceneRenderer::ComputeVolumetricFog(FRDGBuilder& GraphBuilder,
 			float CloudShadowmap_FarDepthKm = 0.0f;
 			float CloudShadowmap_Strength = 0.0f;
 			FRDGTexture* CloudShadowmap_RDGTexture = BlackDummyTexture;
-			if (CloudInfo && AtmosphericDirectionalLightIndex >= 0 && AtmosphereLightProxy && !bIsMobilePlatform)
+			if (CloudInfo && AtmosphericDirectionalLightIndex >= 0 && AtmosphereLightProxy)
 			{
 				CloudShadowmap_RDGTexture = GraphBuilder.RegisterExternalTexture(View.VolumetricCloudShadowExtractedRenderTarget[AtmosphericDirectionalLightIndex]);
 				CloudWorldToLightClipShadowMatrix = CloudInfo->GetVolumetricCloudCommonShaderParameters().CloudShadowmapTranslatedWorldToLightClipMatrix[AtmosphericDirectionalLightIndex];
@@ -1956,8 +1928,7 @@ void FSceneRenderer::ComputeVolumetricFog(FRDGBuilder& GraphBuilder,
 				&& SupportsDistanceFieldAO(View.GetFeatureLevel(), View.GetShaderPlatform())
 				&& bUseGlobalDistanceField
 				&& Views.Num() == 1
-				&& View.IsPerspectiveProjection()
-				&& !IsMobileVolumetricFog(View.GetShaderPlatform());
+				&& View.IsPerspectiveProjection();
 
 			const int32 SuperSampleCount = FVolumetricFogLightScatteringCS::GetSuperSampleCount(GVolumetricFogHistoryMissSupersampleCount);
 			check(PassParameters->VolumetricFogParameters.HistoryMissSuperSampleCount == SuperSampleCount);
diff --git a/Engine/Source/Runtime/Renderer/Private/VolumetricFog.h b/Engine/Source/Runtime/Renderer/Private/VolumetricFog.h
--- a/Engine/Source/Runtime/Renderer/Private/VolumetricFog.h
+++ b/Engine/Source/Runtime/Renderer/Private/VolumetricFog.h
@@ -21,7 +21,6 @@ extern FVector2f GetVolumetricFogFroxelToScreenSVPosRatio(const FViewInfo& View)
 
 extern FRDGTextureDesc GetVolumetricFogRDGTextureDesc(const FIntVector& VolumetricFogResourceGridSize);
 
-extern bool DoesPlatformSupportVolumetricFogVoxelization(const FStaticShaderPlatform Platform);
 extern bool ShouldRenderVolumetricFog(const FScene* Scene, const FSceneViewFamily& ViewFamily);
 extern const FProjectedShadowInfo* GetShadowForInjectionIntoVolumetricFog(const FVisibleLightInfo& VisibleLightInfo);
 
diff --git a/Engine/Source/Runtime/Renderer/Private/VolumetricFogVoxelization.cpp b/Engine/Source/Runtime/Renderer/Private/VolumetricFogVoxelization.cpp
--- a/Engine/Source/Runtime/Renderer/Private/VolumetricFogVoxelization.cpp
+++ b/Engine/Source/Runtime/Renderer/Private/VolumetricFogVoxelization.cpp
@@ -33,6 +33,11 @@ FAutoConsoleVariableRef CVarVolumetricFogVoxelizationShowOnlyPassIndex(
 	ECVF_RenderThreadSafe
 	);
 
+static bool DoesPlatformSupportVolumetricFogVoxelization(const FStaticShaderPlatform Platform)
+{
+	return RHISupportsGeometryShaders(Platform) || RHISupportsVertexShaderLayer(Platform);
+}
+
 static FORCEINLINE int32 GetVoxelizationSlicesPerPass(EShaderPlatform Platform)
 {
 	return RHISupportsGeometryShaders(Platform) ? GVolumetricFogVoxelizationSlicesPerGSPass : 1;
