VulkanRHI: disable r.LightGridInjection.UseSpecializationConstants until issues are resolved (UE-372026, UE-371934)

#rnx
#Synced-CL 52347521

diff --git a/Engine/Source/Runtime/Renderer/Private/LightGridInjection.cpp b/Engine/Source/Runtime/Renderer/Private/LightGridInjection.cpp
--- a/Engine/Source/Runtime/Renderer/Private/LightGridInjection.cpp
+++ b/Engine/Source/Runtime/Renderer/Private/LightGridInjection.cpp
@@ -190,8 +190,8 @@ FAutoConsoleVariableRef CVarLightCullingMaxDistanceOverride(
 // Experimental version of FLightGridInjectionCS that uses specialization constants
 static TAutoConsoleVariable<int32> CVarLightGridInjectionUseSpecializationConstants(
 	TEXT("r.LightGridInjection.UseSpecializationConstants"),
-	1,
-	TEXT("Experimental: Selects whether LightGridInjection should use specialization constant. (default: 1)"),
+	0,
+	TEXT("Experimental: Selects whether LightGridInjection should use specialization constant. (default: 0)"),
 	ECVF_ReadOnly);
 
 static inline bool UseSpecializationConstants(EShaderPlatform ShaderPlatform)
