Separate ViewState->PrevFrameViewInfo setup from PrepareViewStateForVisibility so it can be performed after each view's ViewExtension->PreRenderView_RenderThread has a chance to update the camera position.

Previously, PrevFrameViewInfo was set up for all views in all renderers within the FIRST renderer's OnRenderBegin(), as part of the scene update. This meant that if a renderer other than the first had a view extension which updated the camera position (e.g. HMD late update/late latching), that change would not be reflected in PrevFrameViewInfo, causing temporal reprojection issues in Lumen and TSR.

Note that there are still other parts of the scene update which seem to be dependent on camera position/view matrices (e.g. SDF setup). These will likely also exhibit artifacts when more than one renderer is active and a view extension updates the camera position.

#Synced-CL 52356014

diff --git a/Engine/Source/Runtime/Renderer/Private/SceneRendering.cpp b/Engine/Source/Runtime/Renderer/Private/SceneRendering.cpp
--- a/Engine/Source/Runtime/Renderer/Private/SceneRendering.cpp
+++ b/Engine/Source/Runtime/Renderer/Private/SceneRendering.cpp
@@ -4312,6 +4312,9 @@ IVisibilityTaskData* FSceneRenderer::OnRenderBegin(FRDGBuilder& GraphBuilder, co
 			}
 		}
 
+		// This must be done after ViewExtension->PreRenderView_RenderThread performs any updates to the camera position
+		PrepareViewStatePrevFrameViewInfo();
+
 		if (ViewFamily.EngineShowFlags.LensDistortion && FPaniniProjectionConfig::IsEnabledByCVars())
 		{
 			const FPaniniProjectionConfig PaniniProjection = FPaniniProjectionConfig::ReadCVars();
diff --git a/Engine/Source/Runtime/Renderer/Private/SceneRendering.h b/Engine/Source/Runtime/Renderer/Private/SceneRendering.h
--- a/Engine/Source/Runtime/Renderer/Private/SceneRendering.h
+++ b/Engine/Source/Runtime/Renderer/Private/SceneRendering.h
@@ -2779,6 +2779,7 @@ private:
 	friend class FVisibilityTaskData;
 
 	void PrepareViewStateForVisibility(const FSceneTexturesConfig& SceneTexturesConfig);
+	void PrepareViewStatePrevFrameViewInfo();
 
 	bool bShadowDepthRenderCompleted;
 
diff --git a/Engine/Source/Runtime/Renderer/Private/SceneVisibility.cpp b/Engine/Source/Runtime/Renderer/Private/SceneVisibility.cpp
--- a/Engine/Source/Runtime/Renderer/Private/SceneVisibility.cpp
+++ b/Engine/Source/Runtime/Renderer/Private/SceneVisibility.cpp
@@ -5203,12 +5203,6 @@ void FSceneRenderer::PrepareViewStateForVisibility(const FSceneTexturesConfig& S
 	bool bFreezeTemporalSequences = bFreezeTemporalHistories || CVarFreezeTemporalSequences.GetValueOnRenderThread() != 0;
 #endif
 
-	// Load this field once so it has a consistent value for all views (and to avoid the atomic load in the loop).
-	// While the value may not be perfectly in sync when we render other view families, this is ok as this
-	// invalidation mechanism is only used for interactive rendering where we expect to be constantly drawing the scene.
-	// Therefore it is acceptable for some view families to be a frame or so behind others.
-	uint32 CurrentPathTracingInvalidationCounter = Scene->PathTracingInvalidationCounter.Load();
-
 	// Setup motion blur parameters (also check for camera movement thresholds)
 	for(int32 ViewIndex = 0;ViewIndex < AllViews.Num();ViewIndex++)
 	{
@@ -5510,6 +5504,25 @@ void FSceneRenderer::PrepareViewStateForVisibility(const FSceneTexturesConfig& S
 
 			View.ViewMatrices.HackAddTemporalAAProjectionJitter(FVector2D(SampleX * 2.0f / View.ViewRect.Width(), SampleY * -2.0f / View.ViewRect.Height()));
 		}
+	}
+}
+
+void FSceneRenderer::PrepareViewStatePrevFrameViewInfo()
+{
+	// Load this field once so it has a consistent value for all views (and to avoid the atomic load in the loop).
+	// While the value may not be perfectly in sync when we render other view families, this is ok as this
+	// invalidation mechanism is only used for interactive rendering where we expect to be constantly drawing the scene.
+	// Therefore it is acceptable for some view families to be a frame or so behind others.
+	uint32 CurrentPathTracingInvalidationCounter = Scene->PathTracingInvalidationCounter.Load();
+
+	for (int32 ViewIndex = 0; ViewIndex < AllViews.Num(); ViewIndex++)
+	{
+		FViewInfo& View = *AllViews[ViewIndex];
+		FSceneViewState* ViewState = View.ViewState;
+
+#if DO_CHECK || USING_CODE_ANALYSIS
+		check(View.VerifyMembersChecks());
+#endif
 
 		// Setup a new FPreviousViewInfo from current frame infos.
 		FPreviousViewInfo NewPrevViewInfo;
