Mobile: Fixed case where volumetric fog was using more than 16 samplers on mobile platforms
#Synced-CL 52374542

diff --git a/Engine/Shaders/Private/VolumetricFog.usf b/Engine/Shaders/Private/VolumetricFog.usf
--- a/Engine/Shaders/Private/VolumetricFog.usf
+++ b/Engine/Shaders/Private/VolumetricFog.usf
@@ -991,7 +991,7 @@ void LightScatteringCS(
 
 		float3 BrickTextureUVs = 0;
 
-#if ALLOW_STATIC_LIGHTING
+#if ALLOW_STATIC_LIGHTING && !SHADING_PATH_MOBILE // Mobile does not support VLM
 		if (View.SkyLightVolumetricScatteringIntensity > 0 || StaticLightingScatteringIntensity > 0)
 		{
 			BrickTextureUVs = ComputeVolumetricLightmapBrickTextureUVs(WorldPosition);
@@ -1039,7 +1039,7 @@ void LightScatteringCS(
 		}
 		#endif
 
-#if ALLOW_STATIC_LIGHTING
+#if ALLOW_STATIC_LIGHTING && !SHADING_PATH_MOBILE // Mobile does not support VLM
 		// Indirect lighting of Stationary lights and Direct + Indirect lighting of Static lights
 		if (StaticLightingScatteringIntensity > 0)
 		{
diff --git a/Engine/Source/Runtime/Renderer/Private/VolumetricFog.cpp b/Engine/Source/Runtime/Renderer/Private/VolumetricFog.cpp
--- a/Engine/Source/Runtime/Renderer/Private/VolumetricFog.cpp
+++ b/Engine/Source/Runtime/Renderer/Private/VolumetricFog.cpp
@@ -254,6 +254,12 @@ static FIntPoint GetVolumetricFogTextureResourceRes(const FViewInfo& View)
 	return BufferSize;
 }
 
+static bool PlatformSupportsVolumetricFogUbershader(FStaticShaderPlatform Platform)
+{
+	// GL does not have enough samplers for Ubershader permutations
+	return !IsOpenGLPlatform(Platform);
+}
+
 IMPLEMENT_GLOBAL_SHADER_PARAMETER_STRUCT(FVolumetricFogGlobalData, "VolumetricFog");
 
 DECLARE_GPU_STAT(VolumetricFog);
@@ -1262,7 +1268,8 @@ class FVolumetricFogLightScatteringCS : public FGlobalShader
 				PermutationVector.Get<FTemporalReprojection>() &&
 				PermutationVector.Get<FSampleLightFunctionAtlas>() &&
 				PermutationVector.Get<FLightSoftFading>() &&
-				PermutationVector.Get<FSuperSampling>();
+				PermutationVector.Get<FSuperSampling>()	&& 
+				PlatformSupportsVolumetricFogUbershader(Parameters.Platform);
 		}
 		
 		if (RemapPermutation(PermutationVector, Parameters.Platform) != PermutationVector)
@@ -1959,7 +1966,8 @@ void FSceneRenderer::ComputeVolumetricFog(FRDGBuilder& GraphBuilder,
 			auto ComputeShader = View.ShaderMap->GetShader< FVolumetricFogLightScatteringCS >(PermutationVector);
 
 			// Ubershader default parameters and feature dynamic branching
-			if (!PrecachePSO(FComputePipelineStateInitializer(ComputeShader.GetComputeShader())))
+			if (!PrecachePSO(FComputePipelineStateInitializer(ComputeShader.GetComputeShader())) &&
+				PlatformSupportsVolumetricFogUbershader(View.GetShaderPlatform()))
 			{
 				PermutationVector.Set< FVolumetricFogLightScatteringCS::FUbershader >(true);
 				PermutationVector = FVolumetricFogLightScatteringCS::RemapPermutation(PermutationVector, View.GetShaderPlatform());
