Add support for copying SceneColor in R11G11B10 when the alpha channel is not required. This is needed for SingleLayerWater, which samples RGB for refraction and benefits from a wider dynamic range.

Notes:
-Previously encoded to PF_B8G8R8A8 using pow(x / 10, 0.25), effectively limiting the HDR range to [0, 10]
-Now uses PF_FloatR11G11B10 without encoding, providing a native HDR range of [0, 65504]

Maintaining an HDR range results in consistent bloom behavior between the Mobile and Desktop renderers, as the encoded path is constrained by the bloom threshold.
On Android, even though the SingleLayerWater SceneColor copy path is not typically used, encoding is still forced since R11G11B10 is not supported on those devices.

#Synced-CL 51996830

#Synced-CL 49947857

diff --git a/Engine/Shaders/Private/MobileBasePassPixelShader.usf b/Engine/Shaders/Private/MobileBasePassPixelShader.usf
--- a/Engine/Shaders/Private/MobileBasePassPixelShader.usf
+++ b/Engine/Shaders/Private/MobileBasePassPixelShader.usf
@@ -50,8 +50,6 @@
 		#define SINGLE_LAYER_WATER_SHADING_QUALITY 0
 		// SINGLE_LAYER_WATER_SHADING_QUALITY_FULL normally uses a refraction mask texture which we don't have on mobile
 		#define SINGLE_LAYER_WATER_USE_REFRACTION_MASK 0
-		// The SceneColorCopyTexture uses an encoding. SLW needs to know about this to undo it.		
-		#define SINGLE_LAYER_WATER_ENCODED_SCENE_COLOR 1
 	#else
 		// Value must match SINGLE_LAYER_WATER_SHADING_QUALITY_MOBILE_WITH_DEPTH_BUFFER in SingleLayerWaterCommon.ush!
 		#define SINGLE_LAYER_WATER_SHADING_QUALITY 2
diff --git a/Engine/Shaders/Private/SingleLayerWaterShading.ush b/Engine/Shaders/Private/SingleLayerWaterShading.ush
--- a/Engine/Shaders/Private/SingleLayerWaterShading.ush
+++ b/Engine/Shaders/Private/SingleLayerWaterShading.ush
@@ -8,11 +8,6 @@
 #include "DistortionCommon.ush"
 #include "SingleLayerWaterCommon.ush"
 
-// The deferred mobile renderer with support for SM5 creates a scene color copy with encoding to fit the data into RGB8
-#ifndef SINGLE_LAYER_WATER_ENCODED_SCENE_COLOR
-#define SINGLE_LAYER_WATER_ENCODED_SCENE_COLOR 0
-#endif
-
 #if SINGLE_LAYER_WATER_SHADING_QUALITY != SINGLE_LAYER_WATER_SHADING_QUALITY_FULL && SINGLE_LAYER_WATER_SHADING_QUALITY != SINGLE_LAYER_WATER_SHADING_QUALITY_MOBILE_WITH_DEPTH_TEXTURE && SINGLE_LAYER_WATER_SHADING_QUALITY != SINGLE_LAYER_WATER_SHADING_QUALITY_MOBILE_WITH_DEPTH_BUFFER
 #error SINGLE_LAYER_WATER_SHADING_QUALITY has an invalid value!
 #endif
@@ -34,10 +29,7 @@ float3 SampleSceneColorWithoutWaterBilateral(Texture2D Tex, SamplerState TexSamp
 
 	float3 BehindWaterSceneLuminance = float3(dot(Red, Weights), dot(Green, Weights), dot(Blue, Weights));
 
-#if SINGLE_LAYER_WATER_ENCODED_SCENE_COLOR
-	// Undo the function in EncodeSceneColorForMaterialNode
-	BehindWaterSceneLuminance = pow(BehindWaterSceneLuminance.rgb, 4) * 10;
-#endif
+
 
 	return BehindWaterSceneLuminance;
 }
diff --git a/Engine/Shaders/Private/TranslucentLightingShaders.usf b/Engine/Shaders/Private/TranslucentLightingShaders.usf
--- a/Engine/Shaders/Private/TranslucentLightingShaders.usf
+++ b/Engine/Shaders/Private/TranslucentLightingShaders.usf
@@ -59,7 +59,11 @@ Texture2D SceneColorTexture;
 void CopySceneColorMain(FScreenVertexOutput Input, out float4 OutColor : SV_Target0)
 {
 	float4 LinearColor = Texture2DSample(SceneColorTexture, SceneColorTextureSampler, Input.UV);
+#if SCENE_COLOR_COPY_ENCODED
 	OutColor = float4(EncodeSceneColorForMaterialNode(LinearColor.rgb), LinearColor.a);
+#else
+	OutColor = LinearColor;
+#endif
 }
 
 #if FEATURE_LEVEL >= FEATURE_LEVEL_SM5
diff --git a/Engine/Source/Runtime/Renderer/Private/MobileShadingRenderer.cpp b/Engine/Source/Runtime/Renderer/Private/MobileShadingRenderer.cpp
--- a/Engine/Source/Runtime/Renderer/Private/MobileShadingRenderer.cpp
+++ b/Engine/Source/Runtime/Renderer/Private/MobileShadingRenderer.cpp
@@ -3065,7 +3065,7 @@ void FMobileSceneRenderer::RenderDeferredMultiPass(FRDGBuilder& GraphBuilder, FS
 			// No need to copy SceneColor before translucencies when underwater
 			if (View.TranslucentPrimCount.UseSceneColorCopy(StandardTranslucencyPass) && (!bRenderSingleLayerWater || !bIsCameraUnderWater))
 			{
-				FRDGTextureRef SceneColorCopyTexture = AddCopySceneColorPass(GraphBuilder, Views, SceneTextures.Color, false, false);
+				FRDGTextureRef SceneColorCopyTexture = AddCopySceneColorPass(GraphBuilder, Views, SceneTextures.Color, false, false, true);
 				if (SceneColorCopyTexture != nullptr)
 				{
 					BasePassParameters->SceneColorCopyTexture = SceneColorCopyTexture;
@@ -3105,7 +3105,7 @@ void FMobileSceneRenderer::RenderDeferredMultiPass(FRDGBuilder& GraphBuilder, FS
 			FMobileBasePassUniformParameters* SLWBasePassParameters = GraphBuilder.AllocParameters<FMobileBasePassUniformParameters>();
 			SetupMobileBasePassUniformParameters(GraphBuilder, View, EMobileBasePass::Translucent, SetupMode, *SLWBasePassParameters, BasePassTextures);
 			
-			FRDGTextureRef SceneColorCopyTexture = AddCopySceneColorPass(GraphBuilder, Views, SceneTextures.Color, false, false);
+			FRDGTextureRef SceneColorCopyTexture = AddCopySceneColorPass(GraphBuilder, Views, SceneTextures.Color, false, false, false);
 			if (SceneColorCopyTexture != nullptr)
 			{
 				SLWBasePassParameters->SceneColorCopyTexture = SceneColorCopyTexture;
diff --git a/Engine/Source/Runtime/Renderer/Private/TranslucentRendering.cpp b/Engine/Source/Runtime/Renderer/Private/TranslucentRendering.cpp
--- a/Engine/Source/Runtime/Renderer/Private/TranslucentRendering.cpp
+++ b/Engine/Source/Runtime/Renderer/Private/TranslucentRendering.cpp
@@ -363,6 +363,9 @@ public:
 	DECLARE_GLOBAL_SHADER(FCopySceneColorPS);
 	SHADER_USE_PARAMETER_STRUCT(FCopySceneColorPS, FGlobalShader);
 
+	class FEncodeSceneColorCopy : SHADER_PERMUTATION_BOOL("SCENE_COLOR_COPY_ENCODED");
+	using FPermutationDomain = TShaderPermutationDomain<FEncodeSceneColorCopy>;
+
 	BEGIN_SHADER_PARAMETER_STRUCT(FParameters, )
 		SHADER_PARAMETER_STRUCT_REF(FViewUniformShaderParameters, View)
 		SHADER_PARAMETER_RDG_TEXTURE(Texture2D, SceneColorTexture)
@@ -373,7 +376,7 @@ public:
 
 IMPLEMENT_GLOBAL_SHADER(FCopySceneColorPS, "/Engine/Private/TranslucentLightingShaders.usf", "CopySceneColorMain", SF_Pixel);
 
-FRDGTextureRef AddCopySceneColorPass(FRDGBuilder& GraphBuilder, TArrayView<const FViewInfo> Views, const FRDGTextureMSAA& SceneColor, bool WithAlpha, bool bSkipIfUnderwater)
+FRDGTextureRef AddCopySceneColorPass(FRDGBuilder& GraphBuilder, TArrayView<const FViewInfo> Views, FRDGTextureMSAA SceneColor, bool WithAlpha, bool bSkipIfUnderwater, bool bEncodeSceneColorCopy)
 {
 	FRDGTextureRef SceneColorCopyTexture = nullptr;
 	ERenderTargetLoadAction LoadAction = ERenderTargetLoadAction::ENoAction;
@@ -411,13 +414,17 @@ FRDGTextureRef AddCopySceneColorPass(FRDGBuilder& GraphBuilder, TArrayView<const
 
 			if (!SceneColorCopyTexture)
 			{
-				SceneColorCopyTexture = GraphBuilder.CreateTexture(FRDGTextureDesc::Create2D(SceneColorExtent, PF_B8G8R8A8, FClearValueBinding::White, TexCreate_ShaderResource | TexCreate_RenderTargetable), TEXT("SceneColorCopy"));
+				EPixelFormat CopyFormat = (bEncodeSceneColorCopy || WithAlpha) ? PF_B8G8R8A8 : PF_FloatR11G11B10;
+				SceneColorCopyTexture = GraphBuilder.CreateTexture(FRDGTextureDesc::Create2D(SceneColorExtent, CopyFormat, FClearValueBinding::White, TexCreate_ShaderResource | TexCreate_RenderTargetable), TEXT("SceneColorCopy"));
 			}
 
 			const FScreenPassTextureViewport Viewport(SceneColorCopyTexture, View.ViewRect);
 
 			TShaderMapRef<FScreenVS> VertexShader(View.ShaderMap);
-			TShaderMapRef<FCopySceneColorPS> PixelShader(View.ShaderMap);
+
+			FCopySceneColorPS::FPermutationDomain PermutationVector;
+			PermutationVector.Set<FCopySceneColorPS::FEncodeSceneColorCopy>(bEncodeSceneColorCopy);
+			TShaderMapRef<FCopySceneColorPS> PixelShader(View.ShaderMap, PermutationVector);
 
 			auto* PassParameters = GraphBuilder.AllocParameters<FCopySceneColorPS::FParameters>();
 			PassParameters->View = View.ViewUniformBuffer;
@@ -498,7 +505,10 @@ static void AddCopySceneColorAlphaPass(FRDGBuilder& GraphBuilder, TArrayView<con
 		const FScreenPassTextureViewport Viewport(TargetTexture, View.ViewRect);
 
 		TShaderMapRef<FScreenVS> VertexShader(View.ShaderMap);
-		TShaderMapRef<FCopySceneColorPS> PixelShader(View.ShaderMap);
+
+		FCopySceneColorPS::FPermutationDomain PermutationVector;
+		PermutationVector.Set<FCopySceneColorPS::FEncodeSceneColorCopy>(true);
+		TShaderMapRef<FCopySceneColorPS> PixelShader(View.ShaderMap, PermutationVector);
 
 		auto* PassParameters = GraphBuilder.AllocParameters<FCopySceneColorPS::FParameters>();
 		PassParameters->View = View.ViewUniformBuffer;
@@ -1954,7 +1964,7 @@ void FDeferredShadingSceneRenderer::RenderTranslucency(
 
 	if (EnumHasAnyFlags(ViewsToRender, ETranslucencyView::AboveWater))
 	{
-		SceneColorCopyTexture = AddCopySceneColorPass(GraphBuilder, InViews, SceneTextures.Color, /*WithAlpha*/ bIsTranslucentHoldoutEnabled, /*bSkipIfUnderwater*/ true);
+		SceneColorCopyTexture = AddCopySceneColorPass(GraphBuilder, InViews, SceneTextures.Color, /*WithAlpha*/ bIsTranslucentHoldoutEnabled, /*bSkipIfUnderwater*/ true, /*bEncodeSceneColorCopy*/ true);
 	}
 
 	// Create a shared depth texture at the correct resolution.
diff --git a/Engine/Source/Runtime/Renderer/Private/TranslucentRendering.h b/Engine/Source/Runtime/Renderer/Private/TranslucentRendering.h
--- a/Engine/Source/Runtime/Renderer/Private/TranslucentRendering.h
+++ b/Engine/Source/Runtime/Renderer/Private/TranslucentRendering.h
@@ -48,7 +48,8 @@ struct FSeparateTranslucencyDimensions
 DECLARE_GPU_DRAWCALL_STAT_EXTERN(Translucency);
 
 /** Add Copy SceneColor Pass. */
-FRDGTextureRef AddCopySceneColorPass(FRDGBuilder& GraphBuilder, TArrayView<const FViewInfo> Views, const FRDGTextureMSAA& SceneColor, bool WithAlpha = false, bool bSkipIfUnderwater = true);
+FRDGTextureRef AddCopySceneColorPass(FRDGBuilder& GraphBuilder, TArrayView<const FViewInfo> Views, FRDGTextureMSAA SceneColor, bool WithAlpha = false, bool bSkipIfUnderwater = true, bool bEncodeSceneColorCopy = true);
+
 
 /** Add pass to compose holdout to scene color alpha. */
 void AddTranslucentHoldoutToSceneColorAlphaPass(
