Refcount foliage batching. Index by bIsNanite so GPU skin / nanite get their own batch. Fixed GPU skin skipping the scene extension when no sections are present on any LOD.

#Synced-CL 52384345

diff --git a/Engine/Source/Runtime/Engine/Private/SkeletalRenderGPUSkin.cpp b/Engine/Source/Runtime/Engine/Private/SkeletalRenderGPUSkin.cpp
--- a/Engine/Source/Runtime/Engine/Private/SkeletalRenderGPUSkin.cpp
+++ b/Engine/Source/Runtime/Engine/Private/SkeletalRenderGPUSkin.cpp
@@ -623,19 +623,8 @@ FSkeletalMeshObjectGPUSkin::FSkeletalMeshObjectGPUSkin(const FSkinnedMeshScenePr
 
 		FSkeletalMeshObjectLOD& LOD = LODs.Emplace_GetRef(SkeletalMeshRenderData, LODIndex, InFeatureLevel, RecreateMorphVertexBuffer, GPUSkinTechnique);
 
-		bool bHasValidSections = false;
-
-		for (const FSkelMeshRenderSection& Section : GetRenderSections(LODIndex))
-		{
-			if (Section.IsValid())
-			{
-				bHasValidSections = true;
-				break;
-			}
-		}
-
-		// Fall back to vertex factory buffer allocation mode if any of the LOD's require it or if an LOD has no valid sections in it.
-		if (!bHasValidSections || GPUSkinTechnique != ESkeletalMeshGPUSkinTechnique::Inline)
+		// Fall back to vertex factory buffer allocation mode if any of the LOD's require it.
+		if (GPUSkinTechnique != ESkeletalMeshGPUSkinTechnique::Inline)
 		{
 			BoneTransformAllocationMode = EBoneTransformAllocationMode::VertexFactory;
 		}
diff --git a/Engine/Source/Runtime/Renderer/Private/Skinning/SkinningSceneExtension.cpp b/Engine/Source/Runtime/Renderer/Private/Skinning/SkinningSceneExtension.cpp
--- a/Engine/Source/Runtime/Renderer/Private/Skinning/SkinningSceneExtension.cpp
+++ b/Engine/Source/Runtime/Renderer/Private/Skinning/SkinningSceneExtension.cpp
@@ -180,6 +180,19 @@ static void GetDefaultSkinningParameters(FSkinningSceneParameters& OutParameters
 
 IMPLEMENT_SCENE_EXTENSION(FSkinningSceneExtension);
 
+FSkeletonBatchKey FSkinningSceneExtension::FHeaderData::GetSkeletonBatchKey() const
+{
+	return FSkeletonBatchKey
+	{
+	#if ENABLE_SKELETON_DEBUG_NAME
+		.SkeletonName = Proxy->GetSkinnedAsset()->GetSkeleton()->GetFName(),
+	#endif
+		.SkeletonGuid = Proxy->GetSkinnedAsset()->GetSkeleton()->GetGuid(),
+		.TransformProviderId = Proxy->GetTransformProviderId(),
+		.bNanite = bIsNanite != 0
+	};
+}
+
 bool FSkinningSceneExtension::ShouldCreateExtension(FScene& InScene)
 {
 	return IsGPUSkinSceneExtensionEnabled() || (NaniteSkinnedMeshesSupported() && DoesRuntimeSupportNanite(GetFeatureLevelShaderPlatform(InScene.GetFeatureLevel()), true, true));
@@ -458,8 +471,9 @@ void FSkinningSceneExtension::PerformSkinning(FSkinningSceneParameters& Paramete
 			uint32 TotalOffset = 0;
 			uint32 TotalBatches = 0;
 
-			for (auto& [BatchKey, Header] : BatchHeaderDatas)
+			for (auto& [BatchKey, BatchEntry] : BatchHeaderDatas)
 			{
+				const FHeaderData& Header = BatchEntry.HeaderData;
 				const FGuid ProviderId = BatchKey.TransformProviderId;
 				for (FSkinningTransformProvider::FProviderRange& Range : SkeletonRanges)
 				{
@@ -500,7 +514,7 @@ void FSkinningSceneExtension::PerformSkinning(FSkinningSceneParameters& Paramete
 			uint32 TotalBatchIndices = 0;
 
 			TArrayView<FSkinningTransformProvider::FProviderIndirection> BatchIndices = GraphBuilder.AllocPODArrayView<FSkinningTransformProvider::FProviderIndirection>(IndirectionCount);
-			for (auto& [HeaderDataCacheKey, Header] : BatchHeaderDatas)
+			for (auto& [HeaderDataCacheKey, BatchEntry] : BatchHeaderDatas)
 			{
 				const FGuid ProviderId = HeaderDataCacheKey.TransformProviderId;
 
@@ -512,7 +526,7 @@ void FSkinningSceneExtension::PerformSkinning(FSkinningSceneParameters& Paramete
 							TotalBatchIndices,
 							Offsets[TotalBatchIndices].CurrentTransformOffset  * sizeof(FCompressedBoneTransform),
 							Offsets[TotalBatchIndices].PreviousTransformOffset * sizeof(FCompressedBoneTransform),
-							Offsets[TotalBatchIndices].BoneMapOffset * sizeof(uint32),
+							Offsets[TotalBatchIndices].BoneMapOffset != INDEX_NONE ? Offsets[TotalBatchIndices].BoneMapOffset * sizeof(uint32) : INDEX_NONE,
 							Offsets[TotalBatchIndices].DirtyBoneTransforms
 						);
 						++Range.Count;
@@ -615,7 +629,7 @@ void FSkinningSceneExtension::PerformSkinning(FSkinningSceneParameters& Paramete
 					PrimitiveIndex,
 					Offsets[PrimitiveIndex].CurrentTransformOffset * sizeof(FCompressedBoneTransform),
 					Offsets[PrimitiveIndex].PreviousTransformOffset * sizeof(FCompressedBoneTransform),
-					Offsets[PrimitiveIndex].BoneMapOffset * sizeof(uint32),
+					Offsets[PrimitiveIndex].BoneMapOffset != INDEX_NONE ? Offsets[PrimitiveIndex].BoneMapOffset * sizeof(uint32) : INDEX_NONE,
 					Offsets[PrimitiveIndex].DirtyBoneTransforms
 				);
 				++Range.Count;
@@ -814,27 +828,48 @@ void FSkinningSceneExtension::FUpdater::PreSceneUpdate(FRDGBuilder& GraphBuilder
 				{
 					FSkinningSceneExtension::FHeaderData& Data = SceneData->HeaderDatas[PersistentIndex.Index];
 
-					if (!Data.bIsBatched)
+					auto FreeAllocatorSlots = [this](const FHeaderData& FreeData)
 					{
-						if (Data.ObjectSpaceBufferOffset != INDEX_NONE)
+						if (FreeData.ObjectSpaceBufferOffset != INDEX_NONE)
+						{
+							SceneData->ObjectSpaceAllocator.Free(FreeData.ObjectSpaceBufferOffset, FreeData.ObjectSpaceBufferCount);
+						}
+
+						if (FreeData.BoneMapBufferOffset != INDEX_NONE)
 						{
-							SceneData->ObjectSpaceAllocator.Free(Data.ObjectSpaceBufferOffset, Data.ObjectSpaceBufferCount);
+							SceneData->BoneMapAllocator.Free(FreeData.BoneMapBufferOffset, FreeData.BoneMapBufferCount);
 						}
 
-						if (Data.BoneMapBufferOffset != INDEX_NONE)
+						if (FreeData.HierarchyBufferOffset != INDEX_NONE)
 						{
-							SceneData->BoneMapAllocator.Free(Data.BoneMapBufferOffset, Data.BoneMapBufferCount);
+							SceneData->BoneHierarchyAllocator.Free(FreeData.HierarchyBufferOffset, FreeData.HierarchyBufferCount);
 						}
 
-						if (Data.HierarchyBufferOffset != INDEX_NONE)
+						if (FreeData.TransformBufferOffset != INDEX_NONE)
 						{
-							SceneData->BoneHierarchyAllocator.Free(Data.HierarchyBufferOffset, Data.HierarchyBufferCount);
+							SceneData->TransformAllocator.Free(FreeData.TransformBufferOffset, FreeData.TransformBufferCount);
 						}
+					};
+
+					if (Data.bIsBatched)
+					{
+						const FSkeletonBatchKey BatchKey = Data.GetSkeletonBatchKey();
 
-						if (Data.TransformBufferOffset != INDEX_NONE)
+						if (FBatchHeaderEntry* BatchEntry = SceneData->BatchHeaderDatas.Find(BatchKey))
 						{
-							SceneData->TransformAllocator.Free(Data.TransformBufferOffset, Data.TransformBufferCount);
+							check(BatchEntry->RefCount > 0);
+							BatchEntry->RefCount--;
+
+							if (BatchEntry->RefCount == 0)
+							{
+								FreeAllocatorSlots(BatchEntry->HeaderData);
+								SceneData->BatchHeaderDatas.Remove(BatchKey);
+							}
 						}
+					}
+					else
+					{
+						FreeAllocatorSlots(Data);
 
 						SceneData->AllocatedHeaderDataIndices.Remove(PersistentIndex.Index);
 
@@ -842,11 +877,11 @@ void FSkinningSceneExtension::FUpdater::PreSceneUpdate(FRDGBuilder& GraphBuilder
 						{
 							SceneData->InstancedHeaderDataIndices.Remove(PersistentIndex.Index);
 						}
+					}
 
-						if (Data.bVertexFactoryDirty)
-						{
-							SceneData->VertexFactoryUpdateList.Remove(PersistentIndex.Index);
-						}
+					if (Data.bVertexFactoryDirty)
+					{
+						SceneData->VertexFactoryUpdateList.Remove(PersistentIndex.Index);
 					}
 
 					SceneData->HeaderDatas.RemoveAt(PersistentIndex.Index);
@@ -924,6 +959,7 @@ void FSkinningSceneExtension::FUpdater::PostSceneUpdate(FRDGBuilder& GraphBuilde
 					NewHeader.bIsInstanced                = Proxy->UseInstancing();
 					NewHeader.bIsRefPose                  = NewHeader.ProviderId == RefPoseProviderId;
 					NewHeader.bHasScale                   = Proxy->HasScale();
+					NewHeader.bIsNanite                   = MeshObject->IsNaniteMesh();
 
 					MeshObject->SkinningSceneExtensionPrivateData.CachedPrimitiveIndex = PersistentIndex;
 
@@ -1104,21 +1140,14 @@ void FSkinningSceneExtension::FUpdater::PostMeshUpdate(FRDGBuilder& GraphBuilder
 
 		if (Data.Proxy->UseSkeletonBatching())
 		{
-			const FSkeletonBatchKey BatchKey
-			{
-			#if ENABLE_SKELETON_DEBUG_NAME
-				.SkeletonName = Data.Proxy->GetSkinnedAsset()->GetSkeleton()->GetFName(),
-			#endif
-				.SkeletonGuid = Data.Proxy->GetSkinnedAsset()->GetSkeleton()->GetGuid(),
-				.TransformProviderId = Data.Proxy->GetTransformProviderId()
-			};
+			const FSkeletonBatchKey BatchKey = Data.GetSkeletonBatchKey();
 
-			if (SceneData->BatchHeaderDatas.Contains(BatchKey))
+			if (FBatchHeaderEntry* BatchEntry = SceneData->BatchHeaderDatas.Find(BatchKey))
 			{
 				check(Data.bIsBatched == Data.Proxy->UseSkeletonBatching());
 
-				const FHeaderData& SrcHeaderData = SceneData->BatchHeaderDatas[BatchKey];
-
+				const FHeaderData& SrcHeaderData = BatchEntry->HeaderData;
+				BatchEntry->RefCount++;
 				Data.ObjectSpaceBufferOffset	= SrcHeaderData.ObjectSpaceBufferOffset;
 				Data.ObjectSpaceBufferCount		= SrcHeaderData.ObjectSpaceBufferCount;
 				Data.BoneMapBufferOffset		= SrcHeaderData.BoneMapBufferOffset;
@@ -1235,16 +1264,9 @@ void FSkinningSceneExtension::FUpdater::PostMeshUpdate(FRDGBuilder& GraphBuilder
 
 		if (Data.Proxy->UseSkeletonBatching())
 		{
-			const FSkeletonBatchKey BatchKey
-			{
-			#if ENABLE_SKELETON_DEBUG_NAME
-				.SkeletonName = Data.Proxy->GetSkinnedAsset()->GetSkeleton()->GetFName(),
-			#endif
-				.SkeletonGuid = Data.Proxy->GetSkinnedAsset()->GetSkeleton()->GetGuid(),
-				.TransformProviderId = Data.Proxy->GetTransformProviderId()
-			};
+			const FSkeletonBatchKey BatchKey = Data.GetSkeletonBatchKey();
 
-			SceneData->BatchHeaderDatas.Add(BatchKey, Data);
+			SceneData->BatchHeaderDatas.Add(BatchKey, FBatchHeaderEntry{ Data, 1 });
 		}
 		else
 		{
@@ -1306,7 +1328,7 @@ void FSkinningSceneExtension::FUpdater::PostMeshUpdate(FRDGBuilder& GraphBuilder
 		check(SceneData->Uploader.IsValid()); // Sanity check
 		SceneData->Uploader->HeaderDataUploader.Add(Data.Pack(), PersistentIndex);
 
-		if (!Data.bIsBatched && !Data.bVertexFactoryDirty)
+		if (!Data.bVertexFactoryDirty)
 		{
 			Data.bVertexFactoryDirty = 1;
 			SceneData->VertexFactoryUpdateList.Emplace(PersistentIndex);
diff --git a/Engine/Source/Runtime/Renderer/Private/Skinning/SkinningSceneExtension.h b/Engine/Source/Runtime/Renderer/Private/Skinning/SkinningSceneExtension.h
--- a/Engine/Source/Runtime/Renderer/Private/Skinning/SkinningSceneExtension.h
+++ b/Engine/Source/Runtime/Renderer/Private/Skinning/SkinningSceneExtension.h
@@ -136,6 +136,7 @@ private:
 		uint8  bIsInstanced : 1                  = false;
 		uint8  bIsRefPose : 1                    = false;
 		uint8  bVertexFactoryDirty : 1           = false;
+		uint8  bIsNanite : 1                     = false;
 		uint8  CurrentTransformSlot : 1          = 0;
 
 		bool IsFirstUpdate() const
@@ -143,6 +144,8 @@ private:
 			return UpdateCounter == 0;
 		}
 
+		FSkeletonBatchKey GetSkeletonBatchKey() const;
+
 		FSkinningHeader Pack() const
 		{
 			Validate();
@@ -232,7 +235,12 @@ private:
 	TSparseArray<FHeaderData> HeaderDatas;
 	TSet<int32> AllocatedHeaderDataIndices;
 	TSet<int32> InstancedHeaderDataIndices;
-	TMap<FSkeletonBatchKey, FHeaderData> BatchHeaderDatas;
+	struct FBatchHeaderEntry
+	{
+		FHeaderData HeaderData;
+		uint32 RefCount = 0;
+	};
+	TMap<FSkeletonBatchKey, FBatchHeaderEntry> BatchHeaderDatas;
 	TUniquePtr<FBuffers> Buffers;
 	TUniquePtr<FUploader> Uploader;
 	TStaticArray<UE::Tasks::FTask, NumTasks> TaskHandles;
diff --git a/Engine/Source/Runtime/Renderer/Private/Skinning/SkinningTransformProvider.h b/Engine/Source/Runtime/Renderer/Private/Skinning/SkinningTransformProvider.h
--- a/Engine/Source/Runtime/Renderer/Private/Skinning/SkinningTransformProvider.h
+++ b/Engine/Source/Runtime/Renderer/Private/Skinning/SkinningTransformProvider.h
@@ -31,10 +31,11 @@ struct FSkeletonBatchKey
 #endif
 	FGuid SkeletonGuid;
 	FGuid TransformProviderId;
+	bool bNanite = false;
 
 	bool operator == (const FSkeletonBatchKey& InOther) const
 	{
-		return SkeletonGuid == InOther.SkeletonGuid && TransformProviderId == InOther.TransformProviderId;
+		return SkeletonGuid == InOther.SkeletonGuid && TransformProviderId == InOther.TransformProviderId && bNanite == InOther.bNanite;
 	}
 
 	bool operator != (const FSkeletonBatchKey& InOther) const
@@ -44,7 +45,7 @@ struct FSkeletonBatchKey
 
 	friend inline SIZE_T GetTypeHash(const FSkeletonBatchKey& InKey)
 	{
-		return HashCombine(GetTypeHash(InKey.SkeletonGuid), GetTypeHash(InKey.TransformProviderId));
+		return HashCombine(HashCombine(GetTypeHash(InKey.SkeletonGuid), GetTypeHash(InKey.TransformProviderId)), GetTypeHash(InKey.bNanite));
 	}
 };
 
