VirtualTexture: Fix bug where the page pool isn't fully cleared over a release/init cycle for a physical space.
This could leave pages in the reinitialized physical space which are allocated but never properly refilled in the newly created texture.
The pages would then appear to be corrupted.
#Synced-CL 52391845

diff --git a/Engine/Source/Runtime/Renderer/Private/VT/AllocatedVirtualTexture.cpp b/Engine/Source/Runtime/Renderer/Private/VT/AllocatedVirtualTexture.cpp
--- a/Engine/Source/Runtime/Renderer/Private/VT/AllocatedVirtualTexture.cpp
+++ b/Engine/Source/Runtime/Renderer/Private/VT/AllocatedVirtualTexture.cpp
@@ -173,6 +173,10 @@ void FAllocatedVirtualTexture::Destroy(FVirtualTextureSystem* System)
 			{
 				if (UniquePageTableLayers[PageTableIndex].PhysicalSpace->ReleaseResourceRef() == 0)
 				{
+					// Clean up page pool before release. Pages should all be unmapped, but may still be allocated. 
+					// The physical space can be recycled later, so it's state needs to be clean.
+					check(UniquePageTableLayers[PageTableIndex].PhysicalSpace->GetPagePool().GetNumMappedPages() == 0);
+					UniquePageTableLayers[PageTableIndex].PhysicalSpace->GetPagePool().EvictAllPages(System);
 					UniquePageTableLayers[PageTableIndex].PhysicalSpace->ReleaseResource();
 				}
 				UniquePageTableLayers[PageTableIndex].PhysicalSpace.SafeRelease();
diff --git a/Engine/Source/Runtime/Renderer/Private/VT/TexturePagePool.h b/Engine/Source/Runtime/Renderer/Private/VT/TexturePagePool.h
--- a/Engine/Source/Runtime/Renderer/Private/VT/TexturePagePool.h
+++ b/Engine/Source/Runtime/Renderer/Private/VT/TexturePagePool.h
@@ -34,7 +34,7 @@ public:
 	uint32 GetNumAllocatedPages() const { return NumPagesAllocated; }
 
 	/**
-	 * Reset the page pool. This can be used to flush any caches. Mainly useful for debug and testing purposes.
+	 * Reset the page pool. This can be used to flush any caches.
 	 */
 	void EvictAllPages(FVirtualTextureSystem* System);
 
