Avoid assert about CombineLUTs_Inner being unwritten by using an external placeholder texture when it's not necessary

#Synced-CL 52288345

diff --git a/Engine/Source/Runtime/Renderer/Private/PostProcess/PostProcessCombineLUTs.cpp b/Engine/Source/Runtime/Renderer/Private/PostProcess/PostProcessCombineLUTs.cpp
--- a/Engine/Source/Runtime/Renderer/Private/PostProcess/PostProcessCombineLUTs.cpp
+++ b/Engine/Source/Runtime/Renderer/Private/PostProcess/PostProcessCombineLUTs.cpp
@@ -799,18 +799,18 @@ FRDGTextureRef AddCombineLUTPass(FRDGBuilder& GraphBuilder, const FViewInfo& Vie
 	const bool bOutputDeviceSRGB = (CachedLUTSettings.Parameters.OutputDevice.OutputDevice == (uint32)EDisplayOutputFormat::SDR_sRGB);
 	if (bHasInnerChanged || !InnerTexture || CVarLUTInnerUpdateEveryFrame.GetValueOnRenderThread())
 	{
-		// View doesn't support a persistent LUT, so create a temporary one.
-		if (!InnerTexture)
+		if (bNeedsInnerTonemappingLUT)
 		{
-			InnerTexture = GraphBuilder.CreateTexture(
-				Translate(FSceneViewState::CreateLUTRenderTarget(InnerLUTTextureSize, bUseFloatOutput, TEXT("CombineLUTs_Inner"))),
-				TEXT("CombineLUTs_inner"));
-		}
+			// View doesn't support a persistent LUT, so create a temporary one.
+			if (!InnerTexture)
+			{
+				InnerTexture = GraphBuilder.CreateTexture(
+					Translate(FSceneViewState::CreateLUTRenderTarget(InnerLUTTextureSize, bUseFloatOutput, TEXT("CombineLUTs_Inner"))),
+					TEXT("CombineLUTs_Inner"));
+			}
 
-		check(InnerTexture);
+			check(InnerTexture);
 
-		if (bNeedsInnerTonemappingLUT)
-		{
 			FLUTInnerCS::FPermutationDomain PermutationVector;
 			PermutationVector.Set<FLUTInnerCS::FUseACES2>(CachedInnerLUTSettings.AcesVersion > 1);
 
@@ -825,7 +825,7 @@ FRDGTextureRef AddCombineLUTPass(FRDGBuilder& GraphBuilder, const FViewInfo& Vie
 			const uint32 GroupSize = FMath::DivideAndRoundUp(InnerLUTSize, FLUTInnerCS::GroupSize);
 			FComputeShaderUtils::AddPass(
 				GraphBuilder,
-				RDG_EVENT_NAME("CombineLUTs_inner %d (CS)", InnerLUTSize),
+				RDG_EVENT_NAME("CombineLUTs_Inner %d (CS)", InnerLUTSize),
 				CVarLUTAsyncCompute.GetValueOnRenderThread() ? ERDGPassFlags::AsyncCompute : ERDGPassFlags::Compute,
 				ComputeShader,
 				PassParameters,
@@ -839,7 +839,9 @@ FRDGTextureRef AddCombineLUTPass(FRDGBuilder& GraphBuilder, const FViewInfo& Vie
 
 	FLUTBlenderCS::FParameters* PassParameters = GraphBuilder.AllocParameters<FLUTBlenderCS::FParameters>();
 	PassParameters->CombineLUT = CachedLUTSettings.Parameters;
-	PassParameters->InnerLUT = InnerTexture;
+	PassParameters->InnerLUT = bNeedsInnerTonemappingLUT
+		? InnerTexture
+		: RegisterExternalTexture(GraphBuilder, GBlackVolumeTexture->GetTextureRHI(), TEXT("CombineLUTs_Inner_Fallback"));
 	PassParameters->InnerLUTSampler = TStaticSamplerState<SF_Bilinear, AM_Clamp, AM_Clamp, AM_Clamp>::GetRHI();
 	PassParameters->InvLUTSizeMinusOne = 1.f / (LUTSize - 1);
 	PassParameters->RWOutputTexture = GraphBuilder.CreateUAV(OutputTexture);
