Avoid race condition caused by writes to bIsCachedRayTracingInstanceValid and bDynamicRayTracingInstanceCachedDataDirty from parallel tasks.
- force separate storage unit between raytracing and non-raytracing bits.
- force separate storage unit between bIsCachedRayTracingInstanceValid and bDynamicRayTracingInstanceCachedDataDirty.

#Synced-CL 52364312

diff --git a/Engine/Source/Runtime/Renderer/Public/PrimitiveSceneInfo.h b/Engine/Source/Runtime/Renderer/Public/PrimitiveSceneInfo.h
--- a/Engine/Source/Runtime/Renderer/Public/PrimitiveSceneInfo.h
+++ b/Engine/Source/Runtime/Renderer/Public/PrimitiveSceneInfo.h
@@ -656,6 +656,7 @@ public:
 	bool bNeedsCachedReflectionCaptureUpdate = true;
 
 #if RHI_RAYTRACING
+	bool : 0; // force separate storage unit to avoid race conditions with tasks modifying other non-raytracing bits.
 	bool bHasRayTracingRepresentation : 1 = false;
 	bool bIsRelevantToPathTracing : 1;
 	bool bIsCachedRayTracingInstanceValid : 1;
@@ -667,6 +668,8 @@ public:
 	bool bCachedRayTracingInstanceAllSegmentsDecal : 1;
 	UE_DEPRECATED(5.7, "Use DynamicRayTracingInstanceCachedData or GetRayTracingLODData(...) instead.")
 	bool bCachedRayTracingInstanceAllSegmentsTranslucent : 1;
+	bool : 0; // force separate storage unit to avoid race condition modifying bIsCachedRayTracingInstanceValid and bDynamicRayTracingInstanceCachedDataDirty.
+	bool bDynamicRayTracingInstanceCachedDataDirty : 1 = true;
 
 	struct FDynamicRayTracingInstanceCachedData
 	{
@@ -678,7 +681,6 @@ public:
 		uint8 bAllSegmentsDecal : 1 = false;
 		uint8 bAllSegmentsTranslucent : 1 = false;
 	};
-	bool bDynamicRayTracingInstanceCachedDataDirty : 1 = true;
 	FDynamicRayTracingInstanceCachedData DynamicRayTracingInstanceCachedData;
 
 	struct FRayTracingInstanceData
