[HWRT] Remove GeometriesToBuild path.
- this path didn't have a specific use case. Replaced with ensure so any instances using ray tracing geometry that is pending a build are identified.
- also added an ensure when checking for Evicted geometries since that case should be handled at higher level.

#Synced-CL 52364312

diff --git a/Engine/Source/Runtime/Renderer/Private/DeferredShadingRenderer.cpp b/Engine/Source/Runtime/Renderer/Private/DeferredShadingRenderer.cpp
--- a/Engine/Source/Runtime/Renderer/Private/DeferredShadingRenderer.cpp
+++ b/Engine/Source/Runtime/Renderer/Private/DeferredShadingRenderer.cpp
@@ -1091,11 +1091,6 @@ bool FDeferredShadingSceneRenderer::DispatchRayTracingWorldUpdates(FRDGBuilder&
 
 	const bool bRayTracingAsyncBuild = CVarRayTracingAsyncBuild.GetValueOnRenderThread() != 0 && GRHISupportsRayTracingAsyncBuildAccelerationStructure;
 	const ERDGPassFlags ComputePassFlags = bRayTracingAsyncBuild ? ERDGPassFlags::AsyncCompute : ERDGPassFlags::Compute;
-	if (RayTracingScene.GeometriesToBuild.Num() > 0)
-	{
-		// Force update all the collected geometries (use stack allocator?)
-		GRayTracingGeometryManager->ForceBuildIfPending(GraphBuilder.RHICmdList, RayTracingScene.GeometriesToBuild);
-	}
 
 	{
 		Nanite::GRayTracingManager.UpdateStreaming(GraphBuilder, Views, GetSceneUniforms(), GetActiveSceneTexturesConfig().Extent);
diff --git a/Engine/Source/Runtime/Renderer/Private/RayTracing/RayTracing.cpp b/Engine/Source/Runtime/Renderer/Private/RayTracing/RayTracing.cpp
--- a/Engine/Source/Runtime/Renderer/Private/RayTracing/RayTracing.cpp
+++ b/Engine/Source/Runtime/Renderer/Private/RayTracing/RayTracing.cpp
@@ -1008,23 +1008,26 @@ namespace RayTracing
 
 					((FRayTracingGeometryManager*)GRayTracingGeometryManager)->AddReferencedGeometry(Geometry);
 
-					if (Geometry->IsEvicted())
+					if (!ensureMsgf(!Geometry->IsEvicted(),
+						TEXT("Dynamic ray tracing instance skipped because geometry '%s' is evicted."),
+						*GeometryInitializer.DebugName.ToString()))
 					{
 						continue;
 					}
 
-					// If geometry still has pending build request then add to list which requires a force build
-					if (Geometry->HasPendingBuildRequest())
+					if (!ensureMsgf(!Geometry->HasPendingBuildRequest(),
+						TEXT("Dynamic ray tracing instance skipped because geometry '%s' has a pending build request."),
+						*GeometryInitializer.DebugName.ToString()))
 					{
-						RayTracingScene.GeometriesToBuild.Add(Geometry);
+						continue;
 					}
 
 					TConstArrayView<FMeshBatch> InstanceMaterials = Instance.GetMaterials();
 
 					// Validate the material/segment counts
 					if (!ensureMsgf(InstanceMaterials.Num() <= NumSegments,
-						TEXT("Ray tracing material assignment validation failed for geometry '%s'. "
-							"Instance.GetMaterials().Num() = %d, GeometryInitializer.Segments.Num() = %d."),
+						TEXT("Ray tracing material assignment validation failed for geometry '%s'. ")
+						TEXT("Instance.GetMaterials().Num() = %d, GeometryInitializer.Segments.Num() = %d."),
 						*GeometryInitializer.DebugName.ToString(), InstanceMaterials.Num(),
 						NumSegments))
 					{
diff --git a/Engine/Source/Runtime/Renderer/Private/RayTracing/RayTracingScene.cpp b/Engine/Source/Runtime/Renderer/Private/RayTracing/RayTracingScene.cpp
--- a/Engine/Source/Runtime/Renderer/Private/RayTracing/RayTracingScene.cpp
+++ b/Engine/Source/Runtime/Renderer/Private/RayTracing/RayTracingScene.cpp
@@ -1216,8 +1216,6 @@ void FRayTracingScene::Reset()
 
 	InstanceStatsBuffer = nullptr;
 
-	GeometriesToBuild.Reset();
-
 	Allocator.Flush();
 
 	LightingChannelAndMask = 0xFF;
@@ -1248,8 +1246,6 @@ void FRayTracingScene::EndFrame()
 			}
 		}
 
-		GeometriesToBuild.Empty();
-
 		ReleaseFeedbackReadbackBuffers();
 		ReleaseReadbackBuffers();
 	}
diff --git a/Engine/Source/Runtime/Renderer/Private/RayTracing/RayTracingScene.h b/Engine/Source/Runtime/Renderer/Private/RayTracing/RayTracingScene.h
--- a/Engine/Source/Runtime/Renderer/Private/RayTracing/RayTracingScene.h
+++ b/Engine/Source/Runtime/Renderer/Private/RayTracing/RayTracingScene.h
@@ -211,9 +211,6 @@ public:
 
 	// Public members for initial refactoring step (previously were public members of FViewInfo).
 
-	// Geometries which still have a pending build request but are used this frame and require a force build.
-	TArray<const FRayTracingGeometry*> GeometriesToBuild;
-
 	uint8 LightingChannelAndMask = 0xFF;
 
 	UE::Tasks::FTask InitTask; // Task to asynchronously call BuildInitializationData()
