[Preselection Highlight] Separating Hovering color/logic from Selection

- Use subuded color/index for preselection highlight
- store preselection highlight indexes in ActorNameToStencilIndex instead of ProxyToStencilIndex, which is used for actual selection

#Synced-CL 52417822

diff --git a/Engine/Source/Runtime/Renderer/Private/SceneHitProxyRendering.cpp b/Engine/Source/Runtime/Renderer/Private/SceneHitProxyRendering.cpp
--- a/Engine/Source/Runtime/Renderer/Private/SceneHitProxyRendering.cpp
+++ b/Engine/Source/Runtime/Renderer/Private/SceneHitProxyRendering.cpp
@@ -1109,7 +1109,15 @@ int32 FEditorSelectionMeshProcessor::GetStencilValue(const FSceneView* View, con
 	static constexpr int UniqueIdBits = BitsAvailable - ColorBits;
 	static constexpr int MaxColor = (1 << ColorBits);
 	static constexpr int MaxUniqueId = (1 << UniqueIdBits);
-	
+
+	// Keep Selection Outline color index (0) just for selection.
+	// It is handled differently by EncodeSelectionStencilValue lambda below.
+	// Using 0 index for other colors might have the function generate the same stencil value for different elements,
+	// preventing the shader from distinguishing them as separate objects.
+	static constexpr int SelectionOutlineColorIndex = 0; // See FSceneView::SelectionOutlineColor
+	static constexpr int SubduedSelectionOutlineColorIndex = 1; // See FSceneView::SubduedSelectionOutlineColor
+	static constexpr int PreselectionOutlineColorIndex = SubduedSelectionOutlineColorIndex;
+
 	auto EncodeSelectionStencilValue = [](int32 ColorIndex, int32 UniqueId) -> int32
 	{
 		uint8 Bits = 0;
@@ -1139,7 +1147,7 @@ int32 FEditorSelectionMeshProcessor::GetStencilValue(const FSceneView* View, con
 	{
 		StencilValue = *ExistingStencilValue;
 	}
-	else if (PrimitiveSceneProxy->IsIndividuallySelected() || PrimitiveSceneProxy->IsHovered())
+	else if (PrimitiveSceneProxy->IsIndividuallySelected())
 	{
 		const int Color = 0;
 		const int UniqueId = ProxyToStencilIndex.Num();
@@ -1149,15 +1157,21 @@ int32 FEditorSelectionMeshProcessor::GetStencilValue(const FSceneView* View, con
 	else if (PrimitiveSceneProxy->IsParentSelected())
 	{
 		int Color = PrimitiveSceneProxy->GetSelectionOutlineColorIndex();
-		if (bActorSelectionColorIsSubdued && (Color == 0))
+		if (bActorSelectionColorIsSubdued && (Color == SelectionOutlineColorIndex))
 		{
-			Color = 1;
+			Color = SubduedSelectionOutlineColorIndex;
 		}
 		const int UniqueId = ActorNameToStencilIndex.Num();
 
 		StencilValue = EncodeSelectionStencilValue(Color, UniqueId);
 		ActorNameToStencilIndex.Add(PrimitiveSceneProxy->GetOwnerName(), StencilValue);
 	}
+	else if (PrimitiveSceneProxy->IsHovered())
+	{
+		const int UniqueId = ActorNameToStencilIndex.Num();
+		StencilValue = EncodeSelectionStencilValue(PreselectionOutlineColorIndex, UniqueId);
+		ActorNameToStencilIndex.Add(PrimitiveSceneProxy->GetOwnerName(), StencilValue);
+	}
 
 	return StencilValue;
 }
