Fixed: Missing volumetric fog permutation when running SM6 client with a -noraytracing parameter
#Synced-CL 52440158

diff --git a/Engine/Source/Runtime/Renderer/Private/VolumetricFog.cpp b/Engine/Source/Runtime/Renderer/Private/VolumetricFog.cpp
--- a/Engine/Source/Runtime/Renderer/Private/VolumetricFog.cpp
+++ b/Engine/Source/Runtime/Renderer/Private/VolumetricFog.cpp
@@ -1238,7 +1238,7 @@ class FVolumetricFogLightScatteringCS : public FGlobalShader
 			PermutationVector.Set<FDistanceFieldSkyOcclusion>(DoesPlatformSupportDistanceFieldAO(ShaderPlatform));
 			PermutationVector.Set<FLumenGI>(FDataDrivenShaderPlatformInfo::GetSupportsLumenGI(ShaderPlatform));
 			PermutationVector.Set<FVirtualShadowMap>(DoesPlatformSupportVirtualShadowMaps(ShaderPlatform));
-			PermutationVector.Set<FRaytracedShadowsVolume>(ShouldCompileRayTracingShadersForProject(ShaderPlatform));
+			PermutationVector.Set<FRaytracedShadowsVolume>(RHISupportsRayTracingShaders(ShaderPlatform));
 			PermutationVector.Set<FSampleLightFunctionAtlas>(true);
 			PermutationVector.Set<FMegaLights>(MegaLights::ShouldCompileShaders(ShaderPlatform));
 			PermutationVector.Set<FLightSoftFading>(true);
@@ -1264,7 +1264,7 @@ class FVolumetricFogLightScatteringCS : public FGlobalShader
 				GetIfSupported<FMegaLights>(PermutationVector, MegaLights::ShouldCompileShaders(Parameters.Platform)) &&
 				GetIfSupported<FVirtualShadowMap>(PermutationVector, DoesPlatformSupportVirtualShadowMaps(Parameters.Platform)) &&
 				GetIfSupported<FDistanceFieldSkyOcclusion>(PermutationVector, DoesPlatformSupportDistanceFieldAO(Parameters.Platform)) &&
-				GetIfSupported<FRaytracedShadowsVolume>(PermutationVector, ShouldCompileRayTracingShadersForProject(Parameters.Platform)) &&
+				GetIfSupported<FRaytracedShadowsVolume>(PermutationVector, RHISupportsRayTracingShaders(Parameters.Platform)) &&
 				PermutationVector.Get<FTemporalReprojection>() &&
 				PermutationVector.Get<FSampleLightFunctionAtlas>() &&
 				PermutationVector.Get<FLightSoftFading>() &&
