LocalExposure  integrated in Mobile Tonemap
Disabled by default for now on mobile
#Synced-CL 52441908

#Synced-CL 52441908

diff --git a/Engine/Config/BaseDeviceProfiles.ini b/Engine/Config/BaseDeviceProfiles.ini
--- a/Engine/Config/BaseDeviceProfiles.ini
+++ b/Engine/Config/BaseDeviceProfiles.ini
@@ -378,6 +378,8 @@ bIsVisibleForAssets=True
 ; PreviewAllowlistCVars and PreviewDenyListCVars are arrays of cvars (or prefixes) that are included or excluded from being applied in editor preview
 ; If any PreviewAllowlistCVars is set, cvars are denied by default.
 PreviewAllowlistCVars=sg.
+; disable LocalExposure
++CVars=r.LocalExposure=0
 
 [IOS_Low DeviceProfile]
 BaseProfileName=IOS
@@ -986,6 +988,8 @@ bIsVisibleForAssets=True
 ; If any PreviewAllowlistCVars is set, cvars are denied by default.
 PreviewAllowlistCVars=sg.
 +CVars=AsyncTraceDelegateHitchThresholdMS=1.5
+; disable LocalExposure
++CVars=r.LocalExposure=0
 
 [Android_Preview_OpenGL DeviceProfile]
 DeviceType=Android
diff --git a/Engine/Source/Runtime/Renderer/Private/PostProcess/PostProcessHistogram.cpp b/Engine/Source/Runtime/Renderer/Private/PostProcess/PostProcessHistogram.cpp
--- a/Engine/Source/Runtime/Renderer/Private/PostProcess/PostProcessHistogram.cpp
+++ b/Engine/Source/Runtime/Renderer/Private/PostProcess/PostProcessHistogram.cpp
@@ -51,11 +51,6 @@ public:
 
 	static bool ShouldCompilePermutation(const FGlobalShaderPermutationParameters& Parameters)
 	{
-		if (!IsFeatureLevelSupported(Parameters.Platform, ERHIFeatureLevel::SM5))
-		{
-			return false;
-		}
-
 		const FPermutationDomain PermutationVector(Parameters.PermutationId);
 		const AutoExposurePermutation::FCommonDomain& AEPermutationCommon = PermutationVector.Get<AutoExposurePermutation::FCommonDomain>();
 
@@ -494,6 +489,12 @@ FRDGTextureRef AddLocalExposurePass(
 {
 	check(SceneColor.IsValid());
 
+	// Bilateral grid requires PF_G32R32F UAV. If the format doesn't support it, skip.
+	if (!UE::PixelFormat::HasCapabilities(PF_G32R32F, EPixelFormatCapabilities::UAV))
+	{
+		return nullptr;
+	}
+
 	const FIntPoint ThreadGroupSize = FHistogramCS::GetThreadGroupSize(true);
 	const FIntPoint ThreadGroupCount = FHistogramCS::GetThreadGroupCount(SceneColor.ViewRect.Size(), true);
 	const FIntPoint NumTiles = ThreadGroupCount;
diff --git a/Engine/Source/Runtime/Renderer/Private/PostProcess/PostProcessLocalExposure.cpp b/Engine/Source/Runtime/Renderer/Private/PostProcess/PostProcessLocalExposure.cpp
--- a/Engine/Source/Runtime/Renderer/Private/PostProcess/PostProcessLocalExposure.cpp
+++ b/Engine/Source/Runtime/Renderer/Private/PostProcess/PostProcessLocalExposure.cpp
@@ -46,11 +46,6 @@ public:
 		SHADER_PARAMETER_RDG_TEXTURE_UAV(RWTexture2D<float>, OutputFloat)
 	END_SHADER_PARAMETER_STRUCT()
 
-	static bool ShouldCompilePermutation(const FGlobalShaderPermutationParameters& Parameters)
-	{
-		return IsFeatureLevelSupported(Parameters.Platform, ERHIFeatureLevel::SM5);
-	}
-
 	static void ModifyCompilationEnvironment(const FGlobalShaderPermutationParameters& Parameters, FShaderCompilerEnvironment& OutEnvironment)
 	{
 		FGlobalShader::ModifyCompilationEnvironment(Parameters, OutEnvironment);
@@ -138,11 +133,6 @@ public:
 		SHADER_PARAMETER(float, FilmWhiteClip)
 	END_SHADER_PARAMETER_STRUCT()
 
-	static bool ShouldCompilePermutation(const FGlobalShaderPermutationParameters& Parameters)
-	{
-		return IsFeatureLevelSupported(Parameters.Platform, ERHIFeatureLevel::SM5);
-	}
-
 	static void ModifyCompilationEnvironment(const FGlobalShaderPermutationParameters& Parameters, FShaderCompilerEnvironment& OutEnvironment)
 	{
 		FGlobalShader::ModifyCompilationEnvironment(Parameters, OutEnvironment);
@@ -186,11 +176,6 @@ public:
 		SHADER_PARAMETER(FScreenTransform, DispatchThreadToCoarseMipUV)
 	END_SHADER_PARAMETER_STRUCT()
 
-	static bool ShouldCompilePermutation(const FGlobalShaderPermutationParameters& Parameters)
-	{
-		return IsFeatureLevelSupported(Parameters.Platform, ERHIFeatureLevel::SM5);
-	}
-
 	static void ModifyCompilationEnvironment(const FGlobalShaderPermutationParameters& Parameters, FShaderCompilerEnvironment& OutEnvironment)
 	{
 		FGlobalShader::ModifyCompilationEnvironment(Parameters, OutEnvironment);
diff --git a/Engine/Source/Runtime/Renderer/Private/PostProcess/PostProcessing.cpp b/Engine/Source/Runtime/Renderer/Private/PostProcess/PostProcessing.cpp
--- a/Engine/Source/Runtime/Renderer/Private/PostProcess/PostProcessing.cpp
+++ b/Engine/Source/Runtime/Renderer/Private/PostProcess/PostProcessing.cpp
@@ -2713,10 +2713,17 @@ void AddMobilePostProcessingPasses(FRDGBuilder& GraphBuilder, FScene* Scene, con
 
 	FVisualizeBufferContext VisualizeBufferContext{};
 
+	// Local exposure state
+	FLocalExposureParameters LocalExposureParameters;
+	FRDGTextureRef LocalExposureBilateralGridTexture = nullptr;
+	FRDGTextureRef LocalExposureBlurredLogLumTexture = nullptr;
+	FExposureFusionData ExposureFusionData;
+
 	enum class EPass : uint32
 	{
 		Distortion,
 		SunMask,
+		LocalExposure,
 		BloomSetup,
 		DepthOfField,
 		Bloom,
@@ -2764,6 +2771,7 @@ void AddMobilePostProcessingPasses(FRDGBuilder& GraphBuilder, FScene* Scene, con
 	{
 		TEXT("Distortion"),
 		TEXT("SunMask"),
+		TEXT("LocalExposure"),
 		TEXT("BloomSetup"),
 		TEXT("DepthOfField"),
 		TEXT("Bloom"),
@@ -2966,10 +2974,18 @@ void AddMobilePostProcessingPasses(FRDGBuilder& GraphBuilder, FScene* Scene, con
 
 		bool bUseSeparateTranslucency = IsMobileSeparateTranslucencyActive(View);
 
+		const bool bLocalExposureEnabled =
+			(!FMath::IsNearlyEqual(View.FinalPostProcessSettings.LocalExposureHighlightContrastScale, 1.0f) ||
+			!FMath::IsNearlyEqual(View.FinalPostProcessSettings.LocalExposureShadowContrastScale, 1.0f) ||
+			View.FinalPostProcessSettings.LocalExposureHighlightContrastCurve ||
+			View.FinalPostProcessSettings.LocalExposureShadowContrastCurve ||
+			!FMath::IsNearlyEqual(View.FinalPostProcessSettings.LocalExposureDetailStrength, 1.0f));
+
 		const FPostProcessMaterialChain PostProcessMaterialAfterTonemappingChain = GetPostProcessMaterialChain(View, BL_SceneColorAfterTonemapping);
 
 		PassSequence.SetEnabled(EPass::Distortion, bUseDistortion);
 		PassSequence.SetEnabled(EPass::SunMask, bUseSun || bUseDof);
+		PassSequence.SetEnabled(EPass::LocalExposure, bLocalExposureEnabled);
 		PassSequence.SetEnabled(EPass::EyeAdaptation, bUseBasicEyeAdaptation || bUseHistogramEyeAdaptation);
 		PassSequence.SetEnabled(EPass::BloomSetup, bUseSun || bUseMobileDof || bUseBloom || PassSequence.IsEnabled(EPass::EyeAdaptation));
 		PassSequence.SetEnabled(EPass::DepthOfField, bUseDof);
@@ -3067,6 +3083,36 @@ void AddMobilePostProcessingPasses(FRDGBuilder& GraphBuilder, FScene* Scene, con
 			//@todo Ronin sunmask pass isnt clipping to image only.
 		}
 
+		// Mobile local exposure
+		if (PassSequence.IsEnabled(EPass::LocalExposure))
+		{
+			PassSequence.AcceptPass(EPass::LocalExposure);
+			const ELocalExposureMethod LocalExposureMethod = View.FinalPostProcessSettings.LocalExposureMethod;
+
+			FScreenPassTextureSlice LocalExposureSceneColor = FScreenPassTextureSlice::CreateFromScreenPassTexture(GraphBuilder, SceneColor);
+
+			if (LocalExposureMethod == ELocalExposureMethod::Bilateral)
+			{
+				LocalExposureBilateralGridTexture = AddLocalExposurePass(GraphBuilder, View, EyeAdaptationParameters, LocalExposureSceneColor);
+			}
+
+			LocalExposureParameters = GetLocalExposureParameters(View, LocalExposureSceneColor.ViewRect.Size(), EyeAdaptationParameters);
+
+			if (LocalExposureMethod == ELocalExposureMethod::Bilateral && LocalExposureBilateralGridTexture != nullptr && View.FinalPostProcessSettings.LocalExposureBlurredLuminanceBlend > 0.0f)
+			{
+				LocalExposureBlurredLogLumTexture = AddLocalExposureBlurredLogLuminancePass(
+					GraphBuilder, View, EyeAdaptationParameters, LocalExposureSceneColor);
+			}
+
+			if (LocalExposureMethod == ELocalExposureMethod::Fusion && LastEyeAdaptationBuffer)
+			{
+				ExposureFusionData = AddLocalExposureFusionPass(
+					GraphBuilder, View, EyeAdaptationParameters,
+					LastEyeAdaptationBuffer, LocalExposureParameters,
+					LocalExposureSceneColor);
+			}
+		}
+
 		FMobileBloomSetupOutputs BloomSetupOutputs;
 
 		const int32 BloomQualityDivider = GetBloomQuality() <= EBloomQuality::Q2 ? 4 : 2;
@@ -3442,6 +3488,7 @@ void AddMobilePostProcessingPasses(FRDGBuilder& GraphBuilder, FScene* Scene, con
 	{
 		PassSequence.SetEnabled(EPass::Distortion, false);
 		PassSequence.SetEnabled(EPass::SunMask, false);
+		PassSequence.SetEnabled(EPass::LocalExposure, false);
 		PassSequence.SetEnabled(EPass::BloomSetup, false);
 		PassSequence.SetEnabled(EPass::DepthOfField, false);
 		PassSequence.SetEnabled(EPass::Bloom, false);
@@ -3530,6 +3577,12 @@ void AddMobilePostProcessingPasses(FRDGBuilder& GraphBuilder, FScene* Scene, con
 			TonemapperInputs.bMetalMSAAHDRDecode = bMetalMSAAHDRDecode;
 			TonemapperInputs.EyeAdaptationBuffer = bUseEyeAdaptation ? LastEyeAdaptationBuffer : nullptr;
 
+			// Local exposure
+			TonemapperInputs.LocalExposureBilateralGridTexture = LocalExposureBilateralGridTexture;
+			TonemapperInputs.BlurredLogLuminanceTexture = LocalExposureBlurredLogLumTexture;
+			TonemapperInputs.ExposureFusion = FScreenPassTextureSlice::CreateFromScreenPassTexture(GraphBuilder, ExposureFusionData.Result);
+			TonemapperInputs.LocalExposureParameters = &LocalExposureParameters;
+
 			SceneColor = AddTonemapPass(GraphBuilder, View, TonemapperInputs);
 		}
 
