[TSR] Fix a runtime error

#Synced-CL 52438879

diff --git a/Engine/Source/Runtime/Renderer/Private/PostProcess/TemporalSuperResolution.cpp b/Engine/Source/Runtime/Renderer/Private/PostProcess/TemporalSuperResolution.cpp
--- a/Engine/Source/Runtime/Renderer/Private/PostProcess/TemporalSuperResolution.cpp
+++ b/Engine/Source/Runtime/Renderer/Private/PostProcess/TemporalSuperResolution.cpp
@@ -2760,7 +2760,7 @@ FDefaultTemporalUpscaler::FOutputs AddTemporalSuperResolutionPasses(
 						int32 ArrayIndex = GetTSRMainFlickeringLumaTextureArraySize() - 1;
 						PassParameters->CurrentCoverageTexture = GraphBuilder.CreateSRV(
 							FRDGTextureSRVDesc::CreateForSlice(PassInputs.FlickeringInputTexture.Texture,  /* SliceIndex = */ ArrayIndex));
-						if (MinKeepLineIntensity != 0.0f)
+						if (MinKeepLineIntensity != 0.0f && ArrayIndex == 1)
 						{
 							PassParameters->FlickeringLumaTexture = GraphBuilder.CreateSRV(
 								FRDGTextureSRVDesc::CreateForSlice(PassInputs.FlickeringInputTexture.Texture, /* SliceIndex = */ ArrayIndex - 1));
