[HWRT] Fix race condition between AddDynamicInstancesToRayTracingScene(...) allocating dynamic ranges in SBT and SetupRayTracingPipelineStatesAndSBT(...) resizing SBT.
- wait for task that runs AddDynamicInstancesToRayTracingScene(...) before executing SetupRayTracingPipelineStatesAndSBT(...)

#Synced-CL 52420933

diff --git a/Engine/Source/Runtime/Renderer/Private/DeferredShadingRenderer.cpp b/Engine/Source/Runtime/Renderer/Private/DeferredShadingRenderer.cpp
--- a/Engine/Source/Runtime/Renderer/Private/DeferredShadingRenderer.cpp
+++ b/Engine/Source/Runtime/Renderer/Private/DeferredShadingRenderer.cpp
@@ -1173,8 +1173,10 @@ void FDeferredShadingSceneRenderer::SetupRayTracingRenderingData(FRDGBuilder& Gr
 	{
 		SCOPE_CYCLE_COUNTER(STAT_WaitRayTracingDynamicBindings);
 
-		// need to wait for dynamic mesh batches tasks to finish before executing SetupRayTracingPipelineStatesAndSBT(...)
-		// since they can request new materials that need to be included in RTPSO 
+		// Need to wait for dynamic mesh batches tasks to finish before executing SetupRayTracingPipelineStatesAndSBT(...)
+		// since they can request new materials that need to be included in RTPSO.
+		// Also waits for dynamic range allocations in SBT to be complete
+		// since CheckPersistentRHI (inside SetupRayTracingPipelineStatesAndSBT) reads MaxNumDynamicGeometrySegments.
 		RayTracing::WaitForDynamicBindings(RayTracingGatherInstancesTaskData);
 	}
 
diff --git a/Engine/Source/Runtime/Renderer/Private/RayTracing/RayTracing.cpp b/Engine/Source/Runtime/Renderer/Private/RayTracing/RayTracing.cpp
--- a/Engine/Source/Runtime/Renderer/Private/RayTracing/RayTracing.cpp
+++ b/Engine/Source/Runtime/Renderer/Private/RayTracing/RayTracing.cpp
@@ -527,6 +527,9 @@ namespace RayTracing
 		UE::Tasks::FTaskEvent GatherDynamicRayTracingInstancesPrerequisites{ UE_SOURCE_LOCATION };
 		bool bGatherDynamicRayTracingInstancesPrerequisitesTriggered = false;
 
+		// This task adds dynamic instances and allocates dynamic ranges in the SBT.
+		UE::Tasks::FTask AddDynamicInstancesTask;
+
 		UE::Tasks::FTask VisibleRayTracingShaderBindingsFinalizeTask;
 		FRayTracingShaderBindingDataOneFrameArray VisibleShaderBindings;
 
@@ -2192,13 +2195,15 @@ namespace RayTracing
 		UE::Tasks::FTaskEvent AddInstancesTaskEvent{ UE_SOURCE_LOCATION };
 
 		// This adds final dynamic instances to the ray tracing scene and must be done before FRayTracingScene::BuildInitializationData().
-		AddInstancesTaskEvent.AddPrerequisites(TaskData.AddInstancesPipe.Launch(UE_SOURCE_LOCATION,
+		TaskData.AddDynamicInstancesTask = TaskData.AddInstancesPipe.Launch(UE_SOURCE_LOCATION,
 			[&TaskData, &RayTracingScene, &RayTracingSBT, SharedBufferGenerationID]
 			{
 				FTaskTagScope TaskTagScope(ETaskTag::EParallelRenderingThread);
 
 				AddDynamicInstancesToRayTracingScene(TaskData, RayTracingScene, RayTracingSBT, SharedBufferGenerationID);
-			}));
+			});
+
+		AddInstancesTaskEvent.AddPrerequisites(TaskData.AddDynamicInstancesTask);
 
 		for (FGatherInstancesViewTaskData& ViewTaskData : TaskData.ViewTaskDatas)
 		{
@@ -2303,6 +2308,10 @@ namespace RayTracing
 			UE::Tasks::Wait(View->AddDynamicRayTracingMeshBatchTaskList);
 			View->AddDynamicRayTracingMeshBatchTaskList.Empty();
 		}
+
+		// Wait for AddDynamicInstancesToRayTracingScene to complete since
+		// it also allocates dynamic ranges in the SBT (increasing MaxNumDynamicGeometrySegments).
+		TaskData.AddDynamicInstancesTask.Wait();
 	}
 
 	bool FinishGatherVisibleShaderBindings(FGatherInstancesTaskData& TaskData)
