When the Sky Atmosphere pass runs before the prepass, wait for FilterDynamicShadowsTask to complete to avoid a race condition and a use-after-free if one of atmosphere lights cast shadow on atmosphere

VisibleLightInfo.ShadowsToProject could be read on the render thread while being modified asynchronously by FilterDynamicShadowsTask.

#rnx

#Synced-CL 52463319

diff --git a/Engine/Source/Runtime/Renderer/Private/DeferredShadingRenderer.cpp b/Engine/Source/Runtime/Renderer/Private/DeferredShadingRenderer.cpp
--- a/Engine/Source/Runtime/Renderer/Private/DeferredShadingRenderer.cpp
+++ b/Engine/Source/Runtime/Renderer/Private/DeferredShadingRenderer.cpp
@@ -2418,7 +2418,9 @@ void FDeferredShadingSceneRenderer::Render(FRDGBuilder& GraphBuilder, const FSce
 	if (SkyAtmospherePassLocation == ESkyAtmospherePassLocation::BeforePrePass && bShouldRenderSkyAtmosphere)
 	{
 		// Generate the Sky/Atmosphere look up tables overlaping the pre-pass
-		RenderSkyAtmosphereLookUpTables(GraphBuilder, /* out */ SkyAtmospherePendingRDGResources);
+		// Pass DynamicShadowsTaskData pointer in case that task needs to be waited before accessing VisibleLightInfo.ShadowsToProject
+		// (i.e FilterDynamicShadows runs on a parallel task and writes ShadowsToProject concurrently)
+		RenderSkyAtmosphereLookUpTables(GraphBuilder, /* out */ SkyAtmospherePendingRDGResources, InitViewTaskDatas.DynamicShadows);
 	}
 
 	RenderWaterInfoTexture(GraphBuilder, *this, Scene);
diff --git a/Engine/Source/Runtime/Renderer/Private/SceneRendering.h b/Engine/Source/Runtime/Renderer/Private/SceneRendering.h
--- a/Engine/Source/Runtime/Renderer/Private/SceneRendering.h
+++ b/Engine/Source/Runtime/Renderer/Private/SceneRendering.h
@@ -2691,7 +2691,7 @@ protected:
 	void InitSkyAtmosphereForViews(FRHICommandListImmediate& RHICmdList, FRDGBuilder& GraphBuilder);
 	
 	/** Render the sky atmosphere look up table needed for this frame.*/
-	void RenderSkyAtmosphereLookUpTables(FRDGBuilder& GraphBuilder, class FSkyAtmospherePendingRDGResources& PendingRDGResources);
+	void RenderSkyAtmosphereLookUpTables(FRDGBuilder& GraphBuilder, class FSkyAtmospherePendingRDGResources& PendingRDGResources, FDynamicShadowsTaskData* FilterDynamicShadowTaskData = nullptr);
 
 	/** Render the sky atmosphere over the scene.*/
 	void RenderSkyAtmosphere(FRDGBuilder& GraphBuilder, const FMinimalSceneTextures& SceneTextures);
diff --git a/Engine/Source/Runtime/Renderer/Private/ShadowSetup.cpp b/Engine/Source/Runtime/Renderer/Private/ShadowSetup.cpp
--- a/Engine/Source/Runtime/Renderer/Private/ShadowSetup.cpp
+++ b/Engine/Source/Runtime/Renderer/Private/ShadowSetup.cpp
@@ -6589,6 +6589,14 @@ void FSceneRenderer::CreateDynamicShadows(FDynamicShadowsTaskData& TaskData)
 	}
 }
 
+void WaitForFilterDynamicShadowsTask(FDynamicShadowsTaskData* TaskData)
+{
+	if (TaskData)
+	{
+		TaskData->FilterDynamicShadowsTask.Wait();
+	}
+}
+
 void FSceneRenderer::FilterDynamicShadows(FDynamicShadowsTaskData& TaskData)
 {
 	TRACE_CPUPROFILER_EVENT_SCOPE(CreateFilterShadowPrimitives);
diff --git a/Engine/Source/Runtime/Renderer/Private/SkyAtmosphereRendering.cpp b/Engine/Source/Runtime/Renderer/Private/SkyAtmosphereRendering.cpp
--- a/Engine/Source/Runtime/Renderer/Private/SkyAtmosphereRendering.cpp
+++ b/Engine/Source/Runtime/Renderer/Private/SkyAtmosphereRendering.cpp
@@ -393,12 +393,23 @@ static bool VirtualShadowMapSamplingSupported(EShaderPlatform ShaderPlatform)
 }
 
 
-bool ShouldSkySampleAtmosphereLightsOpaqueShadow(const FScene& Scene, const FVisibleLightInfoArray& VisibleLightInfos, SkyAtmosphereLightShadowData& LightShadowData)
+bool ShouldSkySampleAtmosphereLightsOpaqueShadow(const FScene& Scene, const FVisibleLightInfoArray& VisibleLightInfos, SkyAtmosphereLightShadowData& LightShadowData, FDynamicShadowsTaskData* FilterDynamicShadowTaskData)
 {
 	LightShadowData.LightVolumetricShadowSceneinfo0 = Scene.AtmosphereLights[0];
 	LightShadowData.LightVolumetricShadowSceneinfo1 = Scene.AtmosphereLights[1];
+
+	bool bCastShadowOnAtmosphereLight0 = LightShadowData.LightVolumetricShadowSceneinfo0 && LightShadowData.LightVolumetricShadowSceneinfo0->Proxy && LightShadowData.LightVolumetricShadowSceneinfo0->Proxy->GetCastShadowsOnAtmosphere();
+	bool bCastShadowOnAtmosphereLight1 = LightShadowData.LightVolumetricShadowSceneinfo1 && LightShadowData.LightVolumetricShadowSceneinfo1->Proxy && LightShadowData.LightVolumetricShadowSceneinfo1->Proxy->GetCastShadowsOnAtmosphere();
 	
-	if (LightShadowData.LightVolumetricShadowSceneinfo0 && LightShadowData.LightVolumetricShadowSceneinfo0->Proxy && LightShadowData.LightVolumetricShadowSceneinfo0->Proxy->GetCastShadowsOnAtmosphere())
+	if (FilterDynamicShadowTaskData != nullptr && (bCastShadowOnAtmosphereLight0 || bCastShadowOnAtmosphereLight1))
+	{
+		// FilterDynamicShadows runs on a parallel task and writes = VisibleLightInfo.ShadowsToProject concurrently
+		// Wait before accessing = VisibleLightInfo.ShadowsToProject in GetFirstWholeSceneShadowMap.
+		extern void WaitForFilterDynamicShadowsTask(FDynamicShadowsTaskData* TaskData);
+		WaitForFilterDynamicShadowsTask(FilterDynamicShadowTaskData);
+	}
+	
+	if (bCastShadowOnAtmosphereLight0)
 	{
 		const FVisibleLightInfo& Light0 = VisibleLightInfos[LightShadowData.LightVolumetricShadowSceneinfo0->Id];
 		LightShadowData.ProjectedShadowInfo0 = GetFirstWholeSceneShadowMap(Light0);	
@@ -413,7 +424,7 @@ bool ShouldSkySampleAtmosphereLightsOpaqueShadow(const FScene& Scene, const FVis
 			LightShadowData.VirtualShadowMapId0 = INDEX_NONE;
 		}
 	}
-	if (LightShadowData.LightVolumetricShadowSceneinfo1 && LightShadowData.LightVolumetricShadowSceneinfo1->Proxy && LightShadowData.LightVolumetricShadowSceneinfo1->Proxy->GetCastShadowsOnAtmosphere())
+	if (bCastShadowOnAtmosphereLight1)
 	{
 		const FVisibleLightInfo& Light1 = VisibleLightInfos[LightShadowData.LightVolumetricShadowSceneinfo1->Id];
 		LightShadowData.ProjectedShadowInfo1 = GetFirstWholeSceneShadowMap(Light1);
@@ -1464,7 +1475,7 @@ static void SetupSkyAtmosphereInternalCommonParameters(
 	InternalCommonParameters.CameraAerialPerspectiveVolumeSizeAndInvSize = GetSizeAndInvSize(CameraAerialPerspectiveVolumeScreenResolution, CameraAerialPerspectiveVolumeScreenResolution);
 }
 
-void FSceneRenderer::RenderSkyAtmosphereLookUpTables(FRDGBuilder& GraphBuilder, class FSkyAtmospherePendingRDGResources& PendingRDGResources)
+void FSceneRenderer::RenderSkyAtmosphereLookUpTables(FRDGBuilder& GraphBuilder, class FSkyAtmospherePendingRDGResources& PendingRDGResources, FDynamicShadowsTaskData* FilterDynamicShadowTaskData)
 {
 	check(ShouldRenderSkyAtmosphere(Scene, ViewFamily.EngineShowFlags)); // This should not be called if we should not render SkyAtmosphere
 
@@ -1605,7 +1616,7 @@ void FSceneRenderer::RenderSkyAtmosphereLookUpTables(FRDGBuilder& GraphBuilder,
 	}
 
 	SkyAtmosphereLightShadowData LightShadowData;
-	const bool bShouldSampleOpaqueShadow = ShouldSkySampleAtmosphereLightsOpaqueShadow(*Scene, VisibleLightInfos, LightShadowData);
+	const bool bShouldSampleOpaqueShadow = ShouldSkySampleAtmosphereLightsOpaqueShadow(*Scene, VisibleLightInfos, LightShadowData, FilterDynamicShadowTaskData);
 
 	FLightSceneProxy* AtmosphereLight0Proxy = Scene->AtmosphereLights[0] ? Scene->AtmosphereLights[0]->Proxy : nullptr;
 	FLightSceneProxy* AtmosphereLight1Proxy = Scene->AtmosphereLights[1] ? Scene->AtmosphereLights[1]->Proxy : nullptr;
@@ -2393,4 +2404,3 @@ FScreenPassTexture AddSkyAtmosphereDebugPasses(FRDGBuilder& GraphBuilder, FScene
 	return MoveTemp(ScreenPassSceneColor);
 }
 
-
diff --git a/Engine/Source/Runtime/Renderer/Private/SkyAtmosphereRendering.h b/Engine/Source/Runtime/Renderer/Private/SkyAtmosphereRendering.h
--- a/Engine/Source/Runtime/Renderer/Private/SkyAtmosphereRendering.h
+++ b/Engine/Source/Runtime/Renderer/Private/SkyAtmosphereRendering.h
@@ -244,7 +244,7 @@ struct SkyAtmosphereLightShadowData
 	int VirtualShadowMapId0 = INDEX_NONE;
 	int VirtualShadowMapId1 = INDEX_NONE;
 };
-bool ShouldSkySampleAtmosphereLightsOpaqueShadow(const FScene& Scene, const FVisibleLightInfoArray& VisibleLightInfos, SkyAtmosphereLightShadowData& LightShadowData);
+bool ShouldSkySampleAtmosphereLightsOpaqueShadow(const FScene& Scene, const FVisibleLightInfoArray& VisibleLightInfos, SkyAtmosphereLightShadowData& LightShadowData, struct FDynamicShadowsTaskData* FilterDynamicShadowTaskData = nullptr);
 void GetSkyAtmosphereLightsUniformBuffers(
 	FRDGBuilder& GraphBuilder,
 	TRDGUniformBufferRef<FVolumeShadowingShaderParametersGlobal0>& OutLightShadowShaderParams0UniformBuffer,
