Multiply paper white ratio against scene color rather than display color

Also clean up logic around inner LUT.

#Synced-CL 52473974

diff --git a/Engine/Shaders/Private/PostProcessCombineLUTsInner.usf b/Engine/Shaders/Private/PostProcessCombineLUTsInner.usf
--- a/Engine/Shaders/Private/PostProcessCombineLUTsInner.usf
+++ b/Engine/Shaders/Private/PostProcessCombineLUTsInner.usf
@@ -45,8 +45,7 @@ uint OutputDevice;
 uint OutputGamut;
 float OutputMaxLuminance;
 
-float HDRPaperWhiteRatio;
-float InvHDRPaperWhiteRatio;
+float TotalSceneColorMultiplier;
 
 uint GetOutputDevice()
 {
@@ -61,15 +60,13 @@ float3 ComputeOutputACES(
 	float3 GradedColor,
 	float3x3 LimitingRgbToXYZ,
 	float3x3 LimitingXYZToRgb,
-	float3x3 AP1_2_Output,
-	float bSDR)
+	float3x3 AP1_2_Output)
 {
 #if USE_ACES_2
-	float PaperMaxLuminance = lerp(OutputMaxLuminance * InvHDRPaperWhiteRatio, OutputMaxLuminance, bSDR);
 	float3 ODTColor = ACESOutputTransform(
-		GradedColor * ACESSceneColorMultiplier,
+		GradedColor * TotalSceneColorMultiplier,
 		(float3x3)WorkingColorSpace.ToAP0,
-		PaperMaxLuminance,
+		OutputMaxLuminance,
 		ACESReachTable,
 		ACESGamutTable,
 		ACESGammaTable,
@@ -77,7 +74,6 @@ float3 ComputeOutputACES(
 		LimitingRgbToXYZ,
 		LimitingXYZToRgb,
 		GetColorSpace(OutputGamut));
-	ODTColor = lerp(ODTColor * HDRPaperWhiteRatio, ODTColor, bSDR);
 	return ODTColor;
 #else
 	FACESTonemapParams AcesParams = ComputeACESTonemapParams(ACESMinMaxData, ACESMidData, ACESCoefsLow_0, ACESCoefsHigh_0, ACESCoefsLow_4, ACESCoefsHigh_4, ACESSceneColorMultiplier, ACESGamutCompression);
@@ -126,19 +122,18 @@ float4 MethodLUTsCommon(float3 LUTEncodedColor)
 	}
 
 	const float3x3 AP1_2_Output = OutputGamutMappingMatrix(OutputGamut);
-	float bSDR = step(OutputMaxLuminance, LinearToNitsScale);
-	float3 OutputGamutColor = ComputeOutputACES(GradedColor, GetLimitingRgbToXYZ(), GetLimitingXYZToRgb(), AP1_2_Output, bSDR);
+	float3 OutputGamutColor = ComputeOutputACES(GradedColor, GetLimitingRgbToXYZ(), GetLimitingXYZToRgb(), AP1_2_Output);
 
 	BRANCH
-	if (bSDR)
+	if (OutputMaxLuminance > LinearToNitsScale)
 	{
-		OutputGamutColor = saturate(OutputGamutColor * LinearToNitsScaleInverse);
-		OutputGamutColor = LinearToSrgb(OutputGamutColor);
+		OutputGamutColor = clamp(OutputGamutColor, 0., UE_PQ_WHITE_NITS);
+		OutputGamutColor = LinearToST2084(OutputGamutColor);
 	}
 	else
 	{
-		OutputGamutColor = clamp(OutputGamutColor, 0., 10000.);
-		OutputGamutColor = LinearToST2084(OutputGamutColor);
+		OutputGamutColor = saturate(OutputGamutColor * LinearToNitsScaleInverse);
+		OutputGamutColor = LinearToSrgb(OutputGamutColor);
 	}
 
 	return float4(OutputGamutColor, 0.);
diff --git a/Engine/Shaders/Private/PostProcessDeviceEncodingOnly.usf b/Engine/Shaders/Private/PostProcessDeviceEncodingOnly.usf
--- a/Engine/Shaders/Private/PostProcessDeviceEncodingOnly.usf
+++ b/Engine/Shaders/Private/PostProcessDeviceEncodingOnly.usf
@@ -27,8 +27,7 @@ uint OutputDevice;
 uint OutputGamut;
 float OutputMaxLuminance;
 
-float HDRPaperWhiteRatio;
-float InvHDRPaperWhiteRatio;
+float TotalSceneColorMultiplier;
 
 float4 ACESMinMaxData;
 float4 ACESMidData;
@@ -113,9 +112,9 @@ float4 DeviceEncodingOnlyCommonPS(
 		// TODO: The display limiting color space should be user-exposed as part of display characterization.
 		FColorSpace LimitingColorSpace = GetColorSpace(TONEMAPPER_GAMUT_Rec2020_D65);
 		float3 ODTColor = ACESOutputTransform(
-			FilmColor * ACESSceneColorMultiplier,
+			FilmColor * TotalSceneColorMultiplier,
 			(float3x3)WorkingColorSpace.ToAP0,
-			OutputMaxLuminance * InvHDRPaperWhiteRatio,
+			OutputMaxLuminance,
 			ACESReachTable,
 			ACESGamutTable,
 			ACESGammaTable,
@@ -123,7 +122,6 @@ float4 DeviceEncodingOnlyCommonPS(
 			LimitingColorSpace.RGBtoXYZ,
 			LimitingColorSpace.XYZtoRGB,
 			EncodingColorSpace);
-		ODTColor *= HDRPaperWhiteRatio;
 #else
 		FACESTonemapParams AcesParams = ComputeACESTonemapParams(ACESMinMaxData, ACESMidData, ACESCoefsLow_0, ACESCoefsHigh_0, ACESCoefsLow_4, ACESCoefsHigh_4, ACESSceneColorMultiplier, ACESGamutCompression);
 		float3 ODTColor = ACESOutputTransform( FilmColor, (float3x3)WorkingColorSpace.ToAP0, AcesParams);
@@ -141,9 +139,9 @@ float4 DeviceEncodingOnlyCommonPS(
 		// TODO: The display limiting color space should be user-exposed as part of display characterization.
 		FColorSpace LimitingColorSpace = GetColorSpace(TONEMAPPER_GAMUT_Rec2020_D65);
 		float3 ODTColor = ACESOutputTransform(
-			FilmColor * ACESSceneColorMultiplier,
+			FilmColor * TotalSceneColorMultiplier,
 			(float3x3)WorkingColorSpace.ToAP0,
-			OutputMaxLuminance * InvHDRPaperWhiteRatio,
+			OutputMaxLuminance,
 			ACESReachTable,
 			ACESGamutTable,
 			ACESGammaTable,
@@ -151,7 +149,6 @@ float4 DeviceEncodingOnlyCommonPS(
 			LimitingColorSpace.RGBtoXYZ,
 			LimitingColorSpace.XYZtoRGB,
 			GetColorSpace(TONEMAPPER_GAMUT_sRGB_D65));
-		ODTColor *= HDRPaperWhiteRatio;
 		OutDeviceColor = ODTColor / UE_SCRGB_WHITE_NITS;
 #else
 		FACESTonemapParams AcesParams = ComputeACESTonemapParams(ACESMinMaxData, ACESMidData, ACESCoefsLow_0, ACESCoefsHigh_0, ACESCoefsLow_4, ACESCoefsHigh_4, ACESSceneColorMultiplier, ACESGamutCompression);
diff --git a/Engine/Source/Runtime/Renderer/Private/PostProcess/PostProcessCombineLUTs.cpp b/Engine/Source/Runtime/Renderer/Private/PostProcess/PostProcessCombineLUTs.cpp
--- a/Engine/Source/Runtime/Renderer/Private/PostProcess/PostProcessCombineLUTs.cpp
+++ b/Engine/Source/Runtime/Renderer/Private/PostProcess/PostProcessCombineLUTs.cpp
@@ -277,8 +277,7 @@ BEGIN_SHADER_PARAMETER_STRUCT(FInnerLUTParameters, )
 	SHADER_PARAMETER(FVector3f, LimitingXYZToRgb0)
 	SHADER_PARAMETER(FVector3f, LimitingXYZToRgb1)
 	SHADER_PARAMETER(FVector3f, LimitingXYZToRgb2)
-	SHADER_PARAMETER(float, HDRPaperWhiteRatio)
-	SHADER_PARAMETER(float, InvHDRPaperWhiteRatio)
+	SHADER_PARAMETER(float, TotalSceneColorMultiplier)
 	SHADER_PARAMETER(uint32, TonemappingMethod)
 	SHADER_PARAMETER_STRUCT_INCLUDE(FTonemapperOutputDeviceParameters, OutputDevice)
 END_SHADER_PARAMETER_STRUCT()
@@ -626,8 +625,7 @@ struct FCachedInnerLUTSettings
 		);
 
 		const float HDRPaperWhiteRatio = ViewFamily.RenderTarget->GetHDRPaperWhiteInNits() / 203.f;
-		UPDATE_CACHE_SETTINGS(Parameters.HDRPaperWhiteRatio, HDRPaperWhiteRatio, bHasChanged);
-		UPDATE_CACHE_SETTINGS(Parameters.InvHDRPaperWhiteRatio, 1.f / HDRPaperWhiteRatio, bHasChanged);
+		UPDATE_CACHE_SETTINGS(Parameters.TotalSceneColorMultiplier, HDRPaperWhiteRatio * TonemapperParams.ACESSceneColorMultiplier, bHasChanged);
 
 		UPDATE_CACHE_SETTINGS(Parameters.TonemappingMethod, static_cast<uint32>(GetTonemappingMethod(Settings)), bHasChanged);
 	}
@@ -788,24 +786,25 @@ FRDGTextureRef AddCombineLUTPass(FRDGBuilder& GraphBuilder, const FViewInfo& Vie
 
 	check(OutputTexture);
 
+	const bool bOutputDeviceSRGB = (CachedLUTSettings.Parameters.OutputDevice.OutputDevice == (uint32)EDisplayOutputFormat::SDR_sRGB);
+	FRDGTextureRef InnerTexture;
+
 	const ETonemappingMethod TonemappingMethod = static_cast<ETonemappingMethod>(CachedInnerLUTSettings.Parameters.TonemappingMethod);
 	const bool bNeedsInnerTonemappingLUT =
 		(TonemappingMethod == ETonemappingMethod::StandardACES) ||
 		((TonemappingMethod == ETonemappingMethod::Filmic) && CachedInnerLUTSettings.Parameters.OutputDevice.OutputMaxLuminance > 100.f);
-	const int32 InnerLUTTextureSize = bNeedsInnerTonemappingLUT ? InnerLUTSize : 1;
-	FRDGTextureRef InnerTexture = TryRegisterExternalTexture(GraphBuilder,
-		View.GetInnerTonemappingLUT(GraphBuilder.RHICmdList, InnerLUTTextureSize, bUseFloatOutput));
-
-	const bool bOutputDeviceSRGB = (CachedLUTSettings.Parameters.OutputDevice.OutputDevice == (uint32)EDisplayOutputFormat::SDR_sRGB);
-	if (bHasInnerChanged || !InnerTexture || CVarLUTInnerUpdateEveryFrame.GetValueOnRenderThread())
+	if (bNeedsInnerTonemappingLUT)
 	{
-		if (bNeedsInnerTonemappingLUT)
+		InnerTexture = TryRegisterExternalTexture(GraphBuilder,
+			View.GetInnerTonemappingLUT(GraphBuilder.RHICmdList, InnerLUTSize, bUseFloatOutput));
+
+		if (bHasInnerChanged || !InnerTexture || CVarLUTInnerUpdateEveryFrame.GetValueOnRenderThread())
 		{
 			// View doesn't support a persistent LUT, so create a temporary one.
 			if (!InnerTexture)
 			{
 				InnerTexture = GraphBuilder.CreateTexture(
-					Translate(FSceneViewState::CreateLUTRenderTarget(InnerLUTTextureSize, bUseFloatOutput, TEXT("CombineLUTs_Inner"))),
+					Translate(FSceneViewState::CreateLUTRenderTarget(InnerLUTSize, bUseFloatOutput, TEXT("CombineLUTs_Inner"))),
 					TEXT("CombineLUTs_Inner"));
 			}
 
@@ -832,6 +831,11 @@ FRDGTextureRef AddCombineLUTPass(FRDGBuilder& GraphBuilder, const FViewInfo& Vie
 				FIntVector(GroupSize, GroupSize, GroupSize));
 		}
 	}
+	else
+	{
+		View.ClearInnerTonemappingLUT();
+		InnerTexture = RegisterExternalTexture(GraphBuilder, GBlackVolumeTexture->GetTextureRHI(), TEXT("CombineLUTs_Inner_Fallback"));
+	}
 
 	FLUTBlenderCS::FPermutationDomain PermutationVector;
 	PermutationVector.Set<FLUTBlenderCS::FBlendCount>(LocalCount);
@@ -839,9 +843,7 @@ FRDGTextureRef AddCombineLUTPass(FRDGBuilder& GraphBuilder, const FViewInfo& Vie
 
 	FLUTBlenderCS::FParameters* PassParameters = GraphBuilder.AllocParameters<FLUTBlenderCS::FParameters>();
 	PassParameters->CombineLUT = CachedLUTSettings.Parameters;
-	PassParameters->InnerLUT = bNeedsInnerTonemappingLUT
-		? InnerTexture
-		: RegisterExternalTexture(GraphBuilder, GBlackVolumeTexture->GetTextureRHI(), TEXT("CombineLUTs_Inner_Fallback"));
+	PassParameters->InnerLUT = InnerTexture;
 	PassParameters->InnerLUTSampler = TStaticSamplerState<SF_Bilinear, AM_Clamp, AM_Clamp, AM_Clamp>::GetRHI();
 	PassParameters->InvLUTSizeMinusOne = 1.f / (LUTSize - 1);
 	PassParameters->RWOutputTexture = GraphBuilder.CreateUAV(OutputTexture);
diff --git a/Engine/Source/Runtime/Renderer/Private/PostProcess/PostProcessDeviceEncodingOnly.cpp b/Engine/Source/Runtime/Renderer/Private/PostProcess/PostProcessDeviceEncodingOnly.cpp
--- a/Engine/Source/Runtime/Renderer/Private/PostProcess/PostProcessDeviceEncodingOnly.cpp
+++ b/Engine/Source/Runtime/Renderer/Private/PostProcess/PostProcessDeviceEncodingOnly.cpp
@@ -113,8 +113,7 @@ BEGIN_SHADER_PARAMETER_STRUCT(FDeviceEncodingOnlyParameters, )
 	SHADER_PARAMETER_RDG_TEXTURE(Texture2D, ColorTexture)
 	SHADER_PARAMETER_SAMPLER(SamplerState, ColorSampler)
 	SHADER_PARAMETER(float, EditorNITLevel)
-	SHADER_PARAMETER(float, HDRPaperWhiteRatio)
-	SHADER_PARAMETER(float, InvHDRPaperWhiteRatio)
+	SHADER_PARAMETER(float, TotalSceneColorMultiplier)
 END_SHADER_PARAMETER_STRUCT()
 
 class FDeviceEncodingOnlyPS : public FGlobalShader
@@ -252,8 +251,7 @@ FScreenPassTexture AddDeviceEncodingOnlyPass(FRDGBuilder& GraphBuilder, const FV
 	CommonParameters.EditorNITLevel = EditorNITLevel;
 
 	const float HDRPaperWhiteRatio = ViewFamily.RenderTarget->GetHDRPaperWhiteInNits() / 203.f;
-	CommonParameters.HDRPaperWhiteRatio = HDRPaperWhiteRatio;
-	CommonParameters.InvHDRPaperWhiteRatio = 1.f / HDRPaperWhiteRatio;
+	CommonParameters.TotalSceneColorMultiplier = HDRPaperWhiteRatio * CommonParameters.ACESTonemapParameters.ACESSceneColorMultiplier;
 
 	// Generate permutation vector for the desktop tonemapper.
 	DeviceEncodingOnlyPermutation::FDesktopDomain DesktopPermutationVector;
diff --git a/Engine/Source/Runtime/Renderer/Private/SceneRendering.cpp b/Engine/Source/Runtime/Renderer/Private/SceneRendering.cpp
--- a/Engine/Source/Runtime/Renderer/Private/SceneRendering.cpp
+++ b/Engine/Source/Runtime/Renderer/Private/SceneRendering.cpp
@@ -2656,6 +2656,15 @@ IPooledRenderTarget* FViewInfo::GetInnerTonemappingLUT(FRHICommandList& RHICmdLi
 	return nullptr;
 }
 
+void FViewInfo::ClearInnerTonemappingLUT() const
+{
+	FSceneViewState* EffectiveViewState = GetEyeAdaptationViewState();
+	if (EffectiveViewState)
+	{
+		EffectiveViewState->ClearInnerTonemappingLUT();
+	}
+}
+
 bool FViewInfo::RequiresDebugMaterials() const
 {
 	// We can only use debug materials in ODSC environments.
diff --git a/Engine/Source/Runtime/Renderer/Private/SceneRendering.h b/Engine/Source/Runtime/Renderer/Private/SceneRendering.h
--- a/Engine/Source/Runtime/Renderer/Private/SceneRendering.h
+++ b/Engine/Source/Runtime/Renderer/Private/SceneRendering.h
@@ -1835,6 +1835,8 @@ public:
 	/** Gets the other rendertarget that will be populated by CombineLUTS post process */
 	IPooledRenderTarget* GetInnerTonemappingLUT(FRHICommandList& RHICmdList, const int32 LUTSize, const bool bNeedFloatOutput) const;
 
+	void ClearInnerTonemappingLUT() const;
+
 	bool IsFirstInFamily() const
 	{
 		return Family->Views[0] == this;
diff --git a/Engine/Source/Runtime/Renderer/Private/SceneViewState.h b/Engine/Source/Runtime/Renderer/Private/SceneViewState.h
--- a/Engine/Source/Runtime/Renderer/Private/SceneViewState.h
+++ b/Engine/Source/Runtime/Renderer/Private/SceneViewState.h
@@ -681,6 +681,11 @@ public:
 		return CombinedLUTRenderTarget.GetReference();
 	}
 
+	void ClearInnerTonemappingLUT()
+	{
+		InnerLUTRenderTarget.SafeRelease();
+	}
+
 	FSubstrateViewDebugData& GetSubstrateViewDebugData()
 	{
 		return SubstrateViewDebugData;
