CopyTextureRegionChecked crash fix.  Light function atlas debug visualization mode asserts due to AddDrawTexturePass attempting to use a copy operation on the 1x1 mid grey dummy texture used for the initial clear of the visualization atlas, because the wrong source size for the dummy texture is passed in (the atlas size rather than the actual 1x1 size).

#rnx

#Synced-CL 52460536

diff --git a/Engine/Source/Runtime/Renderer/Private/LightFunctionAtlas.cpp b/Engine/Source/Runtime/Renderer/Private/LightFunctionAtlas.cpp
--- a/Engine/Source/Runtime/Renderer/Private/LightFunctionAtlas.cpp
+++ b/Engine/Source/Runtime/Renderer/Private/LightFunctionAtlas.cpp
@@ -859,7 +859,7 @@ FScreenPassTexture FLightFunctionAtlas::AddDebugVisualizationPasses(FRDGBuilder&
 	const uint32 AtlasResolution		= AtlasSlotResolution * AtlasEdgeSize;
 
 	const FIntPoint SrcPoint			= FIntPoint::ZeroValue;
-	const FIntPoint SrcSize				= RDGAtlasTexture2D->Desc.Extent;
+	const FIntPoint SrcSize				= FIntPoint(1, 1);			// Size of MidGreyDummy, used for initial clear of the atlas to grey
 	const FIntPoint DstPoint			= FIntPoint(100, 100);
 	const FIntPoint DstSize				= FIntPoint(512, 512);
 
