Mobile view lighting channel support.  View lighting channels already worked on mobile for non-directional lights, but not directional lights due to their separate mobile specific code paths.  Need to check view lighting channel mask wherever lights are fetched from MobileDirectionalLights.

#Synced-CL 52530005

diff --git a/Engine/Source/Runtime/Renderer/Private/LocalFogVolumeRendering.cpp b/Engine/Source/Runtime/Renderer/Private/LocalFogVolumeRendering.cpp
--- a/Engine/Source/Runtime/Renderer/Private/LocalFogVolumeRendering.cpp
+++ b/Engine/Source/Runtime/Renderer/Private/LocalFogVolumeRendering.cpp
@@ -606,7 +606,7 @@ void CreateViewLocalFogVolumeBufferSRV(const FScene* Scene, FViewInfo& View, FRD
 		View.LocalFogVolumeViewData.UniformParametersStruct.LocalFogVolumeCommon.DirectionalLightDirection				= FVector3f::Zero();
 		for (uint32 ChannelIdx = 0; ChannelIdx < UE_ARRAY_COUNT(Scene->MobileDirectionalLights); ChannelIdx++)
 		{
-			FLightSceneInfo* Light = Scene->MobileDirectionalLights[ChannelIdx];
+			FLightSceneInfo* Light = Scene->GetMobileDirectionalLightForView(ChannelIdx, View);
 			if (Light != nullptr)
 			{
 				ApproximatePerPixelTransmittanceProxy = Light->Proxy;
diff --git a/Engine/Source/Runtime/Renderer/Private/MobileBasePassRendering.cpp b/Engine/Source/Runtime/Renderer/Private/MobileBasePassRendering.cpp
--- a/Engine/Source/Runtime/Renderer/Private/MobileBasePassRendering.cpp
+++ b/Engine/Source/Runtime/Renderer/Private/MobileBasePassRendering.cpp
@@ -420,7 +420,7 @@ void SetupMobileDirectionalLightUniformParameters(
 	FMobileDirectionalLightShaderParameters& Params)
 {
 	ERHIFeatureLevel::Type FeatureLevel = Scene.GetFeatureLevel();
-	FLightSceneInfo* Light = Scene.MobileDirectionalLights[ChannelIdx];
+	FLightSceneInfo* Light = Scene.GetMobileDirectionalLightForView(ChannelIdx, SceneView);
 	if (Light)
 	{
 		Params.DirectionalLightColor = Light->Proxy->GetSunIlluminanceAccountingForSkyAtmospherePerPixelTransmittance();
diff --git a/Engine/Source/Runtime/Renderer/Private/MobileDeferredShadingPass.cpp b/Engine/Source/Runtime/Renderer/Private/MobileDeferredShadingPass.cpp
--- a/Engine/Source/Runtime/Renderer/Private/MobileDeferredShadingPass.cpp
+++ b/Engine/Source/Runtime/Renderer/Private/MobileDeferredShadingPass.cpp
@@ -810,7 +810,7 @@ static int RenderDirectionalLights(
 	uint32 NumLights = 0;
 	for (uint32 ChannelIdx = 0; ChannelIdx < UE_ARRAY_COUNT(Scene.MobileDirectionalLights); ChannelIdx++)
 	{
-		NumLights += (Scene.MobileDirectionalLights[ChannelIdx] ? 1 : 0);
+		NumLights += (Scene.GetMobileDirectionalLightForView(ChannelIdx, View) ? 1 : 0);
 	}
 	// We can merge reflection and skylight pass with a sole directional light pass and if all primitives and the directional light use the default lighting channel
 	bool bPrimitivesUseLightingChannels = (View.bUsesLightingChannels && GMobileIgnoreDeferredShadingSkyLightChannels == 0);
@@ -818,7 +818,7 @@ static int RenderDirectionalLights(
 	bool bClusteredReflection = (View.NumBoxReflectionCaptures + View.NumSphereReflectionCaptures) > 0;
 	bool bSSR = MobileSSRQuality != EMobileSSRQuality::Disabled;
 
-	bool bInlineReflectionAndSky = (NumLights == 1) && !bPrimitivesUseLightingChannels && (Scene.MobileDirectionalLights[0] != nullptr) && !(bPlanarReflection || bClusteredReflection || bSSR);
+	bool bInlineReflectionAndSky = (NumLights == 1) && !bPrimitivesUseLightingChannels && (Scene.GetMobileDirectionalLightForView(0, View) != nullptr) && !(bPlanarReflection || bClusteredReflection || bSSR);
 	if (!bInlineReflectionAndSky)
 	{
 		RenderReflectionEnvironmentSkyLighting(RHICmdList, Scene, View, MobileSSRQuality, DynamicBentNormalAOTexture);
@@ -826,7 +826,7 @@ static int RenderDirectionalLights(
 
 	for (uint32 ChannelIdx = 0; ChannelIdx < UE_ARRAY_COUNT(Scene.MobileDirectionalLights); ChannelIdx++)
 	{
-		FLightSceneInfo* DirectionalLight = Scene.MobileDirectionalLights[ChannelIdx];
+		FLightSceneInfo* DirectionalLight = Scene.GetMobileDirectionalLightForView(ChannelIdx, View);
 		if (DirectionalLight)
 		{
 			RenderDirectionalLight(RHICmdList, Scene, View, DefaultLightMaterial, *DirectionalLight, ChannelIdx, bInlineReflectionAndSky, DynamicBentNormalAOTexture);
diff --git a/Engine/Source/Runtime/Renderer/Private/RendererScene.cpp b/Engine/Source/Runtime/Renderer/Private/RendererScene.cpp
--- a/Engine/Source/Runtime/Renderer/Private/RendererScene.cpp
+++ b/Engine/Source/Runtime/Renderer/Private/RendererScene.cpp
@@ -4062,6 +4062,23 @@ FRHIUniformBuffer* FScene::GetSpeedTreeUniformBuffer(const FVertexFactory* Verte
 	return nullptr;
 }
 
+FLightSceneInfo* FScene::GetMobileDirectionalLightForView(int32 ChannelIdx, const FViewInfo& View) const
+{
+	// TODO:  Should this call "ShouldRenderLight(View)" instead of just checking the view lighting channel mask?  Main
+	// difference is that ShouldRenderLight filters out dynamic lights when the view is set to render static only, typically
+	// for sky and reflection captures.  The immediate goal of this function is to make the view lighting channel mask work
+	// on mobile, and otherwise avoid changes to existing behavior, so we'll err on the side of caution for now.  If this
+	// gets changed, the other view lighting channel mask checks in FMobileSceneRenderer::InitDynamicShadows and
+	// FMobileSceneRenderer::BuildCSMVisibilityState should also be updated to use "ShouldRenderLight"!
+	if (MobileDirectionalLights[ChannelIdx] &&
+		MobileDirectionalLights[ChannelIdx]->Proxy &&
+		(MobileDirectionalLights[ChannelIdx]->Proxy->GetViewLightingChannelMask() & View.ViewLightingChannelMask))
+	{
+		return MobileDirectionalLights[ChannelIdx];
+	}
+	return nullptr;
+}
+
 /**
  * Retrieves the lights interacting with the passed in primitive and adds them to the out array.
  *
diff --git a/Engine/Source/Runtime/Renderer/Private/ScenePrivate.h b/Engine/Source/Runtime/Renderer/Private/ScenePrivate.h
--- a/Engine/Source/Runtime/Renderer/Private/ScenePrivate.h
+++ b/Engine/Source/Runtime/Renderer/Private/ScenePrivate.h
@@ -2578,6 +2578,9 @@ public:
 	 */
 	FSceneLightSceneInfoUpdateDelegate OnPostLightSceneInfoUpdate;
 
+	/** Returns the mobile directional light in the given channel index, if enabled for the given view */
+	FLightSceneInfo* GetMobileDirectionalLightForView(int32 ChannelIdx, const FViewInfo& View) const;
+
 	/**
 	 * Retrieves the lights interacting with the passed in primitive and adds them to the out array.
 	 * Render thread version of function.
diff --git a/Engine/Source/Runtime/Renderer/Private/ShadowSetupMobile.cpp b/Engine/Source/Runtime/Renderer/Private/ShadowSetupMobile.cpp
--- a/Engine/Source/Runtime/Renderer/Private/ShadowSetupMobile.cpp
+++ b/Engine/Source/Runtime/Renderer/Private/ShadowSetupMobile.cpp
@@ -299,8 +299,12 @@ FDynamicShadowsTaskData* FMobileSceneRenderer::InitDynamicShadows(FRDGBuilder& G
 			int32 PrimitiveCount = Scene->Primitives.Num();
 			for (auto& View : Views)
 			{
-				FMobileCSMSubjectPrimitives& MobileCSMSubjectPrimitives = View.VisibleLightInfos[MobileDirectionalLightSceneInfo->Id].MobileCSMSubjectPrimitives;
-				MobileCSMSubjectPrimitives.InitShadowSubjectPrimitives(PrimitiveCount);
+				// Mirrors logic from FScene::GetMobileDirectionalLightForView, see comment there for more info
+				if (MobileDirectionalLightSceneInfo->Proxy->GetViewLightingChannelMask() & View.ViewLightingChannelMask)
+				{
+					FMobileCSMSubjectPrimitives& MobileCSMSubjectPrimitives = View.VisibleLightInfos[MobileDirectionalLightSceneInfo->Id].MobileCSMSubjectPrimitives;
+					MobileCSMSubjectPrimitives.InitShadowSubjectPrimitives(PrimitiveCount);
+				}
 			}
 		}
 	}
@@ -423,6 +427,12 @@ void FMobileSceneRenderer::BuildCSMVisibilityState(FLightSceneInfo* LightSceneIn
 			bool bStaticCSMReceiversFound = false;
 			FViewInfo& View = Views[ViewIndex];
 
+			// Mirrors logic from FScene::GetMobileDirectionalLightForView, see comment there for more info
+			if (!(LightSceneInfo->Proxy->GetViewLightingChannelMask() & View.ViewLightingChannelMask))
+			{
+				continue;
+			}
+
 			FProjectedShadowInfo SingleCascadeInfo;
 			if (BuildSingleCascadeShadowInfo(View, VisibleLightInfos, LightSceneInfo, SingleCascadeInfo) == false)
 			{
diff --git a/Engine/Source/Runtime/Renderer/Private/VolumetricFog.cpp b/Engine/Source/Runtime/Renderer/Private/VolumetricFog.cpp
--- a/Engine/Source/Runtime/Renderer/Private/VolumetricFog.cpp
+++ b/Engine/Source/Runtime/Renderer/Private/VolumetricFog.cpp
@@ -1866,7 +1866,7 @@ void FSceneRenderer::ComputeVolumetricFog(FRDGBuilder& GraphBuilder,
 			PassParameters->MobileHasDirectionalLight		= 0;
 			for (uint32 ChannelIdx = 0; ChannelIdx < UE_ARRAY_COUNT(Scene->MobileDirectionalLights); ChannelIdx++)
 			{
-				FLightSceneInfo* Light = Scene->MobileDirectionalLights[ChannelIdx];
+				FLightSceneInfo* Light = Scene->GetMobileDirectionalLightForView(ChannelIdx, View);
 				if (Light != nullptr)
 				{
 					PassParameters->MobileDirectionalLightColor		= FVector3f(Light->Proxy->GetSunIlluminanceAccountingForSkyAtmospherePerPixelTransmittance() * Light->Proxy->GetVolumetricScatteringIntensity());
