Fix SelectionOutline/Gizmo being invisible in Buffer visualization mode.

Buffer visualization materials now capture output into a VisualizeBufferContext at their respective pipeline stages, and a final AddVisualizeGBufferPass composites the
captured buffer over SceneColor. This pass was placed after SelectionOutline/EditorPrimitive, overwriting them.

#Synced-CL 52496820

diff --git a/Engine/Source/Runtime/Renderer/Private/PostProcess/PostProcessing.cpp b/Engine/Source/Runtime/Renderer/Private/PostProcess/PostProcessing.cpp
--- a/Engine/Source/Runtime/Renderer/Private/PostProcess/PostProcessing.cpp
+++ b/Engine/Source/Runtime/Renderer/Private/PostProcess/PostProcessing.cpp
@@ -527,6 +527,7 @@ void AddPostProcessingPasses(
 		VisualizeLocalFogVolumes,
 		VisualizeLevelInstance,
 		VisualizeVirtualShadowMaps_PreEditorPrimitives,
+		VisualizeGBuffer,
 		ChannelMask,
 		SelectionOutline,
 		EditorPrimitive,
@@ -536,7 +537,6 @@ void AddPostProcessingPasses(
 		VisualizeShadingModels,
 		VisualizeGBufferHints,
 		VisualizeSubsurface,
-		VisualizeGBuffer,
 		VisualizeLumenSceneOverview,
 		VisualizeMegaLightsOverview,
 		VisualizeHDR,
@@ -592,6 +592,7 @@ void AddPostProcessingPasses(
 		TEXT("VisualizeLocalFogVolumes"),
 		TEXT("VisualizeLevelInstance"),
 		TEXT("VisualizeVirtualShadowMaps_PreEditorPrimitives"),
+		TEXT("VisualizeGBuffer"),
 		TEXT("ChannelMask"),
 		TEXT("SelectionOutline"),
 		TEXT("EditorPrimitive"),
@@ -601,7 +602,6 @@ void AddPostProcessingPasses(
 		TEXT("VisualizeShadingModels"),
 		TEXT("VisualizeGBufferHints"),
 		TEXT("VisualizeSubsurface"),
-		TEXT("VisualizeGBuffer"),
 		TEXT("VisualizeLumenSceneOverview"),
 		TEXT("VisualizeMegaLightsOverview"),
 		TEXT("VisualizeHDR"),
@@ -1908,6 +1908,16 @@ void AddPostProcessingPasses(
 	}
 #endif
 
+	if (PassSequence.IsEnabled(EPass::VisualizeGBuffer))
+	{
+		FVisualizeGBufferInputs PassInputs;
+		PassSequence.AcceptOverrideIfLastPass(EPass::VisualizeGBuffer, PassInputs.OverrideOutput);
+		PassInputs.SceneColor = SceneColor;
+		PassInputs.bOutputInHDR = bOutputInHDR;
+
+		SceneColor = AddVisualizeGBufferPass(GraphBuilder, View, PassInputs, VisualizeBufferContext);
+	}
+
 #if WITH_EDITOR
 	if (PassSequence.IsEnabled(EPass::ChannelMask))
 	{
@@ -2023,16 +2033,6 @@ void AddPostProcessingPasses(
 		SceneColor = AddVisualizeSubsurfacePass(GraphBuilder, View, PassInputs);
 	}
 
-	if (PassSequence.IsEnabled(EPass::VisualizeGBuffer))
-	{
-		FVisualizeGBufferInputs PassInputs;
-		PassSequence.AcceptOverrideIfLastPass(EPass::VisualizeGBuffer, PassInputs.OverrideOutput);
-		PassInputs.SceneColor = SceneColor;
-		PassInputs.bOutputInHDR = bOutputInHDR;
-
-		SceneColor = AddVisualizeGBufferPass(GraphBuilder, View, PassInputs, VisualizeBufferContext);
-	}
-
 	if (PassSequence.IsEnabled(EPass::VisualizeLumenSceneOverview))
 	{
 		FVisualizeLumenSceneInputs PassInputs;
@@ -2737,6 +2737,7 @@ void AddMobilePostProcessingPasses(FRDGBuilder& GraphBuilder, FScene* Scene, con
 		PostProcessMaterialAfterTonemapping,
 		HighResolutionScreenshotMask,
 		VisualizeSkyLightIlluminanceMeter,
+		VisualizeGBuffer,
 		ChannelMask,
 		SelectionOutline,
 		EditorPrimitive,
@@ -2745,7 +2746,6 @@ void AddMobilePostProcessingPasses(FRDGBuilder& GraphBuilder, FScene* Scene, con
 		SecondaryUpscale,
 		Visualize,
 		VisualizeLightGrid,
-		VisualizeGBuffer,
 		HMDDistortion,
 		MAX
 	};
@@ -2785,6 +2785,7 @@ void AddMobilePostProcessingPasses(FRDGBuilder& GraphBuilder, FScene* Scene, con
 		TEXT("FXAA"),
 		TEXT("HighResolutionScreenshotMask"),
 		TEXT("VisualizeSkyLightIlluminanceMeter"),
+		TEXT("VisualizeGBuffer"),
 		TEXT("ChannelMask"),
 		TEXT("SelectionOutline"),
 		TEXT("EditorPrimitive"),
@@ -2793,7 +2794,6 @@ void AddMobilePostProcessingPasses(FRDGBuilder& GraphBuilder, FScene* Scene, con
 		TEXT("SecondaryUpscale"),
 		TEXT("Visualize"),
 		TEXT("VisualizeLightGrid"),
-		TEXT("VisualizeGBuffer"),
 		TEXT("HMDDistortion")
 	};
 
@@ -3654,7 +3654,19 @@ void AddMobilePostProcessingPasses(FRDGBuilder& GraphBuilder, FScene* Scene, con
 		PassSequence.AcceptOverrideIfLastPass(EPass::VisualizeSkyLightIlluminanceMeter, OverrideOutput);
 		SceneColor = ProcessAndRenderIlluminanceMeter(GraphBuilder, View, SceneColor);
 	}
+#endif
 
+	if (PassSequence.IsEnabled(EPass::VisualizeGBuffer))
+	{
+		FVisualizeGBufferInputs PassInputs;
+		PassSequence.AcceptOverrideIfLastPass(EPass::VisualizeGBuffer, PassInputs.OverrideOutput);
+		PassInputs.SceneColor = SceneColor;
+		PassInputs.bOutputInHDR = bOutputInHDR;
+
+		SceneColor = AddVisualizeGBufferPass(GraphBuilder, View, PassInputs, VisualizeBufferContext);
+	}
+
+#if WITH_EDITOR
 	if (PassSequence.IsEnabled(EPass::ChannelMask))
 	{
 		ChannelMask::FChannelMaskInputs PassInputs;
@@ -3816,16 +3828,6 @@ void AddMobilePostProcessingPasses(FRDGBuilder& GraphBuilder, FScene* Scene, con
 		}
 	}
 
-	if (PassSequence.IsEnabled(EPass::VisualizeGBuffer))
-	{
-		FVisualizeGBufferInputs PassInputs;
-		PassSequence.AcceptOverrideIfLastPass(EPass::VisualizeGBuffer, PassInputs.OverrideOutput);
-		PassInputs.SceneColor = SceneColor;
-		PassInputs.bOutputInHDR = bOutputInHDR;
-
-		SceneColor = AddVisualizeGBufferPass(GraphBuilder, View, PassInputs, VisualizeBufferContext);
-	}
-
 	if (PassSequence.IsEnabled(EPass::VisualizeLightGrid))
 	{
 		FScreenPassRenderTarget OverrideOutput;
