Refactored ASTP to support layers and blend spaces. Rather than use a parent / child hierarchy, there are now optional data assets for layer stacks, blend spaces, and sequences. This way authoring can still be heavily shared. Layers are authored at edit time and cannot be modified. Each track now contains N layers. Each layer on a track can be assigned a manual, auto play, or blend space, with transitions between them supported. The API now adds a LayerIndex in addition to TrackIndex. Each layer (except for the base layer 0) also supports mesh space rotation and bone masks. The old 3 pass implementation is now gone and the new single pass approach requires Optimize for Instancing enabled on skeletal meshes. The new bone limit is 1024 bones, or 256 with mesh space rotation. The CPU reference path has been removed entirely as with the complexity it no longer makes sense to maintain it.

Other minor changes:
- Removed the 'Enabled' checkbox from ASTP's, which are often accidentally unchecked leading to confusion.
- Added support for 16 bit sequence transform storage and unpacking in the shader for mobile platforms.
- Converted to structured buffers and refactored into a header / sample split. This improved performance on Switch for blends by nearly 20%.

#Synced-CL 52519763

diff --git a/Engine/Plugins/Experimental/DynamicWind/Source/DynamicWind/Private/DynamicWindData.cpp b/Engine/Plugins/Experimental/DynamicWind/Source/DynamicWind/Private/DynamicWindData.cpp
--- a/Engine/Plugins/Experimental/DynamicWind/Source/DynamicWind/Private/DynamicWindData.cpp
+++ b/Engine/Plugins/Experimental/DynamicWind/Source/DynamicWind/Private/DynamicWindData.cpp
@@ -18,7 +18,7 @@ class FSceneInterface;
 
 bool UDynamicWindData::IsEnabled() const
 {
-	return bEnabled;
+	return true;
 }
 
 const FGuid& UDynamicWindData::GetTransformProviderID() const
diff --git a/Engine/Plugins/Runtime/MassGameplay/Source/MassRepresentation/Private/MassVisualizationComponent.cpp b/Engine/Plugins/Runtime/MassGameplay/Source/MassRepresentation/Private/MassVisualizationComponent.cpp
--- a/Engine/Plugins/Runtime/MassGameplay/Source/MassRepresentation/Private/MassVisualizationComponent.cpp
+++ b/Engine/Plugins/Runtime/MassGameplay/Source/MassRepresentation/Private/MassVisualizationComponent.cpp
@@ -763,22 +763,23 @@ void UMassVisualizationComponent::EndVisualChanges()
 									if (bAlreadyPresent)
 									{
 										TrackId = TrackMap.GetByElementId(TrackElementId).Value;
-										if (ASTPDI->GetSequenceIndex(TrackId) != InstanceAnimationData[IDIndex].SequenceIndex)
+										if (ASTPDI->GetSequenceIndex(TrackId, 0) != InstanceAnimationData[IDIndex].SequenceIndex)
 										{
-											ASTPDI->UpdateAutoPlayTrack(TrackId, InstanceAnimationData[IDIndex]);
+											ASTPDI->SetAutoPlayData(TrackId, 0, InstanceAnimationData[IDIndex]);
 										}
 									}
 									else
 									{
-										TrackId = TrackMap.GetByElementId(TrackElementId).Value = ASTPDI->AllocateAutoPlayTrack(InstanceAnimationData[IDIndex]);
+										TrackId = TrackMap.GetByElementId(TrackElementId).Value = ASTPDI->AllocateTrack();
+										ASTPDI->SetAutoPlayData(TrackId, 0, InstanceAnimationData[IDIndex]);
 									}
 							
 									// Readback resolved animation state for UAF handoff
 									const TArray<FAnimSequenceTransformProviderSequence>& Sequences = ASTPDI->GetSequences();
 									FMassSkinnedMeshResolvedAnimState State = {
-										.SequenceIndex = ASTPDI->GetSequenceIndex(TrackId),
-										.Position = ASTPDI->GetPosition(TrackId),
-										.PlayRate = ASTPDI->GetPlayRate(TrackId),
+										.SequenceIndex = ASTPDI->GetSequenceIndex(TrackId, 0),
+										.Position = ASTPDI->GetPosition(TrackId, 0),
+										.PlayRate = ASTPDI->GetPlayRate(TrackId, 0),
 										.AnimSequence = Sequences.IsValidIndex(State.SequenceIndex) ? Sequences[State.SequenceIndex].Sequence.Get() : nullptr
 									};
 									bool bAnimStatePresent = false;
@@ -853,22 +854,23 @@ void UMassVisualizationComponent::EndVisualChanges()
 								if (bAlreadyPresent)
 								{
 									TrackId = TrackMap.GetByElementId(TrackElementId).Value;
-									if (ASTPDI->GetSequenceIndex(TrackId) != InstanceAnimationData[AddItem.IDIndex].SequenceIndex)
+									if (ASTPDI->GetSequenceIndex(TrackId, 0) != InstanceAnimationData[AddItem.IDIndex].SequenceIndex)
 									{
-										ASTPDI->UpdateAutoPlayTrack(TrackId, InstanceAnimationData[AddItem.IDIndex]);
+										ASTPDI->SetAutoPlayData(TrackId, 0, InstanceAnimationData[AddItem.IDIndex]);
 									}
 								}
 								else
 								{
-									TrackId = TrackMap.GetByElementId(TrackElementId).Value = ASTPDI->AllocateAutoPlayTrack(InstanceAnimationData[AddItem.IDIndex]);
+									TrackId = TrackMap.GetByElementId(TrackElementId).Value = ASTPDI->AllocateTrack();
+									ASTPDI->SetAutoPlayData(TrackId, 0, InstanceAnimationData[AddItem.IDIndex]);
 								}
 						
 								// Readback resolved animation state for UAF handoff
 								const TArray<FAnimSequenceTransformProviderSequence>& Sequences = ASTPDI->GetSequences();
 								FMassSkinnedMeshResolvedAnimState State = {
-									.SequenceIndex = ASTPDI->GetSequenceIndex(TrackId),
-									.Position = ASTPDI->GetPosition(TrackId),
-									.PlayRate = ASTPDI->GetPlayRate(TrackId),
+									.SequenceIndex = ASTPDI->GetSequenceIndex(TrackId, 0),
+									.Position = ASTPDI->GetPosition(TrackId, 0),
+									.PlayRate = ASTPDI->GetPlayRate(TrackId, 0),
 									.AnimSequence = Sequences.IsValidIndex(State.SequenceIndex) ? Sequences[State.SequenceIndex].Sequence.Get() : nullptr
 								};
 								bool bAnimStatePresent = false;
diff --git a/Engine/Shaders/Private/BoneTransform.ush b/Engine/Shaders/Private/BoneTransform.ush
--- a/Engine/Shaders/Private/BoneTransform.ush
+++ b/Engine/Shaders/Private/BoneTransform.ush
@@ -191,4 +191,15 @@ FBoneTransformWithUniformScale BlendBoneTransformsWithUniformScale(FBoneTransfor
 	Blended.Rotation = QuatSlerpApproximate(Atom1.Rotation, Atom2.Rotation, Alpha);
 	Blended.Scale    = lerp(Atom1.Scale, Atom2.Scale, Alpha);
 	return Blended;
+}
+
+FBoneTransformWithUniformScale ApplyAdditiveLocalSpace(FBoneTransformWithUniformScale Base, FBoneTransformWithUniformScale Additive, float Weight)
+{
+	FQuat AdditiveRotation = QuatSlerpApproximate(QuatIdentity(), Additive.Rotation, Weight);
+
+	FBoneTransformWithUniformScale Result;
+	Result.Rotation = QuatNormalize(QuatMultiply(AdditiveRotation, Base.Rotation));
+	Result.Location = Base.Location + Weight * Additive.Location;
+	Result.Scale    = Base.Scale * lerp(1.0, Additive.Scale, Weight);
+	return Result;
 }
\ No newline at end of file
diff --git a/Engine/Shaders/Private/Skinning/AnimSequenceTransformProvider.usf b/Engine/Shaders/Private/Skinning/AnimSequenceTransformProvider.usf
--- a/Engine/Shaders/Private/Skinning/AnimSequenceTransformProvider.usf
+++ b/Engine/Shaders/Private/Skinning/AnimSequenceTransformProvider.usf
@@ -8,12 +8,24 @@
 #include "/Engine/Private/Skinning/SkinningDataDecode.ush"
 #include "/Engine/Shared/SkinningDefinitions.h"
 
+// Allow register spilling for complex mesh-space permutations.
+#pragma warning(disable:7203)
+
 ByteAddressBuffer RetargetingDataBuffer;
-ByteAddressBuffer SequenceTransformBuffer;
-ByteAddressBuffer EvaluateBlockHeaderBuffer;
-RWByteAddressBuffer LocalSpaceTransformBufferUAV;
+ByteAddressBuffer BoneMapBuffer;
+StructuredBuffer<float4> SequenceTransformBuffer;
+StructuredBuffer<FEvaluateLayerHeader> LayerHeaderBuffer;
+StructuredBuffer<FSamplePoseData> SamplePoseBuffer;
+FCompressedBoneTransformRWBuffer SkinningTransformBufferUAV;
 uint NumTotalHeaders;
 uint HeaderBaseOffset;
+uint SlotsPerInstance;
+uint NumCompactBones;
+uint MeshSpaceMask;
+uint LayerModeMask;
+uint NumLayers;
+
+groupshared FBoneTransformWithUniformScale SharedTransform[BONES_PER_GROUP];
 
 #define RETARGET_MODE_ANIMATION 0
 #define RETARGET_MODE_SKELETON 1
@@ -30,16 +42,31 @@ FBoneRetargetingData LoadBoneRetargetingData(ByteAddressBuffer Buffer, uint Base
 	return Buffer.Load<FBoneRetargetingData>(Offset);
 }
 
-FBoneTransformWithUniformScale LoadBoneTransformWithUniformScale(ByteAddressBuffer Buffer, uint BaseOffset, uint Index)
+FBoneTransformWithUniformScale LoadBoneTransformWithUniformScale(StructuredBuffer<float4> Buffer, uint BaseIndex, uint Index)
 {
-	uint Offset = BaseOffset + Index * (uint)sizeof(FBoneTransformWithUniformScale);
-
 	FBoneTransformWithUniformScale Output;
-	Output.Rotation       = asfloat(Buffer.Load4(Offset));
-	float4 LocationScaleX = asfloat(Buffer.Load4(Offset + 16));
 
-	Output.Location = LocationScaleX.xyz;
-	Output.Scale    = LocationScaleX.w;
+#if HALF_PRECISION_SEQUENCES
+	float4 Packed = Buffer[BaseIndex + Index];
+	Output.Rotation = FQuat(
+		f16tof32(asuint(Packed.x)),
+		f16tof32(asuint(Packed.x) >> 16),
+		f16tof32(asuint(Packed.y)),
+		f16tof32(asuint(Packed.y) >> 16));
+	Output.Location = float3(
+		f16tof32(asuint(Packed.z)),
+		f16tof32(asuint(Packed.z) >> 16),
+		f16tof32(asuint(Packed.w)));
+	Output.Scale = f16tof32(asuint(Packed.w) >> 16);
+#else
+	const uint Float4Index = BaseIndex + Index * 2;
+	float4 D0 = Buffer[Float4Index];
+	float4 D1 = Buffer[Float4Index + 1];
+	Output.Rotation = FQuat(D0.x, D0.y, D0.z, D0.w);
+	Output.Location = float3(D1.x, D1.y, D1.z);
+	Output.Scale    = D1.w;
+#endif
+
 	return Output;
 }
 
@@ -89,7 +116,7 @@ FBoneTransformWithUniformScale EvaluateAndRetargetBone(
 		FQuat Q0 = FQuat(RefPoseData.Rotation[0], RefPoseData.Rotation[1], RefPoseData.Rotation[2], RefPoseData.Rotation[3]);
 
 		// Unpack Q1 from Location.xyz + Scale.x
-		FQuat Q1 = FQuat(RefPoseData.Location.x, RefPoseData.Location.y, RefPoseData.Location.z, RefPoseData.Scale.x);
+		FQuat Q1 = FQuat(RefPoseData.Location.x, RefPoseData.Location.y, RefPoseData.Location.z, RefPoseData.Scale);
 
 		MeshLocalAnimTransform.Rotation = QuatNormalize(QuatMultiply(QuatMultiply(Q0, MeshLocalAnimTransform.Rotation), Q1));
 
@@ -130,159 +157,126 @@ FBoneTransformWithUniformScale EvaluateAndRetargetBone(
 	return MeshLocalAnimTransform;
 }
 
-[numthreads(BONES_PER_GROUP, 1, 1)]
-void EvaluateAnimSequenceCS(uint3 GroupId : SV_GroupID, uint GroupIndex : SV_GroupIndex)
+// SharedTransform .Scale field packing:
+// Lower 16 bits = FFloat16 value (bone scale or mesh-space Q.w)
+// Upper 16 bits = ParentIndex (for chain walk without BoneMapBuffer reads)
+float PackScaleWithParent(float Value, uint ParentIndex)
 {
-	const uint InstanceIndex = GroupId.y;
-	const uint HeaderSize    = uint(sizeof(FEvaluateAnimSequenceBlockHeader));
-
-	const FEvaluateAnimSequenceBlockHeader BlockHeader = EvaluateBlockHeaderBuffer.Load<FEvaluateAnimSequenceBlockHeader>((HeaderBaseOffset + InstanceIndex) * HeaderSize);
-
-	const uint BoneIndex = GroupId.x * BONES_PER_GROUP + GroupIndex;
-	if (BoneIndex >= BlockHeader.NumCompactBones)
-	{
-		return;
-	}
-
-	FBoneTransformWithUniformScale MeshLocalAnimTransform = EvaluateAndRetargetBone(
-		BoneIndex,
-		BlockHeader.SequenceTransformBufferOffset,
-		BlockHeader.RetargetingTransformBufferOffset,
-		BlockHeader.RetargetingDataBufferOffset,
-		BlockHeader.KeyIndex0,
-		BlockHeader.KeyIndex1,
-		BlockHeader.Alpha,
-		BlockHeader.NumAnimBones);
-
-	StoreBoneTransformWithUniformScale(LocalSpaceTransformBufferUAV, BlockHeader.DstTransformBufferOffset, BoneIndex, MeshLocalAnimTransform);
+	return asfloat((f32tof16(Value) & 0xFFFF) | (ParentIndex << 16));
 }
 
-[numthreads(BONES_PER_GROUP, 1, 1)]
-void EvaluateAnimSequenceBlendCS(uint3 GroupId : SV_GroupID, uint GroupIndex : SV_GroupIndex)
+float UnpackScale(float Packed)
 {
-	const uint InstanceIndex = GroupId.y;
-	const uint HeaderSize    = uint(sizeof(FEvaluateAnimSequenceBlendBlockHeader));
-
-	const FEvaluateAnimSequenceBlendBlockHeader BlockHeader = EvaluateBlockHeaderBuffer.Load<FEvaluateAnimSequenceBlendBlockHeader>((HeaderBaseOffset + InstanceIndex) * HeaderSize);
+	return f16tof32(asuint(Packed) & 0xFFFF);
+}
 
-	const uint BoneIndex = GroupId.x * BONES_PER_GROUP + GroupIndex;
-	if (BoneIndex >= BlockHeader.NumCompactBones)
-	{
-		return;
-	}
+uint UnpackParentIndex(float Packed)
+{
+	return asuint(Packed) >> 16;
+}
 
-	FBoneTransformWithUniformScale SourceMeshLocalAnimTransform = EvaluateAndRetargetBone(
-		BoneIndex,
-		BlockHeader.SourceSequenceTransformBufferOffset,
-		BlockHeader.SourceRetargetingTransformBufferOffset,
-		BlockHeader.SourceRetargetingDataBufferOffset,
-		BlockHeader.SourceKeyIndex0,
-		BlockHeader.SourceKeyIndex1,
-		BlockHeader.SourceAlpha,
-		BlockHeader.SourceNumAnimBones);
-
-	FBoneTransformWithUniformScale TargetMeshLocalAnimTransform = EvaluateAndRetargetBone(
-		BoneIndex,
-		BlockHeader.TargetSequenceTransformBufferOffset,
-		BlockHeader.TargetRetargetingTransformBufferOffset,
-		BlockHeader.TargetRetargetingDataBufferOffset,
-		BlockHeader.TargetKeyIndex0,
-		BlockHeader.TargetKeyIndex1,
-		BlockHeader.TargetAlpha,
-		BlockHeader.TargetNumAnimBones);
+void StoreSharedTransform(uint Index, FBoneTransformWithUniformScale Transform, uint ParentIndex)
+{
+	Transform.Scale = PackScaleWithParent(Transform.Scale, ParentIndex);
+	SharedTransform[Index] = Transform;
+}
 
-	FBoneTransformWithUniformScale MeshLocalAnimTransform = BlendBoneTransformsWithUniformScale(SourceMeshLocalAnimTransform, TargetMeshLocalAnimTransform, BlockHeader.BlendWeight);
+FBoneTransformWithUniformScale LoadSharedTransform(uint Index)
+{
+	FBoneTransformWithUniformScale T = SharedTransform[Index];
+	T.Scale = UnpackScale(T.Scale);
+	return T;
+}
 
-	StoreBoneTransformWithUniformScale(LocalSpaceTransformBufferUAV, BlockHeader.DstTransformBufferOffset, BoneIndex, MeshLocalAnimTransform);
+uint GetSharedParentIndex(uint Index)
+{
+	return UnpackParentIndex(SharedTransform[Index].Scale);
 }
 
-ByteAddressBuffer LocalSpaceTransformBuffer;
-ByteAddressBuffer ConcatenateBlockHeaderBuffer;
-ByteAddressBuffer LeafBonePathBuffer;
-RWByteAddressBuffer GlobalSpaceTransformBufferUAV;
+FBoneHierarchyData LoadBoneHierarchyData(ByteAddressBuffer Buffer, uint BaseOffset, uint Index)
+{
+	uint Offset = BaseOffset + Index * (uint)sizeof(FBoneHierarchyData);
+	return Buffer.Load<FBoneHierarchyData>(Offset);
+}
 
-groupshared FBoneTransformWithUniformScale SharedTransform[BONES_PER_GROUP];
+ByteAddressBuffer BoneMaskBuffer;
 
-FBoneTransformWithUniformScale LoadLocalSpaceTransform(uint TransformBufferOffset, uint BoneIndex)
+float LoadBoneMaskWeight(FEvaluateLayerHeader Header, uint BoneIndex)
 {
-	if (BoneIndex < BONES_PER_GROUP)
-	{
-		return SharedTransform[BoneIndex];
-	}
-	else
+	if (!HasBoneMask(Header))
 	{
-		return LoadBoneTransformWithUniformScale(LocalSpaceTransformBuffer, TransformBufferOffset, BoneIndex);
+		return 1.0f;
 	}
+
+	uint BoneMaskOffset = GetBoneMaskOffset(Header);
+	return asfloat(BoneMaskBuffer.Load(BoneMaskOffset + BoneIndex * 4));
 }
 
-[numthreads(BONES_PER_GROUP, 1, 1)]
-void ConcatenateLocalToGlobalCS(uint3 GroupId : SV_GroupID, uint GroupIndex : SV_GroupIndex)
+FBoneTransformWithUniformScale EvaluateWeightedSamples(uint BoneIndex, uint Offset, uint Count)
 {
-	const uint InstanceIndex     = GroupId.y;
-	const uint PrefetchBoneIndex = GroupIndex;
-	const uint PathElementSize   = 4;
-
-	FConcatenateBlockHeader BlockHeader = ConcatenateBlockHeaderBuffer.Load<FConcatenateBlockHeader>((HeaderBaseOffset + InstanceIndex) * uint(sizeof(FConcatenateBlockHeader)));
-
-	if (PrefetchBoneIndex < BlockHeader.NumPrefetchBones)
+	FSamplePoseData Sample = SamplePoseBuffer[Offset];
+	FBoneTransformWithUniformScale Result = EvaluateAndRetargetBone(
+		BoneIndex,
+		Sample.SequenceTransformBufferOffset,
+		Sample.RetargetingTransformBufferOffset,
+		Sample.RetargetingDataBufferOffset,
+		Sample.KeyIndex0, Sample.KeyIndex1,
+		Sample.Alpha, Sample.NumAnimBones);
+	float AccumulatedWeight = Sample.Weight;
+
+	if (Count >= 2)
 	{
-		SharedTransform[PrefetchBoneIndex] = LoadBoneTransformWithUniformScale(LocalSpaceTransformBuffer, BlockHeader.TransformBufferOffset, PrefetchBoneIndex);
+		Sample = SamplePoseBuffer[Offset + 1];
+		FBoneTransformWithUniformScale Pose = EvaluateAndRetargetBone(
+			BoneIndex,
+			Sample.SequenceTransformBufferOffset,
+			Sample.RetargetingTransformBufferOffset,
+			Sample.RetargetingDataBufferOffset,
+			Sample.KeyIndex0, Sample.KeyIndex1,
+			Sample.Alpha, Sample.NumAnimBones);
+		float SampleWeight = Sample.Weight;
+		Result = BlendBoneTransformsWithUniformScale(Result, Pose, SampleWeight / (AccumulatedWeight + SampleWeight));
+		AccumulatedWeight += SampleWeight;
 	}
-	GroupMemoryBarrierWithGroupSync();
 
-	const uint LeafBoneIndex = GroupId.x * BONES_PER_GROUP + GroupIndex;
-	if (LeafBoneIndex >= BlockHeader.NumLeafBones)
+	if (Count >= 3)
 	{
-		return;
+		Sample = SamplePoseBuffer[Offset + 2];
+		FBoneTransformWithUniformScale Pose = EvaluateAndRetargetBone(
+			BoneIndex,
+			Sample.SequenceTransformBufferOffset,
+			Sample.RetargetingTransformBufferOffset,
+			Sample.RetargetingDataBufferOffset,
+			Sample.KeyIndex0, Sample.KeyIndex1,
+			Sample.Alpha, Sample.NumAnimBones);
+		float SampleWeight = Sample.Weight;
+		Result = BlendBoneTransformsWithUniformScale(Result, Pose, SampleWeight / (AccumulatedWeight + SampleWeight));
 	}
 
-	// Load relative offset from offset table, then add base to get absolute offset
-	const uint RelativePathHeaderOffset = LeafBonePathBuffer.Load<uint>(BlockHeader.LeafBonePathBufferOffset + LeafBoneIndex * PathElementSize);
-	const uint PathHeaderOffset = BlockHeader.LeafBonePathBufferOffset + RelativePathHeaderOffset;
+	return Result;
+}
 
-	FLeafBonePathHeader PathHeader = LeafBonePathBuffer.Load<FLeafBonePathHeader>(PathHeaderOffset);
-	FBoneTransformWithUniformScale ParentTransform = LoadLocalSpaceTransform(BlockHeader.TransformBufferOffset, 0);
+FBoneTransformWithUniformScale ConcatenateToGlobal(
+	uint BoneIndex,
+	uint SharedTransformOffset,
+	uint InvGlobalRefPoseTransformBufferOffset)
+{
+	FBoneTransformWithUniformScale GlobalTransform = LoadSharedTransform(SharedTransformOffset + BoneIndex);
 
-	if (PathHeader.bIsRootOwned)
+	uint ParentIndex = GetSharedParentIndex(SharedTransformOffset + BoneIndex);
+	while (ParentIndex != INVALID_BONE_INDEX)
 	{
-		StoreBoneTransformWithUniformScale(GlobalSpaceTransformBufferUAV, BlockHeader.TransformBufferOffset, 0, ParentTransform);
+		GlobalTransform = MultiplyBoneTransformWithUniformScale(GlobalTransform, LoadSharedTransform(SharedTransformOffset + ParentIndex));
+		ParentIndex = GetSharedParentIndex(SharedTransformOffset + ParentIndex);
 	}
 
-	if (PathHeader.PathDepth > 0)
-	{
-		uint PathNodeOffset = PathHeaderOffset + PathElementSize;
-		FLeafBonePathNode PathNode;
-
-		for (uint PathIndex = 0; PathIndex < PathHeader.PathDepth; ++PathIndex)
-		{
-			PathNode = LeafBonePathBuffer.Load<FLeafBonePathNode>(PathNodeOffset);
-			PathNodeOffset += PathElementSize;
-
-			FBoneTransformWithUniformScale ChildTransform = LoadLocalSpaceTransform(BlockHeader.TransformBufferOffset, PathNode.BoneIndex);
-			ParentTransform = MultiplyBoneTransformWithUniformScale(ChildTransform, ParentTransform);
-
-			if (PathNode.bIsOwned)
-			{
-				StoreBoneTransformWithUniformScale(GlobalSpaceTransformBufferUAV, BlockHeader.TransformBufferOffset, PathNode.BoneIndex, ParentTransform);
-			}
-
-			if (PathNode.bIsLeaf)
-			{
-				break;
-			}
-		}
-	}
+	FBoneTransformWithUniformScale InvRefPose = LoadBoneTransformWithUniformScale(SequenceTransformBuffer, InvGlobalRefPoseTransformBufferOffset, BoneIndex);
+	return MultiplyBoneTransformWithUniformScale(InvRefPose, GlobalTransform);
 }
 
-ByteAddressBuffer ScatterBlockHeaderBuffer;
-ByteAddressBuffer GlobalSpaceTransformBuffer;
-ByteAddressBuffer BoneMapBuffer;
-FCompressedBoneTransformRWBuffer SkinningTransformBufferUAV;
-
 void ScatterGlobalSpaceTransform(uint DstTransformBufferOffset1, uint DstTransformBufferOffset2, uint TransformIndex, FBoneTransformWithUniformScale GlobalTransform)
 {
 	float4x3 OutputMatrix = (float4x3)MatrixFromBoneTransformWithUniformScale(GlobalTransform);
-
 	StoreCompressedBoneTransform(SkinningTransformBufferUAV, DstTransformBufferOffset1, TransformIndex, OutputMatrix);
 
 	if (DstTransformBufferOffset1 != DstTransformBufferOffset2)
@@ -291,154 +285,267 @@ void ScatterGlobalSpaceTransform(uint DstTransformBufferOffset1, uint DstTransfo
 	}
 }
 
-[numthreads(BONES_PER_GROUP, 1, 1)]
-void ScatterTransformsCS(uint3 GroupId : SV_GroupID, uint GroupIndex : SV_GroupIndex)
+#if HAS_MESH_SPACE_LAYERS
+
+void StoreMeshSpaceLayerRotation(uint Index, FQuat Q, uint ParentIndex)
 {
-	const uint InstanceIndex = GroupId.y;
-	const uint HeaderSize    = uint(sizeof(FScatterTransformsBlockHeader));
+	SharedTransform[Index].Location = float3(Q.x, Q.y, Q.z);
+	SharedTransform[Index].Scale = PackScaleWithParent(Q.w, ParentIndex);
+}
 
-	const FScatterTransformsBlockHeader BlockHeader = ScatterBlockHeaderBuffer.Load<FScatterTransformsBlockHeader>((HeaderBaseOffset + InstanceIndex) * HeaderSize);
+FQuat LoadMeshSpaceLayerRotation(uint Index)
+{
+	float3 XYZ = SharedTransform[Index].Location;
+	float W = UnpackScale(SharedTransform[Index].Scale);
+	return FQuat(XYZ.x, XYZ.y, XYZ.z, W);
+}
 
-	const uint TransformIndex = GroupId.x * BONES_PER_GROUP + GroupIndex;
-	if (TransformIndex >= BlockHeader.MaxBoneTransformCount)
-	{
-		return;
-	}
+struct FMeshSpaceBlendResult
+{
+	FQuat BlendedGlobal;
+	uint ParentIndex;
+	float3 Location;
+	float Scale;
+};
+
+FMeshSpaceBlendResult BlendMeshSpaceOverride(
+	FBoneTransformWithUniformScale LayerPose,
+	FBoneTransformWithUniformScale Result,
+	uint BoneIndex,
+	uint SharedOffset,
+	float EffectiveWeight)
+{
+	uint SharedIndex        = SharedOffset + BoneIndex;
+	FQuat BaseLocal         = SharedTransform[SharedIndex].Rotation;
+	FQuat LayerLocal        = LoadMeshSpaceLayerRotation(SharedIndex);
+	FQuat BaseParentGlobal  = QuatIdentity();
+	FQuat LayerParentGlobal = QuatIdentity();
 
-	const uint CompactBoneIndex = BoneMapBuffer.Load<uint>(BlockHeader.BoneMapBufferOffset + TransformIndex * uint(sizeof(uint)));
+	uint ParentIndex = GetSharedParentIndex(SharedIndex);
 
-	if (CompactBoneIndex != INVALID_BONE_INDEX)
+	uint WalkIndex = ParentIndex;
+	while (WalkIndex != INVALID_BONE_INDEX)
 	{
-		FBoneTransformWithUniformScale InvGlobalRefPoseTransform = LoadBoneTransformWithUniformScale(SequenceTransformBuffer, BlockHeader.InvGlobalRefPoseTransformBufferOffset, CompactBoneIndex);
-		FBoneTransformWithUniformScale GlobalSpaceTransform = LoadBoneTransformWithUniformScale(GlobalSpaceTransformBuffer, BlockHeader.SrcTransformBufferOffset, CompactBoneIndex);
-		GlobalSpaceTransform = MultiplyBoneTransformWithUniformScale(InvGlobalRefPoseTransform, GlobalSpaceTransform);
-
-		ScatterGlobalSpaceTransform(BlockHeader.DstTransformBufferOffset1, BlockHeader.DstTransformBufferOffset2, TransformIndex, GlobalSpaceTransform);
+		uint ParentSharedIndex = SharedOffset + WalkIndex;
+		BaseParentGlobal       = QuatMultiply(SharedTransform[ParentSharedIndex].Rotation, BaseParentGlobal);
+		LayerParentGlobal      = QuatMultiply(LoadMeshSpaceLayerRotation(ParentSharedIndex), LayerParentGlobal);
+		WalkIndex              = GetSharedParentIndex(ParentSharedIndex);
 	}
-}
 
-FBoneHierarchyData LoadBoneHierarchyData(ByteAddressBuffer Buffer, uint BaseOffset, uint Index)
-{
-	uint Offset = BaseOffset + Index * (uint)sizeof(FBoneHierarchyData);
-	return Buffer.Load<FBoneHierarchyData>(Offset);
+	FQuat BaseGlobal    = QuatMultiply(BaseParentGlobal, BaseLocal);
+	FQuat LayerGlobal   = QuatMultiply(LayerParentGlobal, LayerLocal);
+	FQuat BlendedGlobal = QuatNormalize(QuatSlerpApproximate(BaseGlobal, LayerGlobal, EffectiveWeight));
+
+	FMeshSpaceBlendResult Out;
+	Out.BlendedGlobal = BlendedGlobal;
+	Out.ParentIndex   = ParentIndex;
+	Out.Location      = lerp(Result.Location, LayerPose.Location, EffectiveWeight);
+	Out.Scale         = lerp(Result.Scale, LayerPose.Scale, EffectiveWeight);
+	return Out;
 }
 
-FBoneTransformWithUniformScale ConcatenateToGlobal(
+FMeshSpaceBlendResult BlendMeshSpaceAdditive(
+	FBoneTransformWithUniformScale LayerPose,
+	FBoneTransformWithUniformScale Result,
 	uint BoneIndex,
-	uint SharedTransformOffset,
-	uint BoneMapBufferOffset,
-	uint InvGlobalRefPoseTransformBufferOffset,
-	out uint OutTransformIndex)
+	uint SharedOffset,
+	float EffectiveWeight)
 {
-	FBoneTransformWithUniformScale GlobalTransform = SharedTransform[SharedTransformOffset + BoneIndex];
-	
-	FBoneHierarchyData HierarchyData = LoadBoneHierarchyData(BoneMapBuffer, BoneMapBufferOffset, BoneIndex);
-	OutTransformIndex = HierarchyData.TransformIndex;
+	uint SharedIndex = SharedOffset + BoneIndex;
+	FQuat BaseGlobal       = SharedTransform[SharedIndex].Rotation;
+	FQuat BaseParentGlobal = QuatIdentity();
+
+	uint ParentIndex = GetSharedParentIndex(SharedIndex);
 
-	while (HierarchyData.ParentIndex != INVALID_BONE_INDEX)
+	uint WalkIndex = ParentIndex;
+	while (WalkIndex != INVALID_BONE_INDEX)
 	{
-		GlobalTransform = MultiplyBoneTransformWithUniformScale(GlobalTransform, SharedTransform[SharedTransformOffset + HierarchyData.ParentIndex]);
-		HierarchyData = LoadBoneHierarchyData(BoneMapBuffer, BoneMapBufferOffset, HierarchyData.ParentIndex);
+		BaseParentGlobal = QuatMultiply(SharedTransform[SharedOffset + WalkIndex].Rotation, BaseParentGlobal);
+		WalkIndex = GetSharedParentIndex(SharedOffset + WalkIndex);
 	}
 
-	FBoneTransformWithUniformScale InvRefPose = LoadBoneTransformWithUniformScale(SequenceTransformBuffer, InvGlobalRefPoseTransformBufferOffset, BoneIndex);
-	return MultiplyBoneTransformWithUniformScale(InvRefPose, GlobalTransform);
+	BaseGlobal = QuatMultiply(BaseParentGlobal, BaseGlobal);
+
+	FQuat WeightedDelta = QuatSlerpApproximate(QuatIdentity(), LayerPose.Rotation, EffectiveWeight);
+	FQuat BlendedGlobal = QuatNormalize(QuatMultiply(WeightedDelta, BaseGlobal));
+
+	FMeshSpaceBlendResult Out;
+	Out.BlendedGlobal = BlendedGlobal;
+	Out.ParentIndex   = ParentIndex;
+	Out.Location      = Result.Location + LayerPose.Location * EffectiveWeight;
+	Out.Scale         = Result.Scale * lerp(1.0, LayerPose.Scale, EffectiveWeight);
+	return Out;
 }
 
+#endif
+
 [numthreads(BONES_PER_GROUP, 1, 1)]
 void EvaluateConcatenateScatterCS(uint3 GroupId : SV_GroupID, uint GroupIndex : SV_GroupIndex)
 {
-	const uint HeaderSize      = uint(sizeof(FEvaluateConcatenateScatterBlockHeader));
-	const uint SlotsPerInstance = EvaluateBlockHeaderBuffer.Load<FEvaluateConcatenateScatterBlockHeader>(HeaderBaseOffset * HeaderSize).SlotsPerInstance;
 	const uint InstancesPerGroup = BONES_PER_GROUP / SlotsPerInstance;
 
 	const uint InstanceInGroup       = GroupIndex / SlotsPerInstance;
 	const uint BoneIndex             = GroupIndex % SlotsPerInstance;
-	const uint HeaderIndex           = GroupId.x * InstancesPerGroup + InstanceInGroup;
+	const uint HeaderIndex           = min(GroupId.x * InstancesPerGroup + InstanceInGroup, NumTotalHeaders - 1);
 	const uint SharedTransformOffset = InstanceInGroup * SlotsPerInstance;
 
 	const bool bValidInstance = HeaderIndex < NumTotalHeaders;
+	const bool bActiveLane = bValidInstance && BoneIndex < NumCompactBones;
+	const uint SharedIndex = SharedTransformOffset + BoneIndex;
 
-	FEvaluateConcatenateScatterBlockHeader BlockHeader = (FEvaluateConcatenateScatterBlockHeader)0;
-	if (bValidInstance)
-	{
-		BlockHeader = EvaluateBlockHeaderBuffer.Load<FEvaluateConcatenateScatterBlockHeader>((HeaderBaseOffset + HeaderIndex) * HeaderSize);
-	}
+	FBoneTransformWithUniformScale Result = MakeBoneTransformWithUniformScaleIdentity();
 
-	if (bValidInstance && BoneIndex < BlockHeader.NumCompactBones)
-	{
-		SharedTransform[SharedTransformOffset + BoneIndex] = EvaluateAndRetargetBone(
-			BoneIndex,
-			BlockHeader.SequenceTransformBufferOffset,
-			BlockHeader.RetargetingTransformBufferOffset,
-			BlockHeader.RetargetingDataBufferOffset,
-			BlockHeader.KeyIndex0,
-			BlockHeader.KeyIndex1,
-			BlockHeader.Alpha,
-			BlockHeader.NumAnimBones);
-	}
+	uint InvGlobalRefPoseTransformBufferOffset = 0;
+	uint BoneMapBufferOffset = 0;
+	uint DstTransformBufferOffset1 = 0;
+	uint DstTransformBufferOffset2 = 0;
+	uint BoneTransformIndex = 0;
 
-	GroupMemoryBarrierWithGroupSync();
+	uint BoneParentIndex = INVALID_BONE_INDEX;
 
-	if (bValidInstance && BoneIndex < BlockHeader.NumCompactBones)
+#if SINGLE_LAYER
+	const uint LayerCount = 1;
+#else
+	const uint LayerCount = NumLayers;
+#endif
+
+	for (uint LayerIndex = 0; LayerIndex < LayerCount; LayerIndex++)
 	{
-		uint TransformIndex;
-		FBoneTransformWithUniformScale GlobalTransform = ConcatenateToGlobal(BoneIndex, SharedTransformOffset, BlockHeader.BoneMapBufferOffset, BlockHeader.InvGlobalRefPoseTransformBufferOffset, TransformIndex);
-		ScatterGlobalSpaceTransform(BlockHeader.DstTransformBufferOffset1, BlockHeader.DstTransformBufferOffset2, TransformIndex, GlobalTransform);
-	}
-}
+		FEvaluateLayerHeader Header = LayerHeaderBuffer[HeaderBaseOffset + (HeaderIndex * LayerCount) + LayerIndex];
 
-[numthreads(BONES_PER_GROUP, 1, 1)]
-void EvaluateBlendConcatenateScatterCS(uint3 GroupId : SV_GroupID, uint GroupIndex : SV_GroupIndex)
-{
-	const uint HeaderSize      = uint(sizeof(FEvaluateBlendConcatenateScatterBlockHeader));
-	const uint SlotsPerInstance = EvaluateBlockHeaderBuffer.Load<FEvaluateBlendConcatenateScatterBlockHeader>(HeaderBaseOffset * HeaderSize).SlotsPerInstance;
-	const uint InstancesPerGroup = BONES_PER_GROUP / SlotsPerInstance;
+		if (LayerIndex == 0)
+		{
+			BoneMapBufferOffset       = Header.BoneMapBufferOffset;
+			DstTransformBufferOffset1 = Header.DstTransformBufferOffset1;
+			DstTransformBufferOffset2 = Header.DstTransformBufferOffset2;
+			InvGlobalRefPoseTransformBufferOffset = Header.InvGlobalRefPoseTransformBufferOffset;
 
-	const uint InstanceInGroup       = GroupIndex / SlotsPerInstance;
-	const uint BoneIndex             = GroupIndex % SlotsPerInstance;
-	const uint HeaderIndex           = GroupId.x * InstancesPerGroup + InstanceInGroup;
-	const uint SharedTransformOffset = InstanceInGroup * SlotsPerInstance;
+			if (bActiveLane)
+			{
+				FBoneHierarchyData InitHierarchy = LoadBoneHierarchyData(BoneMapBuffer, BoneMapBufferOffset, BoneIndex);
+				BoneParentIndex = InitHierarchy.ParentIndex;
+				BoneTransformIndex = InitHierarchy.TransformIndex;
+			}
+		}
 
-	const bool bValidInstance = HeaderIndex < NumTotalHeaders;
+#if SINGLE_LAYER
+		const bool  bMeshSpaceRotation = false;
+		const float LayerWeight = 1.0f;
+		const float EffectiveWeight = 1.0f;
+#else
+		const float LayerWeight        = Header.LayerWeight;
+		const uint  LayerMode          = (LayerModeMask >> (LayerIndex * SKINNING_LAYER_MODE_BITS_PER_LAYER)) & SKINNING_LAYER_MODE_MASK_PER_LAYER;
+		const bool  bMeshSpaceRotation = (MeshSpaceMask & (1u << LayerIndex)) != 0;
+		const bool  bMeshSpaceAdditive = bMeshSpaceRotation && (LayerMode == SKINNING_LAYER_MODE_ADDITIVE);
+
+		float EffectiveWeight = LayerWeight;
+		if (EffectiveWeight > 0 && HasBoneMask(Header))
+		{
+			EffectiveWeight *= LoadBoneMaskWeight(Header, BoneIndex);
+		}
+#endif
 
-	FEvaluateBlendConcatenateScatterBlockHeader BlockHeader = (FEvaluateBlendConcatenateScatterBlockHeader)0;
-	if (bValidInstance)
-	{
-		BlockHeader = EvaluateBlockHeaderBuffer.Load<FEvaluateBlendConcatenateScatterBlockHeader>((HeaderBaseOffset + HeaderIndex) * HeaderSize);
-	}
+		FBoneTransformWithUniformScale LayerPose = MakeBoneTransformWithUniformScaleIdentity();
+
+		if (bActiveLane && (EffectiveWeight > 0 || (bMeshSpaceRotation && LayerWeight > 0)))
+		{
+			const uint SampleOffset = Header.SampleDataOffset;
+			const uint NumSourceSamples = GetNumSourceSamples(Header);
+			const uint NumTargetSamples = GetNumTargetSamples(Header);
+
+			LayerPose = EvaluateWeightedSamples(BoneIndex, SampleOffset, NumSourceSamples);
 
-	if (bValidInstance && BoneIndex < BlockHeader.NumCompactBones)
+			if (NumTargetSamples > 0 && Header.BlendWeight > 0)
+			{
+				FBoneTransformWithUniformScale TargetPose = EvaluateWeightedSamples(
+					BoneIndex, SampleOffset + NumSourceSamples, NumTargetSamples);
+				LayerPose = BlendBoneTransformsWithUniformScale(
+					LayerPose, TargetPose, Header.BlendWeight);
+			}
+
+#if SINGLE_LAYER
+			Result = LayerPose;
+#else
+			if (EffectiveWeight > 0 && !bMeshSpaceRotation)
+			{
+				if (LayerMode == SKINNING_LAYER_MODE_OVERRIDE)
+				{
+					Result = BlendBoneTransformsWithUniformScale(Result, LayerPose, EffectiveWeight);
+				}
+				else
+				{
+					Result = ApplyAdditiveLocalSpace(Result, LayerPose, EffectiveWeight);
+				}
+			}
+#endif
+		}
+
+#if HAS_MESH_SPACE_LAYERS
+		if (bMeshSpaceRotation)
+		{
+			if (bActiveLane && LayerWeight > 0)
+			{
+				StoreSharedTransform(SharedIndex, Result, BoneParentIndex);
+	
+				if (!bMeshSpaceAdditive)
+				{
+					StoreMeshSpaceLayerRotation(SharedIndex, LayerPose.Rotation, BoneParentIndex);
+				}
+			}
+
+			GroupMemoryBarrierWithGroupSync();
+
+			FMeshSpaceBlendResult MeshResult = (FMeshSpaceBlendResult)0;
+
+			if (bActiveLane && LayerWeight > 0)
+			{
+				if (LayerMode == SKINNING_LAYER_MODE_OVERRIDE)
+				{
+					MeshResult = BlendMeshSpaceOverride(LayerPose, Result, BoneIndex, SharedTransformOffset, EffectiveWeight);
+				}
+				else
+				{
+					MeshResult = BlendMeshSpaceAdditive(LayerPose, Result, BoneIndex, SharedTransformOffset, EffectiveWeight);
+				}
+
+				StoreMeshSpaceLayerRotation(SharedIndex, MeshResult.BlendedGlobal, BoneParentIndex);
+			}
+
+			GroupMemoryBarrierWithGroupSync();
+
+			if (bActiveLane && LayerWeight > 0)
+			{
+				FQuat BlendedParentGlobal = QuatIdentity();
+				if (MeshResult.ParentIndex != INVALID_BONE_INDEX)
+				{
+					BlendedParentGlobal = LoadMeshSpaceLayerRotation(SharedTransformOffset + MeshResult.ParentIndex);
+				}
+
+				Result.Rotation = QuatNormalize(QuatMultiply(QuatConjugate(BlendedParentGlobal), MeshResult.BlendedGlobal));
+				Result.Location = MeshResult.Location;
+				Result.Scale    = MeshResult.Scale;
+			}
+
+			GroupMemoryBarrierWithGroupSync();
+		}
+#endif
+	}
+	
+	if (bActiveLane)
 	{
-		FBoneTransformWithUniformScale SourceTransform = EvaluateAndRetargetBone(
-			BoneIndex,
-			BlockHeader.SourceSequenceTransformBufferOffset,
-			BlockHeader.SourceRetargetingTransformBufferOffset,
-			BlockHeader.SourceRetargetingDataBufferOffset,
-			BlockHeader.SourceKeyIndex0,
-			BlockHeader.SourceKeyIndex1,
-			BlockHeader.SourceAlpha,
-			BlockHeader.SourceNumAnimBones);
-
-		FBoneTransformWithUniformScale TargetTransform = EvaluateAndRetargetBone(
-			BoneIndex,
-			BlockHeader.TargetSequenceTransformBufferOffset,
-			BlockHeader.TargetRetargetingTransformBufferOffset,
-			BlockHeader.TargetRetargetingDataBufferOffset,
-			BlockHeader.TargetKeyIndex0,
-			BlockHeader.TargetKeyIndex1,
-			BlockHeader.TargetAlpha,
-			BlockHeader.TargetNumAnimBones);
-
-		SharedTransform[SharedTransformOffset + BoneIndex] = BlendBoneTransformsWithUniformScale(SourceTransform, TargetTransform, BlockHeader.BlendWeight);
+		StoreSharedTransform(SharedIndex, Result, BoneParentIndex);
 	}
-
 	GroupMemoryBarrierWithGroupSync();
 
-	if (bValidInstance && BoneIndex < BlockHeader.NumCompactBones)
+	if (bActiveLane && BoneTransformIndex != INVALID_BONE_INDEX)
 	{
-		uint TransformIndex;
-		FBoneTransformWithUniformScale GlobalTransform = ConcatenateToGlobal(BoneIndex, SharedTransformOffset, BlockHeader.BoneMapBufferOffset, BlockHeader.InvGlobalRefPoseTransformBufferOffset, TransformIndex);
-		ScatterGlobalSpaceTransform(BlockHeader.DstTransformBufferOffset1, BlockHeader.DstTransformBufferOffset2, TransformIndex, GlobalTransform);
+		FBoneTransformWithUniformScale GlobalTransform = ConcatenateToGlobal(
+			BoneIndex, SharedTransformOffset,
+			InvGlobalRefPoseTransformBufferOffset);
+
+		ScatterGlobalSpaceTransform(
+			DstTransformBufferOffset1,
+			DstTransformBufferOffset2,
+			BoneTransformIndex, GlobalTransform);
 	}
 }
\ No newline at end of file
diff --git a/Engine/Shaders/Shared/SkinningDefinitions.h b/Engine/Shaders/Shared/SkinningDefinitions.h
--- a/Engine/Shaders/Shared/SkinningDefinitions.h
+++ b/Engine/Shaders/Shared/SkinningDefinitions.h
@@ -45,88 +45,48 @@ struct FBoneRetargetingData
 };
 SKINNING_STATIC_ASSERT(sizeof(FBoneRetargetingData) == 4, "FBoneRetargetingData must be 4 bytes");
 
-struct FEvaluateAnimSequenceBlockHeader
-{
-	float  Alpha;
-	uint KeyIndex0;
-	uint KeyIndex1;
-	uint SequenceTransformBufferOffset;
-	uint RetargetingTransformBufferOffset;
-	uint RetargetingDataBufferOffset;
-	uint DstTransformBufferOffset;
-	uint NumAnimBones;
-	uint NumCompactBones : 16;
-	uint Padding         : 16;
-};
-SKINNING_STATIC_ASSERT(sizeof(FEvaluateAnimSequenceBlockHeader) == 36, "FEvaluateAnimSequenceBlockHeader must be 36 bytes");
+#define MAX_ANIM_SEQUENCE_LAYERS 4
+
+#define SKINNING_LAYER_MODE_OVERRIDE 0
+#define SKINNING_LAYER_MODE_ADDITIVE 1
+
+#define SKINNING_MAX_KEY_INDEX 65535
+
+#define SKINNING_LAYER_MODE_BITS_PER_LAYER 2
+#define SKINNING_LAYER_MODE_MASK_PER_LAYER ((1 << SKINNING_LAYER_MODE_BITS_PER_LAYER) - 1)
+
+#define SKINNING_INVALID_BONE_MASK_OFFSET 0x3FFFFFF
+
+#define SKINNING_MAX_BONES_PER_GROUP 1024
+#define SKINNING_MAX_BONES_PER_GROUP_MESHSPACE 256
 
-struct FEvaluateAnimSequenceBlendBlockHeader
+// Per-layer header. One entry per layer per track.
+struct FEvaluateLayerHeader
 {
-	float SourceAlpha;
-	float TargetAlpha;
+	uint  BoneMapBufferOffset;
+	uint  DstTransformBufferOffset1;
+	uint  DstTransformBufferOffset2;
+	uint  SampleDataOffset;
 	float BlendWeight;
-	uint SourceKeyIndex0;
-	uint SourceKeyIndex1;
-	uint TargetKeyIndex0;
-	uint TargetKeyIndex1;
-	uint SourceSequenceTransformBufferOffset;
-	uint TargetSequenceTransformBufferOffset;
-	uint SourceRetargetingTransformBufferOffset;
-	uint TargetRetargetingTransformBufferOffset;
-	uint SourceRetargetingDataBufferOffset;
-	uint TargetRetargetingDataBufferOffset;
-	uint DstTransformBufferOffset;
-	uint SourceNumAnimBones   : 16;
-	uint TargetNumAnimBones   : 16;
-	uint NumCompactBones      : 16;
-	uint Padding              : 16;
- };
-SKINNING_STATIC_ASSERT(sizeof(FEvaluateAnimSequenceBlendBlockHeader) == 64, "FEvaluateAnimSequenceBlendBlockHeader must be 64 bytes");
-
-struct FEvaluateConcatenateScatterBlockHeader
-{
-	float Alpha;
-	uint KeyIndex0;
-	uint KeyIndex1;
-	uint SequenceTransformBufferOffset;
-	uint RetargetingTransformBufferOffset;
-	uint RetargetingDataBufferOffset;
-	uint InvGlobalRefPoseTransformBufferOffset;
-	uint BoneMapBufferOffset;
-	uint DstTransformBufferOffset1;
-	uint DstTransformBufferOffset2;
-	uint NumAnimBones       : 16;
-	uint NumCompactBones    : 16;
-	uint SlotsPerInstance   : 16;
-	uint Padding            : 16;
+	float LayerWeight;
+	uint  BoneMaskAndSampleCounts;               // [0:25] BoneMask, [26:28] NumSourceSamples, [29:31] NumTargetSamples
+	uint  InvGlobalRefPoseTransformBufferOffset;
 };
-SKINNING_STATIC_ASSERT(sizeof(FEvaluateConcatenateScatterBlockHeader) == 48, "FEvaluateConcatenateScatterBlockHeader must be 48 bytes");
+SKINNING_STATIC_ASSERT(sizeof(FEvaluateLayerHeader) == 32, "FEvaluateLayerHeader must be 32 bytes");
 
-struct FEvaluateBlendConcatenateScatterBlockHeader
+// Per-sample pose evaluation data. Tightly packed per track.
+struct FSamplePoseData
 {
-	float SourceAlpha;
-	float TargetAlpha;
-	float BlendWeight;
-	uint SourceKeyIndex0;
-	uint SourceKeyIndex1;
-	uint TargetKeyIndex0;
-	uint TargetKeyIndex1;
-	uint SourceSequenceTransformBufferOffset;
-	uint TargetSequenceTransformBufferOffset;
-	uint SourceRetargetingTransformBufferOffset;
-	uint TargetRetargetingTransformBufferOffset;
-	uint SourceRetargetingDataBufferOffset;
-	uint TargetRetargetingDataBufferOffset;
-	uint InvGlobalRefPoseTransformBufferOffset;
-	uint BoneMapBufferOffset;
-	uint DstTransformBufferOffset1;
-	uint DstTransformBufferOffset2;
-	uint SourceNumAnimBones : 16;
-	uint TargetNumAnimBones : 16;
-	uint NumCompactBones    : 16;
-	uint SlotsPerInstance   : 16;
+	float Alpha;
+	uint  KeyIndex0;
+	uint  KeyIndex1;
+	uint  SequenceTransformBufferOffset;
+	uint  RetargetingTransformBufferOffset;
+	uint  RetargetingDataBufferOffset;
+	uint  NumAnimBones;
+	float Weight;
 };
-SKINNING_STATIC_ASSERT(sizeof(FEvaluateBlendConcatenateScatterBlockHeader) == 76, "FEvaluateBlendConcatenateScatterBlockHeader must be 76 bytes");
+SKINNING_STATIC_ASSERT(sizeof(FSamplePoseData) == 32, "FSamplePoseData must be 32 bytes");
 
 struct FBoneHierarchyData
 {
@@ -135,42 +95,17 @@ struct FBoneHierarchyData
 };
 SKINNING_STATIC_ASSERT(sizeof(FBoneHierarchyData) == 4, "FBoneHierarchyData must be 4 bytes");
 
-struct FLeafBonePathHeader
-{
-	uint PathDepth      : 31;
-	uint bIsRootOwned   : 1;
-};
-SKINNING_STATIC_ASSERT(sizeof(FLeafBonePathHeader) == 4, "FLeafBonePathHeader must be 4 bytes");
-
-struct FLeafBonePathNode
-{
-	uint BoneIndex : 30;
-	uint bIsOwned  : 1;
-	uint bIsLeaf   : 1;
-};
-SKINNING_STATIC_ASSERT(sizeof(FLeafBonePathNode) == 4, "FLeafBonePathNode must be 4 bytes");
-
-struct FConcatenateBlockHeader
-{
-	uint TransformBufferOffset;
-	uint LeafBonePathBufferOffset : 24;
-	uint NumPrefetchBones   : 8;
-	uint NumLeafBones       : 16;
-	uint Padding            : 16;
-};
-SKINNING_STATIC_ASSERT(sizeof(FConcatenateBlockHeader) == 12, "FConcatenateBlockHeader must be 12 bytes");
-
-struct FScatterTransformsBlockHeader
+#ifdef __cplusplus
+inline void SetBoneMaskAndSampleCounts(FEvaluateLayerHeader& H, uint32 BoneMask, uint32 NumSource, uint32 NumTarget)
 {
-	uint InvGlobalRefPoseTransformBufferOffset;
-	uint SrcTransformBufferOffset;
-	uint DstTransformBufferOffset1;
-	uint DstTransformBufferOffset2;
-	uint BoneMapBufferOffset;
-	uint MaxBoneTransformCount : 16;
-	uint Padding               : 16;
-};
-SKINNING_STATIC_ASSERT(sizeof(FScatterTransformsBlockHeader) == 24, "FScatterTransformsBlockHeader must be 24 bytes");
+	H.BoneMaskAndSampleCounts = (BoneMask & 0x3FFFFFF) | (NumSource << 26) | (NumTarget << 29);
+}
+#else
+uint  GetBoneMaskOffset(FEvaluateLayerHeader H)      { return H.BoneMaskAndSampleCounts & 0x3FFFFFF; }
+uint  GetNumSourceSamples(FEvaluateLayerHeader H)    { return (H.BoneMaskAndSampleCounts >> 26) & 0x7; }
+uint  GetNumTargetSamples(FEvaluateLayerHeader H)    { return (H.BoneMaskAndSampleCounts >> 29) & 0x7; }
+bool  HasBoneMask(FEvaluateLayerHeader H)             { return GetBoneMaskOffset(H) != SKINNING_INVALID_BONE_MASK_OFFSET; }
+#endif
 
 struct FResolveAttachmentBlockHeader
 {
@@ -294,7 +229,6 @@ struct FCompressedBoneTransform
 #define ANIM_BANK_GPU_TRANSFORM_PROVIDER_GUID 0xA5C0027A, 0x8F884C7C, 0x9312F138, 0x71A9300F
 #define ANIM_BANK_CPU_TRANSFORM_PROVIDER_GUID 0xE7D6173D, 0x246F431A, 0x912D384E, 0x156C0D2C
 #define ANIM_SEQUENCE_GPU_TRANSFORM_PROVIDER_GUID 0x5CFA81C1, 0xF3D8486B, 0x81FF23D3, 0x72EB523F
-#define ANIM_SEQUENCE_CPU_TRANSFORM_PROVIDER_GUID 0xBA5349A0, 0xCDFC4F76, 0x987168CE, 0xC7AA0E7B
 #define ANIM_RUNTIME_TRANSFORM_PROVIDER_GUID 0xF1508490, 0xFCC24BB9, 0xA9F277B3, 0x1AF766F0
 #define ANIM_MESH_OBJECT_TRANSFORM_PROVIDER_GUID 0x1073C0C5, 0xB981F6FD, 0xBF41B1C2, 0xC4BA0C50
 
diff --git a/Engine/Source/Editor/AnimGraph/Private/AnimGraphModule.cpp b/Engine/Source/Editor/AnimGraph/Private/AnimGraphModule.cpp
--- a/Engine/Source/Editor/AnimGraph/Private/AnimGraphModule.cpp
+++ b/Engine/Source/Editor/AnimGraph/Private/AnimGraphModule.cpp
@@ -34,6 +34,7 @@
 #include "Features/IModularFeatures.h"
 #include "PropertyAccessAnimBlueprintBinding.h"
 #include "AnimBlueprintDetails.h"
+#include "AnimSequenceTransformProviderLayerCustomization.h"
 
 class FAnimGraphModule : public IModuleInterface
 {
@@ -80,7 +81,7 @@ void FAnimGraphModule::StartupModule()
 	PropertyModule.RegisterCustomClassLayout(UAnimGraphNode_PoseDriver::StaticClass()->GetFName(), FOnGetDetailCustomizationInstance::CreateStatic(&FPoseDriverDetails::MakeInstance));
 
 	PropertyModule.RegisterCustomPropertyTypeLayout("AnimBlueprintFunctionPinInfo", FOnGetPropertyTypeCustomizationInstance::CreateStatic(&FAnimBlueprintFunctionPinInfoDetails::MakeInstance));
-
+	PropertyModule.RegisterCustomPropertyTypeLayout("AnimSequenceTransformProviderLayer", FOnGetPropertyTypeCustomizationInstance::CreateStatic(&FAnimSequenceTransformProviderLayerCustomization::MakeInstance));
 	if(GIsEditor)
 	{
 		// Register built-in graph attributes
diff --git a/Engine/Source/Editor/AnimGraph/Private/AnimSequenceTransformProviderFactories.cpp b/Engine/Source/Editor/AnimGraph/Private/AnimSequenceTransformProviderFactories.cpp
new file mode 100644
--- /dev/null
+++ b/Engine/Source/Editor/AnimGraph/Private/AnimSequenceTransformProviderFactories.cpp
@@ -0,0 +1,28 @@
+// Copyright Epic Games, Inc. All Rights Reserved.
+
+#include "AnimSequenceTransformProviderFactories.h"
+#include "Animation/AnimSequenceTransformProviderData.h"
+
+UAnimSequenceTransformProviderLayerStackFactory::UAnimSequenceTransformProviderLayerStackFactory()
+{
+	SupportedClass = UAnimSequenceTransformProviderLayerStack::StaticClass();
+	bCreateNew = true;
+	bEditAfterNew = true;
+}
+
+UObject* UAnimSequenceTransformProviderLayerStackFactory::FactoryCreateNew(UClass* Class, UObject* InParent, FName Name, EObjectFlags Flags, UObject* Context, FFeedbackContext* Warn)
+{
+	return NewObject<UAnimSequenceTransformProviderLayerStack>(InParent, Class, Name, Flags);
+}
+
+UAnimSequenceTransformProviderSequenceListFactory::UAnimSequenceTransformProviderSequenceListFactory()
+{
+	SupportedClass = UAnimSequenceTransformProviderSequenceList::StaticClass();
+	bCreateNew = true;
+	bEditAfterNew = true;
+}
+
+UObject* UAnimSequenceTransformProviderSequenceListFactory::FactoryCreateNew(UClass* Class, UObject* InParent, FName Name, EObjectFlags Flags, UObject* Context, FFeedbackContext* Warn)
+{
+	return NewObject<UAnimSequenceTransformProviderSequenceList>(InParent, Class, Name, Flags);
+}
diff --git a/Engine/Source/Editor/AnimGraph/Private/AnimSequenceTransformProviderFactories.h b/Engine/Source/Editor/AnimGraph/Private/AnimSequenceTransformProviderFactories.h
new file mode 100644
--- /dev/null
+++ b/Engine/Source/Editor/AnimGraph/Private/AnimSequenceTransformProviderFactories.h
@@ -0,0 +1,26 @@
+// Copyright Epic Games, Inc. All Rights Reserved.
+
+#pragma once
+
+#include "Factories/Factory.h"
+#include "AnimSequenceTransformProviderFactories.generated.h"
+
+UCLASS()
+class UAnimSequenceTransformProviderLayerStackFactory : public UFactory
+{
+	GENERATED_BODY()
+
+public:
+	UAnimSequenceTransformProviderLayerStackFactory();
+	virtual UObject* FactoryCreateNew(UClass* Class, UObject* InParent, FName Name, EObjectFlags Flags, UObject* Context, FFeedbackContext* Warn) override;
+};
+
+UCLASS()
+class UAnimSequenceTransformProviderSequenceListFactory : public UFactory
+{
+	GENERATED_BODY()
+
+public:
+	UAnimSequenceTransformProviderSequenceListFactory();
+	virtual UObject* FactoryCreateNew(UClass* Class, UObject* InParent, FName Name, EObjectFlags Flags, UObject* Context, FFeedbackContext* Warn) override;
+};
diff --git a/Engine/Source/Editor/AnimGraph/Private/AnimSequenceTransformProviderLayerCustomization.cpp b/Engine/Source/Editor/AnimGraph/Private/AnimSequenceTransformProviderLayerCustomization.cpp
new file mode 100644
--- /dev/null
+++ b/Engine/Source/Editor/AnimGraph/Private/AnimSequenceTransformProviderLayerCustomization.cpp
@@ -0,0 +1,35 @@
+// Copyright Epic Games, Inc. All Rights Reserved.
+
+#include "AnimSequenceTransformProviderLayerCustomization.h"
+#include "Animation/AnimSequenceTransformProviderData.h"
+#include "DetailWidgetRow.h"
+#include "IDetailChildrenBuilder.h"
+#include "PropertyHandle.h"
+
+TSharedRef<IPropertyTypeCustomization> FAnimSequenceTransformProviderLayerCustomization::MakeInstance()
+{
+	return MakeShareable(new FAnimSequenceTransformProviderLayerCustomization());
+}
+
+void FAnimSequenceTransformProviderLayerCustomization::CustomizeHeader(TSharedRef<IPropertyHandle> PropertyHandle, FDetailWidgetRow& HeaderRow, IPropertyTypeCustomizationUtils& Utils)
+{
+	HeaderRow.NameContent()
+	[
+		PropertyHandle->CreatePropertyNameWidget()
+	];
+}
+
+void FAnimSequenceTransformProviderLayerCustomization::CustomizeChildren(TSharedRef<IPropertyHandle> PropertyHandle, IDetailChildrenBuilder& ChildBuilder, IPropertyTypeCustomizationUtils& Utils)
+{
+	const int32 LayerIndex = PropertyHandle->GetIndexInArray();
+	const bool bIsBaseLayer = (LayerIndex == 0);
+
+	uint32 NumChildren = 0;
+	PropertyHandle->GetNumChildren(NumChildren);
+
+	for (uint32 ChildIndex = 0; ChildIndex < NumChildren; ChildIndex++)
+	{
+		TSharedRef<IPropertyHandle> ChildHandle = PropertyHandle->GetChildHandle(ChildIndex).ToSharedRef();
+		ChildBuilder.AddProperty(ChildHandle).IsEnabled(!bIsBaseLayer);
+	}
+}
diff --git a/Engine/Source/Editor/AnimGraph/Private/AnimSequenceTransformProviderLayerCustomization.h b/Engine/Source/Editor/AnimGraph/Private/AnimSequenceTransformProviderLayerCustomization.h
new file mode 100644
--- /dev/null
+++ b/Engine/Source/Editor/AnimGraph/Private/AnimSequenceTransformProviderLayerCustomization.h
@@ -0,0 +1,14 @@
+// Copyright Epic Games, Inc. All Rights Reserved.
+
+#pragma once
+
+#include "IPropertyTypeCustomization.h"
+
+class FAnimSequenceTransformProviderLayerCustomization : public IPropertyTypeCustomization
+{
+public:
+	static TSharedRef<IPropertyTypeCustomization> MakeInstance();
+
+	virtual void CustomizeHeader(TSharedRef<IPropertyHandle> PropertyHandle, FDetailWidgetRow& HeaderRow, IPropertyTypeCustomizationUtils& Utils) override;
+	virtual void CustomizeChildren(TSharedRef<IPropertyHandle> PropertyHandle, IDetailChildrenBuilder& ChildBuilder, IPropertyTypeCustomizationUtils& Utils) override;
+};
diff --git a/Engine/Source/Runtime/Core/Public/UObject/FortniteMainBranchObjectVersions.inl b/Engine/Source/Runtime/Core/Public/UObject/FortniteMainBranchObjectVersions.inl
--- a/Engine/Source/Runtime/Core/Public/UObject/FortniteMainBranchObjectVersions.inl
+++ b/Engine/Source/Runtime/Core/Public/UObject/FortniteMainBranchObjectVersions.inl
@@ -814,5 +814,8 @@ FORTNITE_MAIN_VERSION(ConvertFPCGMetadataAttributeToGenericAttributes, 266)
 // Merging PCG Metadata Attribute base and Generic ones. 
 FORTNITE_MAIN_VERSION(MergePCGMetadataAttributeBaseAndGeneric, 267)
 
+// AnimSequenceTransformProviderData supports multiple layers
+FORTNITE_MAIN_VERSION(AnimSequenceTransformProviderLayers, 268)
+
 // -----<new versions can be added above this line>-------------------------------------------------
 #undef FORTNITE_MAIN_VERSION
diff --git a/Engine/Source/Runtime/Engine/Classes/Animation/AnimBank.h b/Engine/Source/Runtime/Engine/Classes/Animation/AnimBank.h
--- a/Engine/Source/Runtime/Engine/Classes/Animation/AnimBank.h
+++ b/Engine/Source/Runtime/Engine/Classes/Animation/AnimBank.h
@@ -601,9 +601,9 @@ public:
 		Data.ItemIndex            = ManualData.SequenceIndex;
 		Data.bAutoPlay            = 0;
 		Data.LoopMode             = 0;
-		Data.bHasPreviousPosition = ManualData.bHasPreviousPosition;
+		Data.bHasPreviousPosition = false;
 		Data.Position             = ManualData.Position;
-		Data.PreviousPosition     = ManualData.bHasPreviousPosition ? ManualData.PreviousPosition : ManualData.Position;
+		Data.PreviousPosition     = ManualData.Position;
 		Data.PlayRate             = 0.0f;
 		return Data;
 	}
@@ -620,8 +620,6 @@ public:
 		FAnimSequenceTrackManualData Data;
 		Data.SequenceIndex        = ItemIndex;
 		Data.Position             = Position;
-		Data.PreviousPosition     = PreviousPosition;
-		Data.bHasPreviousPosition = bHasPreviousPosition;
 		return Data;
 	}
 
@@ -658,6 +656,11 @@ public:
 
 	ENGINE_API void Serialize(FArchive& Ar);
 
+	void SetData(const FAnimBankTrackPackedData& RHS)
+	{
+		*this = RHS;
+	}
+
 private:
 	uint32 ItemIndex            : 20;
 	uint32 LoopMode             : 2;
diff --git a/Engine/Source/Runtime/Engine/Classes/Animation/AnimRuntimeTransformProviderData.h b/Engine/Source/Runtime/Engine/Classes/Animation/AnimRuntimeTransformProviderData.h
--- a/Engine/Source/Runtime/Engine/Classes/Animation/AnimRuntimeTransformProviderData.h
+++ b/Engine/Source/Runtime/Engine/Classes/Animation/AnimRuntimeTransformProviderData.h
@@ -73,6 +73,11 @@ struct FAnimRuntimeTrackData
 {
 	FAnimRuntimeTrackTransformData* Current = nullptr;
 	FAnimRuntimeTrackTransformData* Previous = nullptr;
+
+	void SetData(const FAnimRuntimeTrackData& RHS)
+	{
+		*this = RHS;
+	}
 };
 
 inline FArchive& operator <<(FArchive& Ar, FAnimRuntimeTrackData& Data)
diff --git a/Engine/Source/Runtime/Engine/Classes/Animation/AnimSequenceTransformProviderData.h b/Engine/Source/Runtime/Engine/Classes/Animation/AnimSequenceTransformProviderData.h
--- a/Engine/Source/Runtime/Engine/Classes/Animation/AnimSequenceTransformProviderData.h
+++ b/Engine/Source/Runtime/Engine/Classes/Animation/AnimSequenceTransformProviderData.h
@@ -6,9 +6,11 @@
 #include "Experimental/Containers/RobinHoodHashTable.h"
 #include "RenderCommandFence.h"
 #include "SkinningDefinitions.h"
+#include "SkeletalRenderPublic.h"
 #include "UObject/ObjectPtr.h"
 #include "UObject/WeakObjectPtr.h"
 #include "AnimTrackPool.h"
+#include "Engine/DataAsset.h"
 
 #if WITH_EDITOR
 #include "IO/IoHash.h"
@@ -20,6 +22,7 @@ class ITargetPlatform;
 class USkinnedAsset;
 class UAnimBankData;
 class UAnimSequence;
+class UBlendSpace;
 class FAnimSequenceTransformProviderBuildAsyncCacheTask;
 class UAnimSequenceTransformProviderData;
 class FAnimSequenceTransformProviderProxy;
@@ -36,6 +39,36 @@ enum class EAnimSequenceTrackLoopMode : uint8
 	Clamp
 };
 
+UENUM(BlueprintType)
+enum class EAnimSequenceTransformProviderLayerBlendMode : uint8
+{
+	// Replaces the accumulated result. Weight controls blend factor.
+	Override,
+
+	// Adds a weighted delta on top of the accumulated result.
+	Additive,
+};
+
+USTRUCT(BlueprintType)
+struct FAnimSequenceTransformProviderLayer
+{
+	GENERATED_USTRUCT_BODY()
+
+	UPROPERTY(EditAnywhere, BlueprintReadWrite, Category = Layer)
+	EAnimSequenceTransformProviderLayerBlendMode BlendMode = EAnimSequenceTransformProviderLayerBlendMode::Override;
+
+	UPROPERTY(EditAnywhere, BlueprintReadWrite, Category = Layer, meta = (ClampMin = "0.0", ClampMax = "1.0"))
+	float Weight = 1.0f;
+
+	// Name of a blend profile on the skeleton to use as a per-bone weight mask. Empty means all bones at full weight.
+	UPROPERTY(EditAnywhere, BlueprintReadWrite, Category = Layer)
+	FName BoneMaskProfileName;
+
+	// Blend rotation in global space instead of local space.
+	UPROPERTY(EditAnywhere, BlueprintReadWrite, Category = Layer)
+	bool bMeshSpaceRotation = false;
+};
+
 USTRUCT(BlueprintType)
 struct FAnimSequenceTrackAutoPlayData
 {
@@ -71,19 +104,45 @@ struct FAnimSequenceTrackManualData
 	UPROPERTY(EditAnywhere, BlueprintReadWrite, Category = Animation)
 	int32 SequenceIndex = 0;
 
-	// The time position so sample in the sequence.
+	// The time position to sample in the sequence.
 	UPROPERTY(EditAnywhere, BlueprintReadWrite, Category = Animation)
 	float Position = 0.0f;
 
-	// Whether to apply the previous position when updating the animation. When false, the previous position is derived from the last frame.
+	// The duration of the blend transition in seconds. Use 0 for an instant switch.
+	UPROPERTY(EditAnywhere, BlueprintReadWrite, Category = Animation, meta = (ClampMin = "0.0"))
+	float BlendTime = 0.0f;
+
+	// The mode to control looping behavior for when the position is not within the sequence bounds.
 	UPROPERTY(EditAnywhere, BlueprintReadWrite, Category = Animation)
-	uint32 bHasPreviousPosition : 1 = 0;
+	EAnimSequenceTrackLoopMode LoopMode = EAnimSequenceTrackLoopMode::Loop;
+};
 
-	// A time position to update previous frame transforms.
-	UPROPERTY(EditAnywhere, BlueprintReadWrite, Category = Animation, meta = (EditCondition = "bHasPreviousPosition", EditConditionHides))
-	float PreviousPosition = 0.0f;
+USTRUCT(BlueprintType)
+struct FAnimSequenceTrackBlendSpaceData
+{
+	GENERATED_USTRUCT_BODY()
 
-	// The mode to control looping behavior for when the position is not within the sequence bounds.
+	// Index into the ASTP's BlendSpaces array.
+	UPROPERTY(EditAnywhere, BlueprintReadWrite, Category = Animation)
+	int32 BlendSpaceIndex = 0;
+
+	// The 2D blend position.
+	UPROPERTY(EditAnywhere, BlueprintReadWrite, Category = Animation)
+	FVector2f BlendPosition = FVector2f::ZeroVector;
+
+	// The starting time position to begin play.
+	UPROPERTY(EditAnywhere, BlueprintReadWrite, Category = Animation)
+	float Position = 0.0f;
+
+	// The rate multiplier to play the animation.
+	UPROPERTY(EditAnywhere, BlueprintReadWrite, Category = Animation)
+	float PlayRate = 1.0f;
+
+	// The duration of the blend transition in seconds. Use 0 for an instant switch.
+	UPROPERTY(EditAnywhere, BlueprintReadWrite, Category = Animation, meta = (ClampMin = "0.0"))
+	float BlendTime = 0.0f;
+
+	// The mode to control looping behavior.
 	UPROPERTY(EditAnywhere, BlueprintReadWrite, Category = Animation)
 	EAnimSequenceTrackLoopMode LoopMode = EAnimSequenceTrackLoopMode::Loop;
 };
@@ -105,6 +164,15 @@ struct FAnimSequenceTransformProviderSequence
 	float Position = 0.0f;
 };
 
+USTRUCT(BlueprintType)
+struct FAnimSequenceTransformProviderBlendSpace
+{
+	GENERATED_USTRUCT_BODY()
+
+	UPROPERTY(EditAnywhere, BlueprintReadWrite, Category = Animation)
+	TObjectPtr<UBlendSpace> BlendSpace;
+};
+
 class FAnimSequenceTransformProviderProxy;
 class UAnimSequenceTransformProviderDataInstance;
 
@@ -113,7 +181,7 @@ enum class EAnimSequenceTrackPatchFlags : uint8
 	None         = 0,
 	PlayRate     = 1 << 0,
 	LoopMode     = 1 << 1,
-	All          = 0xFF
+	All          = 0x3
 };
 ENUM_CLASS_FLAGS(EAnimSequenceTrackPatchFlags);
 
@@ -130,12 +198,10 @@ public:
 		Data.bAutoPlay = 1;
 		Data.LoopMode = static_cast<uint32>(AutoPlayData.LoopMode);
 		Data.PatchFlags = static_cast<uint32>(EAnimSequenceTrackPatchFlags::All);
-		Data.bHasPreviousPosition = false;
 		Data.Auto.ReferenceTimestamp = GetTimestamp(WorldTime, AutoPlayData.Position, AutoPlayData.PlayRate);
 		Data.Auto.Position = AutoPlayData.Position;
-		Data.Auto.PreviousPosition = AutoPlayData.Position;
-		Data.Auto.PlayRate = AutoPlayData.PlayRate;
-		Data.Auto.BlendTime = AutoPlayData.BlendTime;
+		Data.Auto.PlayRate = FFloat16(AutoPlayData.PlayRate);
+		Data.BlendTime = FFloat16(AutoPlayData.BlendTime);
 		return Data;
 	}
 
@@ -154,42 +220,133 @@ public:
 		Data.bAutoPlay = 0;
 		Data.LoopMode = static_cast<uint32>(ManualData.LoopMode);
 		Data.PatchFlags = static_cast<uint32>(EAnimSequenceTrackPatchFlags::All);
-		Data.bHasPreviousPosition = ManualData.bHasPreviousPosition;
-		Data.Manual.Position = GetPosition(ManualData.LoopMode, ManualData.Position, SequenceLength);
-		Data.Manual.PreviousPosition = ManualData.bHasPreviousPosition ? GetPosition(ManualData.LoopMode, ManualData.PreviousPosition, SequenceLength) : Data.Manual.Position;
+		Data.Manual.Position = WrapPosition(ManualData.LoopMode, ManualData.Position, SequenceLength);
+		Data.BlendTime = FFloat16(ManualData.BlendTime);
 		return Data;
 	}
 
+	/////////////////////////////////////////////////////////////////////////////////////////////////
+	// BlendSpace
+
+	static FAnimSequenceTrackPackedData Pack(const FAnimSequenceTrackBlendSpaceData& BlendSpaceData, double WorldTime)
+	{
+		FAnimSequenceTrackPackedData Data;
+		Data.bBlendSpace = 1;
+		Data.bAutoPlay = 0;
+		Data.LoopMode = static_cast<uint32>(BlendSpaceData.LoopMode);
+		Data.PatchFlags = static_cast<uint32>(EAnimSequenceTrackPatchFlags::All);
+		Data.BlendSpace.ReferenceTimestamp = GetTimestamp(WorldTime, BlendSpaceData.Position, BlendSpaceData.PlayRate);
+		Data.BlendSpace.Position = BlendSpaceData.Position;
+		Data.BlendSpace.PlayRate = FFloat16(BlendSpaceData.PlayRate);
+		Data.BlendSpace.BlendPosition[0] = FFloat16(BlendSpaceData.BlendPosition.X);
+		Data.BlendSpace.BlendPosition[1] = FFloat16(BlendSpaceData.BlendPosition.Y);
+		Data.BlendSpace.Samples[0] = 0;
+		Data.BlendSpace.Samples[1] = 0;
+		Data.BlendSpace.Samples[2] = 0;
+		Data.BlendSpace.BlendSpaceIndex = static_cast<uint8>(BlendSpaceData.BlendSpaceIndex);
+		Data.BlendTime = FFloat16(BlendSpaceData.BlendTime);
+		return Data;
+	}
+
+	bool IsBlendSpace() const
+	{
+		return bBlendSpace;
+	}
+
+	FVector2f GetBlendSpacePosition() const
+	{
+		return FVector2f(BlendSpace.BlendPosition[0].GetFloat(), BlendSpace.BlendPosition[1].GetFloat());
+	}
+
+	void UpdateBlendSpaceData(const FAnimSequenceTrackPackedData& Data, double WorldTime)
+	{
+		if (BlendSpace.PlayRate.Encoded != Data.BlendSpace.PlayRate.Encoded)
+		{
+			SetPlayRate(Data.BlendSpace.PlayRate.GetFloat(), WorldTime);
+		}
+		BlendSpace.BlendPosition[0] = Data.BlendSpace.BlendPosition[0];
+		BlendSpace.BlendPosition[1] = Data.BlendSpace.BlendPosition[1];
+		BlendSpace.Samples[0] = Data.BlendSpace.Samples[0];
+		BlendSpace.Samples[1] = Data.BlendSpace.Samples[1];
+		BlendSpace.Samples[2] = Data.BlendSpace.Samples[2];
+		BlendSpace.NumSamples = Data.BlendSpace.NumSamples;
+		LoopMode = Data.LoopMode;
+		BlendTime = Data.BlendTime;
+	}
+
+	uint32 GetNumSamples() const
+	{
+		return bBlendSpace ? BlendSpace.NumSamples : 1;
+	}
+
+	void SetSampleData(uint32 SampleIndex, uint16 InSequenceIndex, FFloat16 Weight)
+	{
+		BlendSpace.Samples[SampleIndex] = (uint32(InSequenceIndex) & 0xFFFF) | (uint32(Weight.Encoded) << 16);
+	}
+
+	uint16 GetSampleSequenceIndex(uint32 SampleIndex) const
+	{
+		return BlendSpace.Samples[SampleIndex] & 0xFFFF;
+	}
+
+	FFloat16 GetSampleWeight(uint32 SampleIndex) const
+	{
+		FFloat16 Result;
+		Result.Encoded = BlendSpace.Samples[SampleIndex] >> 16;
+		return Result;
+	}
+
+	void SetNumSamples(uint32 InNumSamples)
+	{
+		BlendSpace.NumSamples = static_cast<uint8>(InNumSamples);
+	}
+
+	uint8 GetBlendSpaceIndex() const
+	{
+		return BlendSpace.BlendSpaceIndex;
+	}
+
 	bool IsManual() const
 	{
-		return !bAutoPlay;
+		return !bAutoPlay && !bBlendSpace;
 	}
 
 	/////////////////////////////////////////////////////////////////////////////////////////////////
 
 	int32 GetSequenceIndex() const
 	{
-		return SequenceIndex;
+		return bBlendSpace ? INDEX_NONE : static_cast<int32>(SequenceIndex);
 	}
 
-	float GetPosition(double WorldTime, double SequenceLength) const
+	// Returns the raw unwrapped time position. Caller must fmod against sequence play length.
+	double GetTimePosition(double WorldTime) const
 	{
-		if (bAutoPlay)
+		if (IsManual())
 		{
-			const double Position = (WorldTime - Auto.ReferenceTimestamp) * Auto.PlayRate;
-			return GetPosition(GetLoopMode(), static_cast<float>(Position), static_cast<float>(SequenceLength));
+			return static_cast<double>(Manual.Position);
 		}
-		return Manual.Position;
-	}
 
-	float GetPreviousPosition() const
-	{
-		return bAutoPlay ? Auto.PreviousPosition : Manual.PreviousPosition;
+		if (IsBlendSpace())
+		{
+			return (WorldTime - BlendSpace.ReferenceTimestamp) * BlendSpace.PlayRate.GetFloat();
+		}
+
+		return (WorldTime - Auto.ReferenceTimestamp) * Auto.PlayRate.GetFloat();
 	}
 
-	bool HasPreviousPosition() const
+	bool IsSameAnimation(const FAnimSequenceTrackPackedData& Other) const
 	{
-		return bHasPreviousPosition;
+		if (bBlendSpace != Other.bBlendSpace)
+		{
+			return false;
+		}
+
+		if (bBlendSpace)
+		{
+			return BlendSpace.BlendSpaceIndex == Other.BlendSpace.BlendSpaceIndex;
+		}
+
+		return SequenceIndex == Other.SequenceIndex;
 	}
 
 	void SetLoopMode(EAnimSequenceTrackLoopMode NewLoopMode)
@@ -205,7 +362,7 @@ public:
 
 	float GetBlendTime() const
 	{
-		return IsAutoPlay() ? Auto.BlendTime : 0.0f;
+		return BlendTime.GetFloat();
 	}
 
 	void ResetPatchFlags()
@@ -220,39 +377,65 @@ public:
 
 	void SetPlayRate(float NewPlayRate, double WorldTime)
 	{
-		check(IsAutoPlay());
-		const double Position = (WorldTime - Auto.ReferenceTimestamp) * Auto.PlayRate;
-		Auto.ReferenceTimestamp = GetTimestamp(WorldTime, static_cast<float>(Position), NewPlayRate);
-		Auto.PlayRate = NewPlayRate;
+		check(IsAutoPlay() || IsBlendSpace());
+
+		if (IsBlendSpace())
+		{
+			const double Position = (WorldTime - BlendSpace.ReferenceTimestamp) * BlendSpace.PlayRate.GetFloat();
+			BlendSpace.ReferenceTimestamp = GetTimestamp(WorldTime, static_cast<float>(Position), NewPlayRate);
+			BlendSpace.PlayRate = FFloat16(NewPlayRate);
+		}
+		else
+		{
+			const double Position = (WorldTime - Auto.ReferenceTimestamp) * Auto.PlayRate.GetFloat();
+			Auto.ReferenceTimestamp = GetTimestamp(WorldTime, static_cast<float>(Position), NewPlayRate);
+			Auto.PlayRate = FFloat16(NewPlayRate);
+		}
+
 		PatchFlags |= (uint32)EAnimSequenceTrackPatchFlags::PlayRate;
 	}
 
 	float GetPlayRate() const
 	{
-		return IsAutoPlay() ? Auto.PlayRate : 0.0f;
+		if (IsBlendSpace())
+		{
+			return BlendSpace.PlayRate.GetFloat();
+		}
+		return IsAutoPlay() ? Auto.PlayRate.GetFloat() : 0.0f;
 	}
 
 	ENGINE_API void Serialize(FArchive& Ar);
 
 	void PostLoad(double WorldTime)
 	{
-		if (IsAutoPlay())
+		if (IsBlendSpace())
+		{
+			BlendSpace.ReferenceTimestamp = GetTimestamp(WorldTime, BlendSpace.Position, BlendSpace.PlayRate.GetFloat());
+		}
+		else if (IsAutoPlay())
 		{
-			Auto.ReferenceTimestamp = GetTimestamp(WorldTime, Auto.Position, Auto.PlayRate);
+			Auto.ReferenceTimestamp = GetTimestamp(WorldTime, Auto.Position, Auto.PlayRate.GetFloat());
 		}
 	}
 
-	void Tick(float Position)
+	void SetLayerWeight(float Weight)
 	{
-		if (IsAutoPlay())
-		{
-			Auto.PreviousPosition = Position;
-			bHasPreviousPosition = true;
-		}
+		LayerWeight = FFloat16(Weight);
 	}
 
-private:
-	static float GetPosition(EAnimSequenceTrackLoopMode LoopMode, float Position, float SequenceLength)
+	float GetLayerWeight() const
+	{
+		return LayerWeight.GetFloat();
+	}
+
+	void SetData(const FAnimSequenceTrackPackedData& RHS)
+	{
+		const FFloat16 PrevLayerWeight = LayerWeight;
+		*this = RHS;
+		LayerWeight = PrevLayerWeight;
+	}
+
+	static float WrapPosition(EAnimSequenceTrackLoopMode LoopMode, float Position, float SequenceLength)
 	{
 		if (LoopMode == EAnimSequenceTrackLoopMode::Loop)
 		{
@@ -260,7 +443,6 @@ private:
 			{
 				Position = SequenceLength == 0.0f ? 0.0f : FMath::Fmod(Position, SequenceLength);
 
-				// Returns a negative range if input is negative.
 				if (Position < 0.0f)
 				{
 					Position += SequenceLength;
@@ -274,6 +456,7 @@ private:
 		return Position;
 	}
 
+private:
 	static double GetTimestamp(double WorldTime, float Position, float PlayRate)
 	{
 		if (FMath::Abs(PlayRate) > UE_SMALL_NUMBER)
@@ -283,11 +466,11 @@ private:
 		return WorldTime;
 	}
 
-	uint32 SequenceIndex : 20;
-	uint32 LoopMode : 1;
-	uint32 bAutoPlay : 1;
-	uint32 bHasPreviousPosition : 1;
-	uint32 PatchFlags : 8;
+	uint32 SequenceIndex        : 20 = 0;
+	uint32 LoopMode             : 1 = 0;
+	uint32 bAutoPlay            : 1 = 0;
+	uint32 bBlendSpace          : 1 = 0;
+	uint32 PatchFlags           : 9 = 0;
 
 	union
 	{
@@ -295,18 +478,29 @@ private:
 		{
 			double ReferenceTimestamp;
 			float Position;
-			float PreviousPosition;
-			float PlayRate;
-			float BlendTime;
+			FFloat16 PlayRate;
 		} Auto;
 
 		struct
 		{
 			float Position;
-			float PreviousPosition;
 		} Manual;
+
+		struct
+		{
+			double ReferenceTimestamp;
+			float Position;
+			uint32 Samples[3];
+			FFloat16 PlayRate;
+			FFloat16 BlendPosition[2];
+			uint8 NumSamples;
+			uint8 BlendSpaceIndex;
+		} BlendSpace = {};
 	};
 
+	FFloat16 BlendTime = FFloat16(0.0f);
+	FFloat16 LayerWeight = FFloat16(1.0f);
+
 	friend class FAnimSequenceTransformProvider;
 };
 
@@ -334,45 +528,219 @@ public:
 	{
 		Source = PackedData;
 		Target = PackedData;
-		BlendWeight = 0.0f;
+		BlendWeight = FFloat16(0.0f);
 		return *this;
 	}
 
 	bool IsBlending() const
 	{
-		// The check for BlendWeight > 0 is for blending two positions of the same sequence.
-		return Source.GetSequenceIndex() != Target.GetSequenceIndex() || BlendWeight > 0.0f;
+		return !Target.IsSameAnimation(Source) || BlendWeight.GetFloat() > 0.0f;
 	}
 
-	float GetBlendWeight() const { return BlendWeight; }
+	float GetBlendWeight(bool bPrevious = false) const
+	{
+		return bPrevious ? PreviousBlendWeight.GetFloat() : BlendWeight.GetFloat();
+	}
+
+	uint32 GetTotalSampleCount(EPreviousBoneTransformUpdateMode UpdateMode = EPreviousBoneTransformUpdateMode::None) const
+	{
+		const bool bUpdatePrevious = (UpdateMode == EPreviousBoneTransformUpdateMode::UpdatePrevious);
+		const uint32 SourceSamples = Source.GetNumSamples();
+		const uint32 TargetSamples = Target.GetNumSamples();
+		const uint32 CurrentSamples = SourceSamples + ((BlendWeight.GetFloat() > 0.0f) ? TargetSamples : 0);
+		const uint32 PreviousSamples = SourceSamples + ((PreviousBlendWeight.GetFloat() > 0.0f) ? TargetSamples : 0);
+
+		return CurrentSamples + (bUpdatePrevious ? PreviousSamples : 0);
+	}
+
+	const FAnimSequenceTrackPackedData& GetSource() const
+	{
+		return Source;
+	}
 
-	const FAnimSequenceTrackPackedData& GetSource() const { return Source; }
-	const FAnimSequenceTrackPackedData& GetTarget() const { return Target; }
+	const FAnimSequenceTrackPackedData& GetTarget() const
+	{
+		return Target;
+	}
 
 	ENGINE_API void Patch(const FAnimSequenceTrackPackedData& Request, double WorldTime, bool bWasActive);
 
-	ENGINE_API void Tick(float DeltaTime, float SourcePosition, float TargetPosition);
+	ENGINE_API void Tick(float DeltaTime, double WorldTime);
+
+	bool HasPreviousPosition() const
+	{
+		return bHasPreviousPosition;
+	}
+
+	double GetPreviousTimePosition() const
+	{
+		return PreviousTimePosition;
+	}
+
+	FVector2f GetPreviousBlendSpacePosition() const
+	{
+		return FVector2f(PreviousBlendSpacePosition[0].GetFloat(), PreviousBlendSpacePosition[1].GetFloat());
+	}
 
 private:
 	FAnimSequenceTrackPackedData Source;
 	FAnimSequenceTrackPackedData Target;
-	float BlendWeight = 0.0f;
+
+	double PreviousTimePosition = 0.0;
+	FFloat16 PreviousBlendSpacePosition[2] = { FFloat16(0.0f), FFloat16(0.0f) };
+	FFloat16 PreviousBlendWeight = FFloat16(0.0f);
+	FFloat16 BlendWeight = FFloat16(0.0f);
+	bool bHasPreviousPosition = false;
 };
 
 using FAnimSequenceTrackRenderPool = TAnimTrackPool<FAnimSequenceTrackRenderData>;
 
+UCLASS(BlueprintType, MinimalAPI)
+class UAnimSequenceTransformProviderLayerStack : public UDataAsset
+{
+	GENERATED_BODY()
+
+public:
+	UAnimSequenceTransformProviderLayerStack()
+	{
+		Layers.Emplace();
+	}
+
+	const TArray<FAnimSequenceTransformProviderLayer>& GetLayers() const
+	{
+		return Layers;
+	}
+
+	void SetLayers(const TArray<FAnimSequenceTransformProviderLayer>& InLayers)
+	{
+		Layers = InLayers;
+	}
+
+	int32 GetNumLayers() const
+	{
+		return Layers.Num();
+	}
+
+#if WITH_EDITOR
+	ENGINE_API virtual void PostEditChangeProperty(FPropertyChangedEvent& PropertyChangedEvent) override;
+#endif
+
+private:
+	UPROPERTY(EditAnywhere, BlueprintReadOnly, Category = Layers, meta = (AllowPrivateAccess = "true"))
+	TArray<FAnimSequenceTransformProviderLayer> Layers;
+};
+
+UCLASS(BlueprintType, MinimalAPI)
+class UAnimSequenceTransformProviderSequenceList : public UDataAsset
+{
+	GENERATED_BODY()
+
+public:
+	const TArray<FAnimSequenceTransformProviderSequence>& GetSequences() const
+	{
+		return Sequences;
+	}
+
+	void SetSequences(const TArray<FAnimSequenceTransformProviderSequence>& InSequences)
+	{
+		Sequences = InSequences;
+	}
+
+#if WITH_EDITOR
+	ENGINE_API virtual void PostEditChangeProperty(FPropertyChangedEvent& PropertyChangedEvent) override;
+#endif
+
+private:
+	UPROPERTY(EditAnywhere, BlueprintReadOnly, Category = Sequences, meta = (AllowPrivateAccess = "true"))
+	TArray<FAnimSequenceTransformProviderSequence> Sequences;
+};
+
+UCLASS(BlueprintType, MinimalAPI)
+class UAnimSequenceTransformProviderBlendSpaceList : public UDataAsset
+{
+	GENERATED_BODY()
+
+public:
+	const TArray<FAnimSequenceTransformProviderBlendSpace>& GetBlendSpaces() const
+	{
+		return BlendSpaces;
+	}
+
+	void SetBlendSpaces(const TArray<FAnimSequenceTransformProviderBlendSpace>& InBlendSpaces)
+	{
+		BlendSpaces = InBlendSpaces;
+	}
+
+#if WITH_EDITOR
+	ENGINE_API virtual void PostEditChangeProperty(FPropertyChangedEvent& PropertyChangedEvent) override;
+#endif
+
+private:
+	UPROPERTY(EditAnywhere, BlueprintReadOnly, Category = BlendSpaces, meta = (AllowPrivateAccess = "true"))
+	TArray<FAnimSequenceTransformProviderBlendSpace> BlendSpaces;
+};
+
+struct FAnimSequenceTransformProviderRenderLayer
+{
+	FAnimSequenceTrackRenderPool Tracks;
+	FAnimSequenceTransformProviderLayer Layer;
+	float Weight = 1.0f;
+};
+
 class FAnimSequenceTransformProviderRenderData
 {
 public:
 	ENGINE_API ~FAnimSequenceTransformProviderRenderData();
 
+	TArrayView<FAnimSequenceTransformProviderRenderLayer> GetLayers()
+	{
+		return Layers;
+	}
+
+	TConstArrayView<FAnimSequenceTransformProviderRenderLayer> GetLayers() const
+	{
+		return Layers;
+	}
+
+	int32 GetNumLayers() const
+	{
+		return Layers.Num();
+	}
+
+	uint32 GetMeshSpaceMask() const
+	{
+		return MeshSpaceMask;
+	}
+
+	uint32 GetLayerModeMask() const
+	{
+		return LayerModeMask;
+	}
+
+	bool IsTrackActive(int32 TrackIndex, int32 LayerIndex) const
+	{
+		return (LayerIndex == 0 || Layers[LayerIndex].Weight > 0.0f) && Layers[LayerIndex].Tracks.IsActiveIndex(TrackIndex);
+	}
+
+	int32 GetNumActiveLayers(int32 TrackIndex) const
+	{
+		int32 Count = 0;
+		for (int32 LayerIndex = 0; LayerIndex < Layers.Num(); LayerIndex++)
+		{
+			Count += IsTrackActive(TrackIndex, LayerIndex) ? 1 : 0;
+		}
+		return Count;
+	}
+
 private:
-	void Patch(FAnimSequenceTrackPool::FPatch&& Patch, double WorldTime);
+	void Patch(TArray<FAnimSequenceTrackPool::FPatch>&& InPatches, TArray<FAnimSequenceTransformProviderRenderLayer>&& InLayers, double WorldTime);
 	void InsertProxy(FInstancedSkinningSceneExtensionProxy* Proxy);
 	void RemoveProxy(FInstancedSkinningSceneExtensionProxy* Proxy);
 
-	FAnimSequenceTrackRenderPool Tracks;
+	TArray<FAnimSequenceTransformProviderRenderLayer> Layers;
 	Experimental::TRobinHoodHashSet<FInstancedSkinningSceneExtensionProxy*> Proxies;
+	uint32 MeshSpaceMask = 0;
+	uint32 LayerModeMask = 0;
 
 	friend FAnimSequenceTransformProviderProxy;
 	friend UAnimSequenceTransformProviderDataInstance;
@@ -385,16 +753,12 @@ class FAnimSequenceTransformProviderProxy : public FTransformProviderRenderProxy
 
 	ENGINE_API FAnimSequenceTransformProviderProxy(TConstArrayView<FAnimSequenceTransformProviderSequence> Sequences, FInstancedSkinningSceneExtensionProxy* SceneExtensionProxy);
 	ENGINE_API FAnimSequenceTransformProviderProxy(TConstArrayView<FAnimSequenceTransformProviderSequence> Sequences, FInstancedSkinningSceneExtensionProxy* SceneExtensionProxy, FAnimSequenceTransformProviderRenderData& RenderData);
+	ENGINE_API ~FAnimSequenceTransformProviderProxy();
 
 public:
 	virtual void CreateRenderThreadResources(FSceneInterface& Scene, FRHICommandListBase& RHICmdList) override;
 	virtual void DestroyRenderThreadResources() override;
 
-	FAnimSequenceTrackRenderPool& GetTracks()
-	{
-		return *Tracks;
-	}
-
 	TConstArrayView<FAnimSequenceTransformProviderSequence> GetSequences() const
 	{
 		return Sequences;
@@ -405,14 +769,28 @@ public:
 		return SceneExtensionProxy;
 	}
 
+	FAnimSequenceTransformProviderRenderData* GetRenderData() const
+	{
+		return RenderData;
+	}
+
+	TConstArrayView<FAnimSequenceTransformProviderRenderLayer> GetLayers() const
+	{
+		return RenderData->GetLayers();
+	}
+
+	int32 GetNumLayers() const
+	{
+		return RenderData->GetNumLayers();
+	}
+
 private:
 	TConstArrayView<FAnimSequenceTransformProviderSequence> Sequences;
-	TUniquePtr<FAnimSequenceTrackRenderPool> TrackAllocator;
 	FAnimSequenceTransformProviderRenderData* RenderData = nullptr;
-	FAnimSequenceTrackRenderPool* Tracks;
 	FInstancedSkinningSceneExtensionProxy* SceneExtensionProxy = nullptr;
 	const USkinnedAsset* SkinnedAsset = nullptr;
 	IAnimSequenceTransformProvider* AnimSequenceTransformProvider = nullptr;
+	bool bOwnsRenderData = false;
 };
 
 class IAnimSequenceTransformProvider
@@ -427,11 +805,24 @@ public:
 struct FAnimSequenceTransformProviderCachedData
 {
 	TArray<FBoxSphereBounds> SequenceBounds;
+	FBoxSphereBounds ConservativeBounds = FBoxSphereBounds(ForceInit);
 };
 
 inline FArchive& operator<<(FArchive& Ar, FAnimSequenceTransformProviderCachedData& CachedData)
 {
 	Ar << CachedData.SequenceBounds;
+
+	// Derive conservative bounds from sequence bounds on load.
+	if (Ar.IsLoading())
+	{
+		CachedData.ConservativeBounds = FBoxSphereBounds(ForceInit);
+		for (int32 Index = 0; Index < CachedData.SequenceBounds.Num(); Index++)
+		{
+			CachedData.ConservativeBounds = (Index == 0)
+				? CachedData.SequenceBounds[Index]
+				: CachedData.ConservativeBounds + CachedData.SequenceBounds[Index];
+		}
+	}
 	return Ar;
 }
 
@@ -454,8 +845,6 @@ public:
 	virtual bool IsReadyForFinishDestroy() override;
 	virtual UTransformProviderData* GetRootTransformProvider() override;
 	virtual void Serialize(FArchive& Ar) override;
-	virtual void PreSave(FObjectPreSaveContext SaveContext) override;
-	virtual void PostSaveRoot(FObjectPostSaveRootContext SaveContext) override;
 	virtual void PostLoad() override;
 
 #if WITH_EDITOR
@@ -478,16 +867,52 @@ public:
 public:
 	// Returns the skinned asset for the transform provider.
 	UFUNCTION(BlueprintCallable, Category = "TransformProviders|AnimSequence")
-	ENGINE_API USkinnedAsset* GetSkinnedAsset() const;
+	USkinnedAsset* GetSkinnedAsset() const
+	{
+		return SkinnedAsset;
+	}
 
-	// Returns the play length for a sequence.
 	UFUNCTION(BlueprintCallable, Category = "TransformProviders|AnimSequence")
-	ENGINE_API const TArray<struct FAnimSequenceTransformProviderSequence>& GetSequences() const;
+	const TArray<struct FAnimSequenceTransformProviderSequence>& GetSequences() const
+	{
+		return Sequences;
+	}
 
 	// Returns the play length for a sequence.
 	UFUNCTION(BlueprintCallable, Category = "TransformProviders|AnimSequence")
 	ENGINE_API float GetSequencePlayLength(int32 SequenceIndex) const;
 
+	UFUNCTION(BlueprintCallable, Category = "TransformProviders|AnimSequence")
+	const TArray<struct FAnimSequenceTransformProviderLayer>& GetLayers() const
+	{
+		return Layers;
+	}
+
+	UFUNCTION(BlueprintCallable, Category = "TransformProviders|AnimSequence")
+	const TArray<FAnimSequenceTransformProviderBlendSpace>& GetBlendSpaces() const
+	{
+		return BlendSpaces;
+	}
+
+#if WITH_EDITORONLY_DATA
+	UAnimSequenceTransformProviderLayerStack* GetLayerStack() const
+	{
+		return LayerStack;
+	}
+
+	UAnimSequenceTransformProviderSequenceList* GetSequenceList() const
+	{
+		return SequenceList;
+	}
+#endif
+
+	// Returns the number of layers. Always >= 1.
+	UFUNCTION(BlueprintCallable, Category = "TransformProviders|AnimSequence")
+	int32 GetNumLayers() const
+	{
+		return GetLayers().Num();
+	}
+
 #if WITH_EDITOR
 	// Creates a transform provider data asset from an AnimBankData asset.
 	UFUNCTION(BlueprintCallable, Category = "TransformProviders|AnimSequence", meta = (DevelopmentOnly))
@@ -502,17 +927,39 @@ protected:
 	bool IsValidFor(FInstancedSkinningSceneExtensionProxy* ExtensionProxy) const;
 
 private:
-	virtual void SetParentProvider(UAnimSequenceTransformProviderData* Provider);
+#if WITH_EDITOR
+	static void ForEachReferencingProvider(
+		TFunctionRef<bool(const UAnimSequenceTransformProviderData*)> Predicate,
+		TFunctionRef<void(UAnimSequenceTransformProviderData*)> Action);
+#endif
 
 	UPROPERTY(EditAnywhere, BlueprintReadWrite, Category = TransformProvider, meta = (AllowPrivateAccess = "true", ShowOnlyInnerProperties))
 	TObjectPtr<USkinnedAsset> SkinnedAsset;
 
-	UPROPERTY(EditAnywhere, BlueprintReadWrite, Category = TransformProvider, meta = (AllowPrivateAccess = "true", ShowOnlyInnerProperties))
-	TObjectPtr<UAnimSequenceTransformProviderData> ParentProvider;
+#if WITH_EDITORONLY_DATA
+	UPROPERTY(EditAnywhere, BlueprintReadOnly, Category = TransformProvider, meta = (AllowPrivateAccess = "true"))
+	TObjectPtr<UAnimSequenceTransformProviderSequenceList> SequenceList;
+#endif
 
-	UPROPERTY(EditAnywhere, BlueprintReadOnly, Category = TransformProvider, meta = (AllowPrivateAccess = "true", ShowOnlyInnerProperties, EditCondition = "ParentProvider == nullptr"))
+	UPROPERTY(EditAnywhere, BlueprintReadOnly, Category = TransformProvider, meta = (AllowPrivateAccess = "true", EditCondition = "SequenceList == nullptr", EditConditionHides))
 	TArray<struct FAnimSequenceTransformProviderSequence> Sequences;
 
+#if WITH_EDITORONLY_DATA
+	UPROPERTY(EditAnywhere, BlueprintReadOnly, Category = TransformProvider, meta = (AllowPrivateAccess = "true"))
+	TObjectPtr<UAnimSequenceTransformProviderLayerStack> LayerStack;
+#endif
+
+	UPROPERTY(EditAnywhere, BlueprintReadOnly, Category = TransformProvider, meta = (AllowPrivateAccess = "true", EditCondition = "LayerStack == nullptr", EditConditionHides))
+	TArray<FAnimSequenceTransformProviderLayer> Layers;
+
+#if WITH_EDITORONLY_DATA
+	UPROPERTY(EditAnywhere, BlueprintReadOnly, Category = TransformProvider, meta = (AllowPrivateAccess = "true"))
+	TObjectPtr<UAnimSequenceTransformProviderBlendSpaceList> BlendSpaceList;
+#endif
+
+	UPROPERTY(EditAnywhere, BlueprintReadOnly, Category = TransformProvider, meta = (AllowPrivateAccess = "true", EditCondition = "BlendSpaceList == nullptr", EditConditionHides))
+	TArray<FAnimSequenceTransformProviderBlendSpace> BlendSpaces;
+
 	double GetWorldTime() const;
 
 	FRenderCommandFence DestroyFence;
@@ -535,6 +982,9 @@ private:
 
 	friend class UAnimSequenceTransformProviderDataInstance;
 	friend class FAnimSequenceTransformProviderCompilingManager;
+	friend class UAnimSequenceTransformProviderLayerStack;
+	friend class UAnimSequenceTransformProviderSequenceList;
+	friend class UAnimSequenceTransformProviderBlendSpaceList;
 };
 
 UCLASS(config = Engine, hidecategories = Object, MinimalAPI, BlueprintType, Within = InstancedSkinnedMeshComponent, NotBlueprintable, NotPlaceable)
@@ -549,77 +999,100 @@ public:
 	virtual bool HasAnimationBounds() const;
 	virtual bool GetAnimationBounds(uint32 AnimationIndex, FRenderBounds& OutBounds) const override;
 	virtual bool IsCompiling() const override;
+	virtual UTransformProviderData* GetRootTransformProvider() override;
 	virtual FTransformProviderRenderProxy* CreateRenderProxy(FInstancedSkinningSceneExtensionProxy* ExtensionProxy) const override;
 	virtual void BeginDestroy() override;
 	virtual void Serialize(FArchive& Ar) override;
 	virtual void SubmitChanges() override;
-
 public:
-	// Copies items from this anim bank to create a new anim bank instance.
+	// Creates an instance from a provider data asset.
 	UFUNCTION(BlueprintCallable, Category = "TransformProviders|AnimSequence")
 	static ENGINE_API UAnimSequenceTransformProviderDataInstance* CreateAnimSequenceTransformProviderDataInstance(UAnimSequenceTransformProviderData* ProviderData, UInstancedSkinnedMeshComponent* Owner);
 
-	// Allocates a new track with auto play data, or -1 if the track allocation failed.
+	// Sets the per-track weight for a layer. Multiplied with the asset's layer weight.
+	UFUNCTION(BlueprintCallable, Category = "TransformProviders|AnimSequence")
+	ENGINE_API void SetLayerWeight(int32 TrackIndex, int32 LayerIndex, float Weight);
+
+	// Returns the per-track weight for a layer, or 1.0 if invalid.
 	UFUNCTION(BlueprintCallable, Category = "TransformProviders|AnimSequence")
-	ENGINE_API int32 AllocateAutoPlayTrack(const FAnimSequenceTrackAutoPlayData& Data);
+	ENGINE_API float GetLayerWeight(int32 TrackIndex, int32 LayerIndex) const;
 
-	// Allocates a new track with manual play data, or -1 if the track allocation failed.
+	// Allocates a track slot across all layer pools with no animation data.
 	UFUNCTION(BlueprintCallable, Category = "TransformProviders|AnimSequence")
-	ENGINE_API int32 AllocateManualTrack(const FAnimSequenceTrackManualData& Data);
+	ENGINE_API int32 AllocateTrack();
 
-	// Updates a track with auto play data. Returns whether the update succeeded.
+	// Deallocates a track across all layer pools.
 	UFUNCTION(BlueprintCallable, Category = "TransformProviders|AnimSequence")
-	ENGINE_API bool UpdateAutoPlayTrack(int32 TrackIndex, const FAnimSequenceTrackAutoPlayData& Data);
+	ENGINE_API bool DeallocateTrack(int32 TrackIndex);
 
-	// Updates a track with manual data. Returns whether the update succeeded.
+	// Sets auto play animation data on a track layer.
 	UFUNCTION(BlueprintCallable, Category = "TransformProviders|AnimSequence")
-	ENGINE_API bool UpdateManualTrack(int32 TrackIndex, const FAnimSequenceTrackManualData& Data);
+	ENGINE_API bool SetAutoPlayData(int32 TrackIndex, int32 LayerIndex, const FAnimSequenceTrackAutoPlayData& Data);
 
-	// Deallocates a track. Returns whether the deallocation succeeded.
+	// Sets manual animation data on a track layer.
 	UFUNCTION(BlueprintCallable, Category = "TransformProviders|AnimSequence")
-	ENGINE_API bool DeallocateTrack(int32 TrackIndex);
+	ENGINE_API bool SetManualData(int32 TrackIndex, int32 LayerIndex, const FAnimSequenceTrackManualData& Data);
 
-	// Get the sequence index currently playing on a track
+	// Sets the blend space data for a track layer. The blend space index must reference a valid entry in the BlendSpaces array.
 	UFUNCTION(BlueprintCallable, Category = "TransformProviders|AnimSequence")
-	ENGINE_API int32 GetSequenceIndex(int32 TrackIndex) const;
+	ENGINE_API bool SetBlendSpaceData(int32 TrackIndex, int32 LayerIndex, const FAnimSequenceTrackBlendSpaceData& Data);
 
-	// Assigns a new play rate to an existing track.
+	// Sets the play rate on a track layer. Works for auto-play and blend space modes.
 	UFUNCTION(BlueprintCallable, Category = "TransformProviders|AnimSequence")
-	ENGINE_API bool SetPlayRate(int32 TrackIndex, float PlayRate);
+	ENGINE_API bool SetPlayRate(int32 TrackIndex, int32 LayerIndex, float PlayRate);
 
-	// Get current play rate of a track.
+	// Returns the play rate for a track layer, or 0.0 for manual mode, or 1.0 if invalid.
 	UFUNCTION(BlueprintCallable, Category = "TransformProviders|AnimSequence")
-	ENGINE_API float GetPlayRate(int32 TrackIndex) const;
+	ENGINE_API float GetPlayRate(int32 TrackIndex, int32 LayerIndex) const;
 
-	// Assigns a new loop mode to an existing track.
+	// Sets the loop mode on a track layer.
 	UFUNCTION(BlueprintCallable, Category = "TransformProviders|AnimSequence")
-	ENGINE_API bool SetLoopMode(int32 TrackIndex, EAnimSequenceTrackLoopMode LoopMode);
+	ENGINE_API bool SetLoopMode(int32 TrackIndex, int32 LayerIndex, EAnimSequenceTrackLoopMode LoopMode);
 
-	// Get current loop mode of a track.
+	// Returns the loop mode for a track layer.
 	UFUNCTION(BlueprintCallable, Category = "TransformProviders|AnimSequence")
-	ENGINE_API EAnimSequenceTrackLoopMode GetLoopMode(int32 TrackIndex) const;
+	ENGINE_API EAnimSequenceTrackLoopMode GetLoopMode(int32 TrackIndex, int32 LayerIndex) const;
 
-	// Get current position of a track.
+	// Returns the current wrapped playback position for a track layer. Returns 0 for blend space mode.
 	UFUNCTION(BlueprintCallable, Category = "TransformProviders|AnimSequence")
-	ENGINE_API float GetPosition(int32 TrackIndex) const;
+	ENGINE_API float GetPosition(int32 TrackIndex, int32 LayerIndex) const;
 
-private:
-	void MarkOwnerRenderStateDirty();
-	void SetParentProvider(UAnimSequenceTransformProviderData* Provider) override;
+	// Returns the 2D blend space position for a track layer, or zero if not in blend space mode.
+	UFUNCTION(BlueprintCallable, Category = "TransformProviders|AnimSequence")
+	ENGINE_API FVector2f GetBlendSpacePosition(int32 TrackIndex, int32 LayerIndex) const;
 
-	// Cached owner component for efficient dirty marking.
-	UInstancedSkinnedMeshComponent* CachedOwnerComponent = nullptr;
+	// Returns the sequence index for a track layer. Returns INDEX_NONE for blend space mode.
+	UFUNCTION(BlueprintCallable, Category = "TransformProviders|AnimSequence")
+	ENGINE_API int32 GetSequenceIndex(int32 TrackIndex, int32 LayerIndex) const;
 
-	// Prevents redundant MarkRenderDynamicDataDirty calls within a frame.
-	bool bPendingRenderStateDirty = false;
+private:
+	void MarkOwnerRenderStateDirty();
+	void InitLayerPools(bool bMarkDirty = true);
+	void SetParentProvider(UAnimSequenceTransformProviderData* ProviderData);
+	bool IsValidTrackAndLayer(int32 TrackIndex, int32 LayerIndex) const;
 
-	// Defers initialization of timestamps after serialization until the world is available.
-	bool bNeedsPostLoadInit = false;
+	UPROPERTY()
+	TObjectPtr<UAnimSequenceTransformProviderData> ParentProvider;
 
-	FAnimSequenceTrackPool Tracks;
+	UInstancedSkinnedMeshComponent* CachedOwnerComponent = nullptr;
+	TArray<FAnimSequenceTrackPool> TrackPools;
 	TUniquePtr<FAnimSequenceTransformProviderRenderData> RenderData{ new FAnimSequenceTransformProviderRenderData };
 	FRenderCommandFence InstanceDestroyFence;
 	UE::FMutex SubmitChangesMutex;
 
+	// Per-blend-space resolution data (one per BlendSpaces entry)
+	struct FBlendSpaceCache
+	{
+		TArray<uint16> SampleToSequenceIndex;
+		bool bValid = false;
+	};
+	TArray<FBlendSpaceCache> BlendSpaceCaches;
+	void BuildBlendSpaceCaches();
+	void ResolveBlendSpaceSamples(FAnimSequenceTrackPackedData& Data) const;
+	bool bTrackPoolsDirty = false;
+	bool bPendingRenderStateDirty = false;
+	bool bNeedsPostLoadInit = false;
+
 	friend class UAnimSequenceTransformProviderData;
+	friend class UAnimSequenceTransformProviderLayerStack;
 };
\ No newline at end of file
diff --git a/Engine/Source/Runtime/Engine/Classes/Animation/AnimTrackPool.h b/Engine/Source/Runtime/Engine/Classes/Animation/AnimTrackPool.h
--- a/Engine/Source/Runtime/Engine/Classes/Animation/AnimTrackPool.h
+++ b/Engine/Source/Runtime/Engine/Classes/Animation/AnimTrackPool.h
@@ -70,8 +70,8 @@ public:
 			uint32 Index;
 		};
 
-		int32 NumTracks;
-		int32 NumActiveTracks;
+		int32 NumTracks = 0;
+		int32 NumActiveTracks = 0;
 		TArray<FActivateOp, FConcurrentLinearArrayAllocator> ActivateOps;
 		TArray<FDeactivateOp, FConcurrentLinearArrayAllocator> DeactivateOps;
 	};
@@ -83,6 +83,22 @@ public:
 		DirtyBits.Reserve(Num);
 	}
 
+	// Initialize all active tracks from source but custom initialize each track.
+	template <typename InitFunctionType>
+	void InitFrom(const TAnimTrackPool& Source, InitFunctionType InitFunction)
+	{
+		Datas.SetNum(Source.Datas.Num());
+		ActiveBits = Source.ActiveBits;
+		NumActiveTracks = Source.NumActiveTracks;
+		DirtyBits = Source.ActiveBits;
+		NumDirtyTracks = Source.NumActiveTracks;
+
+		for (TConstSetBitIterator<> BitIt(ActiveBits); BitIt; ++BitIt)
+		{
+			Datas[BitIt.GetIndex()] = InitFunction(BitIt.GetIndex());
+		}
+	}
+
 	int32 AllocateTrack(const DataType& Data)
 	{
 		const int32 Index = ActiveBits.FindAndSetFirstZeroBit();
@@ -103,12 +119,40 @@ public:
 		return Datas.Num() - 1;
 	}
 
+	int32 AllocateTrackAt(int32 Index, const DataType& Data)
+	{
+		if (Index < 0)
+		{
+			return INDEX_NONE;
+		}
+
+		if (Index >= Datas.Num())
+		{
+			const int32 OldNum = Datas.Num();
+			const int32 NewNum = Index + 1;
+			Datas.SetNum(NewNum);
+			ActiveBits.Add(false, NewNum - OldNum);
+			DirtyBits.Add(false, NewNum - OldNum);
+		}
+
+		if (ActiveBits[Index])
+		{
+			return INDEX_NONE;
+		}
+
+		ActiveBits[Index] = true;
+		NumActiveTracks++;
+		SetDirty(Index);
+		Datas[Index] = Data;
+		return Index;
+	}
+
 	bool UpdateTrack(int32 Index, const DataType& Data)
 	{
 		if (IsActiveIndex(Index))
 		{
 			SetDirty(Index);
-			Datas[Index] = Data;
+			Datas[Index].SetData(Data);
 			return true;
 		}
 		return false;
diff --git a/Engine/Source/Runtime/Engine/Private/Animation/AnimBank.cpp b/Engine/Source/Runtime/Engine/Private/Animation/AnimBank.cpp
--- a/Engine/Source/Runtime/Engine/Private/Animation/AnimBank.cpp
+++ b/Engine/Source/Runtime/Engine/Private/Animation/AnimBank.cpp
@@ -1421,7 +1421,7 @@ bool UAnimBankData::IsEnabled() const
 	#endif
 	}
 
-	return bEnabled;
+	return true;
 }
 
 const FGuid& UAnimBankData::GetTransformProviderID() const
diff --git a/Engine/Source/Runtime/Engine/Private/Animation/AnimRuntimeTransformProviderData.cpp b/Engine/Source/Runtime/Engine/Private/Animation/AnimRuntimeTransformProviderData.cpp
--- a/Engine/Source/Runtime/Engine/Private/Animation/AnimRuntimeTransformProviderData.cpp
+++ b/Engine/Source/Runtime/Engine/Private/Animation/AnimRuntimeTransformProviderData.cpp
@@ -221,7 +221,6 @@ UAnimRuntimeTransformProviderData* UAnimRuntimeTransformProviderData::CreateAnim
 		{
 			UAnimRuntimeTransformProviderData* Instance = NewObject<UAnimRuntimeTransformProviderData>(Owner);
 			Instance->NumTransforms = SkinnedAsset->GetRefBasesInvMatrix().Num();
-			Instance->bEnabled = true;
 			return Instance;
 		}
 	}
diff --git a/Engine/Source/Runtime/Engine/Private/Animation/AnimSequenceTransformProviderCompiler.cpp b/Engine/Source/Runtime/Engine/Private/Animation/AnimSequenceTransformProviderCompiler.cpp
--- a/Engine/Source/Runtime/Engine/Private/Animation/AnimSequenceTransformProviderCompiler.cpp
+++ b/Engine/Source/Runtime/Engine/Private/Animation/AnimSequenceTransformProviderCompiler.cpp
@@ -859,6 +859,13 @@ bool FAnimSequenceTransformProviderBuildAsyncCacheTask::BuildData(const UE::FSha
 		return false;
 	}
 
+	// Compute conservative bounds as the union of all sequence bounds.
+	Data->ConservativeBounds = AssetBounds;
+	for (const FBoxSphereBounds& Bounds : Data->SequenceBounds)
+	{
+		Data->ConservativeBounds = Data->ConservativeBounds + Bounds;
+	}
+
 	return true;
 }
 
diff --git a/Engine/Source/Runtime/Engine/Private/Animation/AnimSequenceTransformProviderData.cpp b/Engine/Source/Runtime/Engine/Private/Animation/AnimSequenceTransformProviderData.cpp
--- a/Engine/Source/Runtime/Engine/Private/Animation/AnimSequenceTransformProviderData.cpp
+++ b/Engine/Source/Runtime/Engine/Private/Animation/AnimSequenceTransformProviderData.cpp
@@ -7,6 +7,7 @@
 #include "Animation/SkeletonRemappingRegistry.h"
 #include "Animation/AnimSequence.h"
 #include "Animation/Skeleton.h"
+#include "Animation/BlendSpace.h"
 #include "SceneInterface.h"
 #include "UObject/UObjectThreadContext.h"
 #include "UObject/FortniteMainBranchObjectVersion.h"
@@ -54,33 +55,59 @@ void FAnimSequenceTrackPackedData::Serialize(FArchive& Ar)
 	constexpr uint32 AutoPlayBits = 1;
 	constexpr uint32 AutoPlayMask = MakeBitMask(AutoPlayBits);
 
-	constexpr uint32 HasPreviousPositionShift = AutoPlayShift + AutoPlayBits;
-	constexpr uint32 HasPreviousPositionBits = 1;
-	constexpr uint32 HasPreviousPositionMask = MakeBitMask(HasPreviousPositionBits);
+	constexpr uint32 BlendSpaceShift = AutoPlayShift + AutoPlayBits;
+	constexpr uint32 BlendSpaceBits = 1;
+	constexpr uint32 BlendSpaceMask = MakeBitMask(BlendSpaceBits);
 
 	uint32 PackedBits;
-	if (Ar.IsSaving())
+
+	if (Ar.CustomVer(FFortniteMainBranchObjectVersion::GUID) >= FFortniteMainBranchObjectVersion::AnimSequenceTransformProviderLayers)
 	{
-		const uint8 bManual = !bAutoPlay;
+		if (Ar.IsSaving())
+		{
+			PackedBits =
+				(SequenceIndex & SequenceIndexMask)
+				| ((LoopMode & LoopModeMask) << LoopModeShift)
+				| ((bAutoPlay & AutoPlayMask) << AutoPlayShift)
+				| ((bBlendSpace & BlendSpaceMask) << BlendSpaceShift);
+		}
 
-		PackedBits =
-			(SequenceIndex & SequenceIndexMask)
-			| ((LoopMode & LoopModeMask) << LoopModeShift)
-			| ((bAutoPlay & AutoPlayMask) << AutoPlayShift)
-			| ((bManual & bHasPreviousPosition & HasPreviousPositionMask) << HasPreviousPositionShift);
 		Ar << PackedBits;
 
-		if (bAutoPlay)
+		if (Ar.IsLoading())
+		{
+			SequenceIndex = PackedBits & SequenceIndexMask;
+			LoopMode = (PackedBits >> LoopModeShift) & LoopModeMask;
+			bAutoPlay = (PackedBits >> AutoPlayShift) & AutoPlayMask;
+			bBlendSpace = (PackedBits >> BlendSpaceShift) & BlendSpaceMask;
+			PatchFlags = static_cast<uint32>(EAnimSequenceTrackPatchFlags::All);
+		}
+
+		if (Ar.IsLoading() && bBlendSpace)
+		{
+			FMemory::Memzero(&BlendSpace, sizeof(BlendSpace));
+		}
+
+		if (bBlendSpace)
+		{
+			Ar << BlendSpace.Position;
+			Ar << BlendSpace.PlayRate.Encoded;
+			Ar << BlendSpace.BlendPosition[0].Encoded;
+			Ar << BlendSpace.BlendPosition[1].Encoded;
+			Ar << BlendSpace.BlendSpaceIndex;
+		}
+		else if (bAutoPlay)
 		{
 			Ar << Auto.Position;
-			Ar << Auto.PlayRate;
-			Ar << Auto.BlendTime;
+			Ar << Auto.PlayRate.Encoded;
 		}
 		else
 		{
 			Ar << Manual.Position;
-			Ar << Manual.PreviousPosition;
 		}
+
+		Ar << BlendTime.Encoded;
+		Ar << LayerWeight.Encoded;
 	}
 	else
 	{
@@ -88,7 +115,7 @@ void FAnimSequenceTrackPackedData::Serialize(FArchive& Ar)
 		SequenceIndex = PackedBits & SequenceIndexMask;
 		LoopMode = (PackedBits >> LoopModeShift) & LoopModeMask;
 		bAutoPlay = (PackedBits >> AutoPlayShift) & AutoPlayMask;
-		bHasPreviousPosition = (PackedBits >> HasPreviousPositionShift) & HasPreviousPositionMask;
+		bBlendSpace = 0;
 		PatchFlags = static_cast<uint32>(EAnimSequenceTrackPatchFlags::All);
 
 		if (Ar.CustomVer(FFortniteMainBranchObjectVersion::GUID) >= FFortniteMainBranchObjectVersion::AnimSequenceTrackTimestamps)
@@ -96,14 +123,18 @@ void FAnimSequenceTrackPackedData::Serialize(FArchive& Ar)
 			if (bAutoPlay)
 			{
 				Ar << Auto.Position;
-				Ar << Auto.PlayRate;
-				Ar << Auto.BlendTime;
-				Auto.PreviousPosition = Auto.Position;
+				float LegacyPlayRate;
+				Ar << LegacyPlayRate;
+				Auto.PlayRate = FFloat16(LegacyPlayRate);
+				float LegacyBlendTime;
+				Ar << LegacyBlendTime;
+				BlendTime = FFloat16(LegacyBlendTime);
 			}
 			else
 			{
 				Ar << Manual.Position;
-				Ar << Manual.PreviousPosition;
+				float LegacyPreviousPosition;
+				Ar << LegacyPreviousPosition; // Discard
 			}
 		}
 		else
@@ -114,14 +145,18 @@ void FAnimSequenceTrackPackedData::Serialize(FArchive& Ar)
 			if (bAutoPlay)
 			{
 				Auto.Position = Position;
-				Auto.PreviousPosition = Auto.Position;
-				Ar << Auto.PlayRate;
-				Ar << Auto.BlendTime;
+				float LegacyPlayRate;
+				Ar << LegacyPlayRate;
+				Auto.PlayRate = FFloat16(LegacyPlayRate);
+				float LegacyBlendTime;
+				Ar << LegacyBlendTime;
+				BlendTime = FFloat16(LegacyBlendTime);
 			}
 			else
 			{
 				Manual.Position = Position;
-				Ar << Manual.PreviousPosition;
+				float LegacyPreviousPosition;
+				Ar << LegacyPreviousPosition; // Discard
 			}
 		}
 	}
@@ -131,6 +166,11 @@ void FAnimSequenceTrackPackedData::Serialize(FArchive& Ar)
 
 void FAnimSequenceTrackRenderData::Patch(const FAnimSequenceTrackPackedData& Data, double WorldTime, bool bWasActive)
 {
+	if (!bWasActive)
+	{
+		bHasPreviousPosition = false;
+	}
+
 	// Only patching specific flags.
 	if (Data.GetPatchFlags() != EAnimSequenceTrackPatchFlags::None && Data.GetPatchFlags() != EAnimSequenceTrackPatchFlags::All)
 	{
@@ -146,43 +186,76 @@ void FAnimSequenceTrackRenderData::Patch(const FAnimSequenceTrackPackedData& Dat
 			Target.SetLoopMode(Data.GetLoopMode());
 		}
 	}
-	// Patching the entire struct. Starts a new animation.
+	// Full struct replacement. Check if it's the same animation (position update) or a new one (transition).
 	else
 	{
-		if (Data.GetBlendTime() <= 0.0f || !bWasActive || !CVarAnimSequenceTransformProviderBlending.GetValueOnRenderThread())
+		if (Target.IsSameAnimation(Data))
 		{
-			Source = Data;
+			if (Data.IsBlendSpace())
+			{
+				Target.UpdateBlendSpaceData(Data, WorldTime);
+				if (BlendWeight.GetFloat() <= 0.0f)
+				{
+					Source.UpdateBlendSpaceData(Data, WorldTime);
+				}
+			}
+			else
+			{
+				Target.SetData(Data);
+				if (BlendWeight.GetFloat() <= 0.0f)
+				{
+					Source.SetData(Data);
+				}
+			}
 		}
-		else if (BlendWeight > 0.0f)
+		else
 		{
-			Source = Target;
-		}
+			if (Data.GetBlendTime() <= 0.0f || !bWasActive || !CVarAnimSequenceTransformProviderBlending.GetValueOnRenderThread())
+			{
+				Source = Data;
+			}
+			else if (BlendWeight.GetFloat() > 0.0f)
+			{
+				Source = Target;
+			}
 
-		BlendWeight = 0.0f;
-		Target = Data;
+			BlendWeight = FFloat16(0.0f);
+			Target = Data;
+		}
 	}
 
 	Source.ResetPatchFlags();
 	Target.ResetPatchFlags();
 }
 
-void FAnimSequenceTrackRenderData::Tick(float DeltaTime, float SourcePosition, float TargetPosition)
+void FAnimSequenceTrackRenderData::Tick(float DeltaTime, double WorldTime)
 {
-	const float BlendTime = Target.GetBlendTime();
+	if (Target.IsBlendSpace())
+	{
+		const FVector2f BlendPos = Target.GetBlendSpacePosition();
+		PreviousBlendSpacePosition[0] = FFloat16(BlendPos.X);
+		PreviousBlendSpacePosition[1] = FFloat16(BlendPos.Y);
+	}
+	PreviousTimePosition = Target.GetTimePosition(WorldTime);
+	PreviousBlendWeight = BlendWeight;
+	bHasPreviousPosition = true;
 
-	if (BlendTime > 0.0f && BlendWeight < 1.0f && CVarAnimSequenceTransformProviderBlending.GetValueOnRenderThread())
+	const float BlendTimeValue = Target.GetBlendTime();
+
+	if (BlendTimeValue > 0.0f && BlendWeight.GetFloat() < 1.0f && CVarAnimSequenceTransformProviderBlending.GetValueOnRenderThread())
 	{
-		BlendWeight += DeltaTime / BlendTime;
+		const float NewBlendWeight = BlendWeight.GetFloat() + DeltaTime / BlendTimeValue;
 
-		if (BlendWeight >= 1.0f)
+		if (NewBlendWeight >= 1.0f)
 		{
-			BlendWeight = 0.0f;
+			BlendWeight = FFloat16(0.0f);
 			Source = Target;
 		}
+		else
+		{
+			BlendWeight = FFloat16(NewBlendWeight);
+		}
 	}
-
-	Source.Tick(SourcePosition);
-	Target.Tick(TargetPosition);
 }
 
 ///////////////////////////////////////////////////////////////////////////////////////////////////
@@ -202,31 +275,59 @@ void FAnimSequenceTransformProviderRenderData::RemoveProxy(FInstancedSkinningSce
 	Proxies.Remove(Proxy);
 }
 
-void FAnimSequenceTransformProviderRenderData::Patch(FAnimSequenceTrackPool::FPatch&& Patch, double WorldTime)
+void FAnimSequenceTransformProviderRenderData::Patch(TArray<FAnimSequenceTrackPool::FPatch>&& InPatches, TArray<FAnimSequenceTransformProviderRenderLayer>&& InLayers, double WorldTime)
 {
-	if (Patch.GetNumTracks() != Tracks.GetNumTracks())
+	check(!InLayers.IsEmpty());
+	check(InPatches.Num() == InLayers.Num());
+	check(InLayers.Num() <= MAX_ANIM_SEQUENCE_LAYERS);
+	Layers.SetNum(InLayers.Num());
+	MeshSpaceMask = 0;
+	LayerModeMask = 0;
+
+	for (int32 LayerIndex = 0; LayerIndex < Layers.Num(); LayerIndex++)
+	{
+		Layers[LayerIndex].Layer  = InLayers[LayerIndex].Layer;
+		Layers[LayerIndex].Weight = InLayers[LayerIndex].Weight;
+
+		if (Layers[LayerIndex].Layer.bMeshSpaceRotation)
+		{
+			MeshSpaceMask |= (1u << LayerIndex);
+		}
+
+		const uint32 Mode = Layers[LayerIndex].Layer.BlendMode == EAnimSequenceTransformProviderLayerBlendMode::Additive ? SKINNING_LAYER_MODE_ADDITIVE : SKINNING_LAYER_MODE_OVERRIDE;
+		LayerModeMask |= (Mode << (LayerIndex * SKINNING_LAYER_MODE_BITS_PER_LAYER));
+	}
+
+	if (InPatches[0].GetNumTracks() != Layers[0].Tracks.GetNumTracks())
 	{
 		for (FInstancedSkinningSceneExtensionProxy* Proxy : Proxies)
 		{
-			Proxy->SetUniqueAnimationCount(Patch.GetNumTracks());
+			Proxy->SetUniqueAnimationCount(InPatches[0].GetNumTracks());
 		}
 	}
 
-	Tracks.Patch<FAnimSequenceTrackPackedData>(MoveTemp(Patch), [WorldTime](FAnimSequenceTrackRenderData& Dst, const FAnimSequenceTrackPackedData& Src, bool bWasActive)
+	for (int32 LayerIndex = 0; LayerIndex < InPatches.Num(); LayerIndex++)
 	{
-		Dst.Patch(Src, WorldTime, bWasActive);
-	});
+		if (InPatches[LayerIndex].GetNumTracks() > 0)
+		{
+			Layers[LayerIndex].Tracks.Patch<FAnimSequenceTrackPackedData>(MoveTemp(InPatches[LayerIndex]), [WorldTime](FAnimSequenceTrackRenderData& Dst, const FAnimSequenceTrackPackedData& Src, bool bWasActive)
+			{
+				Dst.Patch(Src, WorldTime, bWasActive);
+			});
+		}
+	}
 }
 
 ///////////////////////////////////////////////////////////////////////////////////////////////////
 
 FAnimSequenceTransformProviderProxy::FAnimSequenceTransformProviderProxy(TConstArrayView<FAnimSequenceTransformProviderSequence> InSequences, FInstancedSkinningSceneExtensionProxy* InSceneExtensionProxy)
 	: Sequences(InSequences)
-	, TrackAllocator(new FAnimSequenceTrackRenderPool)
-	, Tracks(TrackAllocator.Get())
+	, RenderData(new FAnimSequenceTransformProviderRenderData())
 	, SceneExtensionProxy(InSceneExtensionProxy)
+	, bOwnsRenderData(true)
 {
-	TrackAllocator->Reserve(Sequences.Num());
+	FAnimSequenceTransformProviderRenderLayer& BaseLayer = RenderData->Layers.Emplace_GetRef();
+	BaseLayer.Tracks.Reserve(Sequences.Num());
 
 	for (int32 SequenceIndex = 0; SequenceIndex < Sequences.Num(); ++SequenceIndex)
 	{
@@ -237,18 +338,25 @@ FAnimSequenceTransformProviderProxy::FAnimSequenceTransformProviderProxy(TConstA
 		Data.Position = Sequence.Position;
 		Data.PlayRate = Sequence.PlayRate;
 
-		TrackAllocator->AllocateTrack(FAnimSequenceTrackPackedData::Pack(Data, 0.0));
+		BaseLayer.Tracks.AllocateTrack(FAnimSequenceTrackPackedData::Pack(Data, 0.0));
 	}
 }
 
 FAnimSequenceTransformProviderProxy::FAnimSequenceTransformProviderProxy(TConstArrayView<FAnimSequenceTransformProviderSequence> InSequences, FInstancedSkinningSceneExtensionProxy* InSceneExtensionProxy, FAnimSequenceTransformProviderRenderData& InRenderData)
 	: Sequences(InSequences)
 	, RenderData(&InRenderData)
-	, Tracks(&InRenderData.Tracks)
 	, SceneExtensionProxy(InSceneExtensionProxy)
 {
 }
 
+FAnimSequenceTransformProviderProxy::~FAnimSequenceTransformProviderProxy()
+{
+	if (bOwnsRenderData)
+	{
+		delete RenderData;
+	}
+}
+
 void FAnimSequenceTransformProviderProxy::CreateRenderThreadResources(FSceneInterface& Scene, FRHICommandListBase& RHICmdList)
 {
 	if (RenderData)
@@ -295,32 +403,13 @@ bool UAnimSequenceTransformProviderData::IsEnabled() const
 		}
 	}
 
-	// Any disabled parent in the chain disables the child.
-	for (const UAnimSequenceTransformProviderData* Current = this; Current; Current = Current->ParentProvider)
-	{
-		if (!Current->bEnabled)
-		{
-			return false;
-		}
-	}
-
 	return true;
 }
 
 const FGuid& UAnimSequenceTransformProviderData::GetTransformProviderID() const
 {
-	static FGuid AnimSequenceGPUProviderId(ANIM_SEQUENCE_GPU_TRANSFORM_PROVIDER_GUID);
-	static FGuid AnimSequenceCPUProviderId(ANIM_SEQUENCE_CPU_TRANSFORM_PROVIDER_GUID);
-
-	static const auto AnimSequenceGPUVar = IConsoleManager::Get().FindTConsoleVariableDataBool(TEXT("r.AnimSequenceTransformProvider.GPU"));
-	if (AnimSequenceGPUVar && AnimSequenceGPUVar->GetValueOnAnyThread() == 1)
-	{
-		return AnimSequenceGPUProviderId;
-	}
-	else
-	{
-		return AnimSequenceCPUProviderId;
-	}
+	static FGuid AnimSequenceProviderId(ANIM_SEQUENCE_GPU_TRANSFORM_PROVIDER_GUID);
+	return AnimSequenceProviderId;
 }
 
 uint32 UAnimSequenceTransformProviderData::GetUniqueAnimationCount() const
@@ -335,14 +424,14 @@ bool UAnimSequenceTransformProviderData::UsesSkeletonBatching() const
 
 bool UAnimSequenceTransformProviderData::HasAnimationBounds() const
 {
-	return CachedData.SequenceBounds.Num() > 0;
+	return CachedData.ConservativeBounds.SphereRadius > 0.0f;
 }
 
 bool UAnimSequenceTransformProviderData::GetAnimationBounds(uint32 AnimationIndex, FRenderBounds& OutBounds) const
 {
-	if (AnimationIndex < uint32(CachedData.SequenceBounds.Num()))
+	if (CachedData.ConservativeBounds.SphereRadius > 0.0f)
 	{
-		OutBounds = CachedData.SequenceBounds[AnimationIndex];
+		OutBounds = CachedData.ConservativeBounds;
 		return true;
 	}
 	return false;
@@ -406,6 +495,28 @@ bool UAnimSequenceTransformProviderData::IsValidFor(FInstancedSkinningSceneExten
 		return false;
 	}
 
+	const FSkeletalMeshRenderData* RenderData = SkinnedAsset->GetResourceForRendering();
+	if (!RenderData || !RenderData->HasUnifiedBoneMap())
+	{
+		UE_LOG(LogAnimation, Warning, TEXT("AnimSequenceTransformProvider: SkinnedAsset '%s' in provider '%s' requires 'Optimize for Instancing' enabled on the skeletal mesh in the build settings."),
+			*SkinnedAsset->GetPathName(), *GetPathName());
+		return false;
+	}
+
+	const bool bHasMeshSpaceLayer = Layers.ContainsByPredicate([](const FAnimSequenceTransformProviderLayer& Layer) { return Layer.bMeshSpaceRotation; });
+	const int32 MaxBonesPerGroup = bHasMeshSpaceLayer ? SKINNING_MAX_BONES_PER_GROUP_MESHSPACE : SKINNING_MAX_BONES_PER_GROUP;
+	const int32 MinLODLevel = ExtensionProxy->GetMeshObject() ? ExtensionProxy->GetMeshObject()->MinLODLevel : 0;
+	for (int32 LODIndex = MinLODLevel; LODIndex < RenderData->LODRenderData.Num(); ++LODIndex)
+	{
+		if (RenderData->LODRenderData[LODIndex].RequiredBones.Num() > MaxBonesPerGroup)
+		{
+			UE_LOG(LogAnimation, Warning, TEXT("AnimSequenceTransformProvider: SkinnedAsset '%s' in provider '%s' LOD %d has %d bones which exceeds the maximum of %d%s."),
+				*SkinnedAsset->GetPathName(), *GetPathName(), LODIndex, RenderData->LODRenderData[LODIndex].RequiredBones.Num(), MaxBonesPerGroup,
+				bHasMeshSpaceLayer ? TEXT(" (reduced due to mesh-space rotation layer)") : TEXT(""));
+			return false;
+		}
+	}
+
 	for (const FAnimSequenceTransformProviderSequence& Sequence : Sequences)
 	{
 		if (Sequence.Sequence == nullptr)
@@ -451,6 +562,24 @@ bool UAnimSequenceTransformProviderData::IsValidFor(FInstancedSkinningSceneExten
 		}
 	}
 
+	for (int32 BlendSpaceIndex = 0; BlendSpaceIndex < BlendSpaces.Num(); BlendSpaceIndex++)
+	{
+		const UBlendSpace* BlendSpace = BlendSpaces[BlendSpaceIndex].BlendSpace;
+		if (!BlendSpace)
+		{
+			UE_LOG(LogAnimation, Warning, TEXT("AnimSequenceTransformProvider: Null blend space at index %d in provider '%s'"),
+				BlendSpaceIndex, *GetPathName());
+			return false;
+		}
+
+		if (BlendSpace->GetBlendSamples().IsEmpty())
+		{
+			UE_LOG(LogAnimation, Warning, TEXT("AnimSequenceTransformProvider: Blend space '%s' at index %d in provider '%s' has no samples."),
+				*BlendSpace->GetPathName(), BlendSpaceIndex, *GetPathName());
+			return false;
+		}
+	}
+
 	return true;
 }
 
@@ -472,16 +601,6 @@ float UAnimSequenceTransformProviderData::GetSequencePlayLength(int32 SequenceIn
 	return 0.0f;
 }
 
-USkinnedAsset* UAnimSequenceTransformProviderData::GetSkinnedAsset() const
-{
-	return SkinnedAsset;
-}
-
-const TArray<FAnimSequenceTransformProviderSequence>& UAnimSequenceTransformProviderData::GetSequences() const
-{
-	return Sequences;
-}
-
 namespace AnimSequenceTrackTime
 {
 	float GetTimeScale()
@@ -523,31 +642,79 @@ double UAnimSequenceTransformProviderData::GetWorldTime() const
 
 UTransformProviderData* UAnimSequenceTransformProviderData::GetRootTransformProvider()
 {
-	constexpr int32 MaxDepth = 8;
-	int32 Depth = 0;
-	UAnimSequenceTransformProviderData* Current = this;
-	for (; Current->ParentProvider && Depth < MaxDepth; ++Depth)
+	return this;
+}
+
+#if WITH_EDITOR
+void UAnimSequenceTransformProviderData::ForEachReferencingProvider(
+	TFunctionRef<bool(const UAnimSequenceTransformProviderData*)> Predicate,
+	TFunctionRef<void(UAnimSequenceTransformProviderData*)> Action)
+{
+	ForEachObjectOfClass(UAnimSequenceTransformProviderData::StaticClass(), [&Predicate, &Action](UObject* Obj)
 	{
-		Current = Current->ParentProvider;
-	}
-	ensureMsgf(Depth != MaxDepth, TEXT("UAnimSequenceTransformProvider '%s' either has too deep of a hierarchy or is a circular reference."), *GetPathName());
-	return Current;
+		UAnimSequenceTransformProviderData* Provider = CastChecked<UAnimSequenceTransformProviderData>(Obj);
+		if (!Predicate(Provider))
+		{
+			return;
+		}
+
+		Action(Provider);
+
+		if (UAnimSequenceTransformProviderDataInstance* Instance = Cast<UAnimSequenceTransformProviderDataInstance>(Provider))
+		{
+			Instance->InitLayerPools();
+		}
+
+		for (UInstancedSkinnedMeshComponent* Component : FObjectCacheContextScope().GetContext().GetInstancedSkinnedMeshComponents(Provider))
+		{
+			if (Component && IsValid(Component))
+			{
+				Component->MarkRenderStateDirty();
+			}
+		}
+	});
 }
 
-void UAnimSequenceTransformProviderData::SetParentProvider(UAnimSequenceTransformProviderData* ProviderData)
+void UAnimSequenceTransformProviderLayerStack::PostEditChangeProperty(FPropertyChangedEvent& PropertyChangedEvent)
 {
-	ParentProvider = ProviderData;
+	Super::PostEditChangeProperty(PropertyChangedEvent);
 
-	if (ParentProvider)
+	if (Layers.IsEmpty())
 	{
-		Sequences = ProviderData->Sequences;
+		Layers.Emplace();
 	}
-	else
+
+	Layers[0] = FAnimSequenceTransformProviderLayer();
+
+	if (Layers.Num() > MAX_ANIM_SEQUENCE_LAYERS)
 	{
-		Sequences.Empty();
+		Layers.SetNum(MAX_ANIM_SEQUENCE_LAYERS);
 	}
+
+	UAnimSequenceTransformProviderData::ForEachReferencingProvider(
+		[this](const UAnimSequenceTransformProviderData* Provider) { return Provider->GetLayerStack() == this; },
+		[this](UAnimSequenceTransformProviderData* Provider) { Provider->Layers = Layers; });
 }
 
+void UAnimSequenceTransformProviderSequenceList::PostEditChangeProperty(FPropertyChangedEvent& PropertyChangedEvent)
+{
+	Super::PostEditChangeProperty(PropertyChangedEvent);
+
+	UAnimSequenceTransformProviderData::ForEachReferencingProvider(
+		[this](const UAnimSequenceTransformProviderData* Provider) { return Provider->SequenceList == this; },
+		[this](UAnimSequenceTransformProviderData* Provider) { Provider->Sequences = Sequences; });
+}
+
+void UAnimSequenceTransformProviderBlendSpaceList::PostEditChangeProperty(FPropertyChangedEvent& PropertyChangedEvent)
+{
+	Super::PostEditChangeProperty(PropertyChangedEvent);
+
+	UAnimSequenceTransformProviderData::ForEachReferencingProvider(
+		[this](const UAnimSequenceTransformProviderData* Provider) { return Provider->BlendSpaceList == this; },
+		[this](UAnimSequenceTransformProviderData* Provider) { Provider->BlendSpaces = BlendSpaces; });
+}
+#endif
+
 void UAnimSequenceTransformProviderData::Serialize(FArchive& Ar)
 {
 	Ar.UsingCustomVersion(FFortniteMainBranchObjectVersion::GUID);
@@ -577,25 +744,13 @@ void UAnimSequenceTransformProviderData::Serialize(FArchive& Ar)
 			}
 		}
 	}
-}
-
-void UAnimSequenceTransformProviderData::PreSave(FObjectPreSaveContext SaveContext)
-{
-	Super::PreSave(SaveContext);
-
-	if (ParentProvider)
-	{
-		Sequences.Empty();
-	}
-}
-
-void UAnimSequenceTransformProviderData::PostSaveRoot(FObjectPostSaveRootContext SaveContext)
-{
-	Super::PostSaveRoot(SaveContext);
 
-	if (ParentProvider)
+	if (Ar.IsLoading())
 	{
-		Sequences = ParentProvider->GetSequences();
+		if (Layers.IsEmpty())
+		{
+			Layers.Emplace();
+		}
 	}
 }
 
@@ -603,11 +758,6 @@ void UAnimSequenceTransformProviderData::PostLoad()
 {
 	Super::PostLoad();
 
-	if (ParentProvider && IsValid(ParentProvider))
-	{
-		Sequences = ParentProvider->GetSequences();
-	}
-
 #if WITH_EDITOR
 	if (FApp::CanEverRender() && CachedData.SequenceBounds.IsEmpty())
 	{
@@ -627,72 +777,55 @@ void UAnimSequenceTransformProviderData::PostEditChangeProperty(FPropertyChanged
 		MarkPackageDirty();
 	}
 
-	if (PropertyChangedEvent.GetPropertyName() == GET_MEMBER_NAME_CHECKED(UAnimSequenceTransformProviderData, ParentProvider))
+	// Sync from shared assets to inline data
+#if WITH_EDITORONLY_DATA
+	if (SequenceList)
 	{
-		if (ParentProvider)
-		{
-			if (UAnimSequenceTransformProviderDataInstance* Provider = Cast<UAnimSequenceTransformProviderDataInstance>(ParentProvider))
-			{
-				FMessageDialog::Open(
-					EAppMsgType::Ok,
-					NSLOCTEXT("AnimSequenceTransformProvider", "ParentNotInstance", "AnimSequenceTransformProviderDataInstance cannot be set as a parent.")
-				);
-				ParentProvider = nullptr;
-			}
+		Sequences = SequenceList->GetSequences();
+	}
 
-			UAnimSequenceTransformProviderData* Current = ParentProvider;
-			while (Current)
-			{
-				if (Current == this)
-				{
-					FMessageDialog::Open(
-						EAppMsgType::Ok,
-						NSLOCTEXT("AnimSequenceTransformProvider", "CircularReferences", "Setting this parent provider creates a circular reference.")
-					);
-					ParentProvider = nullptr;
-					break;
-				}
-				Current = Current->ParentProvider;
-			}
-		}
+	if (LayerStack)
+	{
+		Layers = LayerStack->GetLayers();
+	}
+
+	if (BlendSpaceList)
+	{
+		BlendSpaces = BlendSpaceList->GetBlendSpaces();
+	}
+#endif
 
-		SetParentProvider(ParentProvider);
+	// Ensure layer 0 always exists
+	if (Layers.IsEmpty())
+	{
+		Layers.Emplace();
 	}
 
-	// Propagate Sequences changes to all children.
-	TArray<TPair<int32, UAnimSequenceTransformProviderData*>> ChildrenByDepth;
+	// Ensure layer 0 is always the default.
+	Layers[0] = FAnimSequenceTransformProviderLayer();
 
-	ForEachObjectOfClass(UAnimSequenceTransformProviderData::StaticClass(), [&ChildrenByDepth](UObject* Obj)
+	// Ensure we stay under the maximum layer limit.
+	if (Layers.Num() > MAX_ANIM_SEQUENCE_LAYERS)
 	{
-		if (UAnimSequenceTransformProviderData* Provider = Cast<UAnimSequenceTransformProviderData>(Obj))
+		Layers.SetNum(MAX_ANIM_SEQUENCE_LAYERS);
+	}
+
+	ForEachObjectOfClass(UAnimSequenceTransformProviderDataInstance::StaticClass(), [this](UObject* Obj)
+	{
+		if (UAnimSequenceTransformProviderDataInstance* Provider = Cast<UAnimSequenceTransformProviderDataInstance>(Obj))
 		{
-			if (Provider->ParentProvider)
+			if (Provider->GetRootTransformProvider() == this)
 			{
-				int32 Depth = 0;
-				UAnimSequenceTransformProviderData* Current = Provider->ParentProvider;
-				while (Current)
-				{
-					++Depth;
-					Current = Current->ParentProvider;
-				}
-				ChildrenByDepth.Emplace(Depth, Provider);
+				Provider->SetParentProvider(this);
 			}
 		}
 	});
 
-	ChildrenByDepth.Sort([](const auto& A, const auto& B) { return A.Key < B.Key; });
-
-	for (const auto& Pair : ChildrenByDepth)
-	{
-		Pair.Value->SetParentProvider(Pair.Value->ParentProvider);
-		Pair.Value->BeginCacheDerivedData(GetTargetPlatformManagerRef().GetRunningTargetPlatform());
-	}
-
 	BeginCacheDerivedData(GetTargetPlatformManagerRef().GetRunningTargetPlatform());
 
-	// Mark render state dirty for components which reference it.
+	// Mark render state dirty for components which reference this asset.
 	FObjectCacheContextScope ObjectCacheScope;
-	for (UInstancedSkinnedMeshComponent* Component : ObjectCacheScope.GetContext().GetInstancedSkinnedMeshComponents(GetRootTransformProvider()))
+	for (UInstancedSkinnedMeshComponent* Component : ObjectCacheScope.GetContext().GetInstancedSkinnedMeshComponents(this))
 	{
 		if (Component && IsValid(Component))
 		{
@@ -753,6 +886,7 @@ UAnimSequenceTransformProviderData* UAnimSequenceTransformProviderData::CreateFr
 		NewProviderData->Sequences.Add(Sequence);
 	}
 
+	NewProviderData->Layers.Emplace();
 	NewProviderData->MarkPackageDirty();
 	FAssetRegistryModule::AssetCreated(NewProviderData);
 	return NewProviderData;
@@ -798,6 +932,7 @@ UAnimSequenceTransformProviderData* UAnimSequenceTransformProviderData::CreateFr
 	}
 
 	NewProviderData->SkinnedAsset = InAnimBank->Asset;
+	NewProviderData->Layers.Emplace();
 	NewProviderData->MarkPackageDirty();
 	FAssetRegistryModule::AssetCreated(NewProviderData);
 	return NewProviderData;
@@ -807,6 +942,41 @@ EDataValidationResult UAnimSequenceTransformProviderData::IsDataValid(FDataValid
 {
 	EDataValidationResult Result = Super::IsDataValid(Context);
 
+	if (SkinnedAsset)
+	{
+		const FSkeletalMeshRenderData* RenderData = SkinnedAsset->GetResourceForRendering();
+		if (RenderData)
+		{
+			if (!RenderData->HasUnifiedBoneMap())
+			{
+				Context.AddError(FText::Format(LOCTEXT("NoUnifiedBoneMap",
+					"AnimSequenceTransformProvider {0} requires 'Optimize for Instancing' on skinned asset {1}. It can be enabled in the skeletal mesh build settings."),
+					FText::FromString(GetPathName()),
+					FText::FromString(SkinnedAsset->GetPathName())
+				));
+				Result = EDataValidationResult::Invalid;
+			}
+
+			const bool bHasMeshSpaceLayer = Layers.ContainsByPredicate([](const FAnimSequenceTransformProviderLayer& Layer) { return Layer.bMeshSpaceRotation; });
+			const int32 MaxBonesPerGroup = bHasMeshSpaceLayer ? SKINNING_MAX_BONES_PER_GROUP_MESHSPACE : SKINNING_MAX_BONES_PER_GROUP;
+			for (const FSkeletalMeshLODRenderData& LODData : RenderData->LODRenderData)
+			{
+				if (LODData.RequiredBones.Num() > MaxBonesPerGroup)
+				{
+					Context.AddWarning(FText::Format(LOCTEXT("TooManyBones",
+						"AnimSequenceTransformProvider {0} skinned asset {1} LOD has {2} bones which exceeds the maximum of {3}.{4}"),
+						FText::FromString(GetPathName()),
+						FText::FromString(SkinnedAsset->GetPathName()),
+						FText::AsNumber(LODData.RequiredBones.Num()),
+						FText::AsNumber(MaxBonesPerGroup),
+						bHasMeshSpaceLayer ? LOCTEXT("MeshSpaceNote", " Limit is reduced due to mesh-space rotation layer.") : FText::GetEmpty()
+					));
+					break;
+				}
+			}
+		}
+	}
+
 	for (const FAnimSequenceTransformProviderSequence& Sequence : Sequences)
 	{
 		if (Sequence.Sequence == nullptr)
@@ -819,6 +989,19 @@ EDataValidationResult UAnimSequenceTransformProviderData::IsDataValid(FDataValid
 			continue;
 		}
 
+		const FFrameRate SampleRate = Sequence.Sequence->GetSamplingFrameRate();
+		const int32 NumFrames = SampleRate.AsFrameTime(Sequence.Sequence->GetPlayLength()).RoundToFrame().Value;
+		if (NumFrames > SKINNING_MAX_KEY_INDEX)
+		{
+			Context.AddWarning(FText::Format(LOCTEXT("TooManyFrames",
+				"AnimSequenceTransformProvider {0} sequence {1} has {2} frames which exceeds the maximum of {3}. Frames will be clamped."),
+				FText::FromString(GetPathName()),
+				FText::FromString(Sequence.Sequence->GetPathName()),
+				FText::AsNumber(NumFrames),
+				FText::AsNumber(SKINNING_MAX_KEY_INDEX)
+			));
+		}
+
 		const USkeleton* SourceSkeleton = Sequence.Sequence->GetSkeleton();
 
 		if (SourceSkeleton == nullptr)
@@ -867,6 +1050,64 @@ EDataValidationResult UAnimSequenceTransformProviderData::IsDataValid(FDataValid
 		}
 	}
 
+	for (int32 BlendSpaceIndex = 0; BlendSpaceIndex < BlendSpaces.Num(); BlendSpaceIndex++)
+	{
+		const UBlendSpace* BlendSpace = BlendSpaces[BlendSpaceIndex].BlendSpace;
+		if (!BlendSpace)
+		{
+			Context.AddError(FText::Format(LOCTEXT("NullBlendSpace",
+				"AnimSequenceTransformProvider {0} has a null blend space at index {1}."),
+				FText::FromString(GetPathName()),
+				FText::AsNumber(BlendSpaceIndex)
+			));
+			Result = EDataValidationResult::Invalid;
+			continue;
+		}
+
+		const TArray<FBlendSample>& Samples = BlendSpace->GetBlendSamples();
+		if (Samples.IsEmpty())
+		{
+			Context.AddError(FText::Format(LOCTEXT("EmptyBlendSpace",
+				"AnimSequenceTransformProvider {0} blend space {1} at index {2} has no samples."),
+				FText::FromString(GetPathName()),
+				FText::FromString(BlendSpace->GetPathName()),
+				FText::AsNumber(BlendSpaceIndex)
+			));
+			Result = EDataValidationResult::Invalid;
+			continue;
+		}
+
+		for (const FBlendSample& Sample : Samples)
+		{
+			const UAnimSequence* SampleSequence = Cast<UAnimSequence>(Sample.Animation);
+			if (!SampleSequence)
+			{
+				continue;
+			}
+
+			bool bFound = false;
+			for (const FAnimSequenceTransformProviderSequence& Seq : Sequences)
+			{
+				if (Seq.Sequence == SampleSequence)
+				{
+					bFound = true;
+					break;
+				}
+			}
+
+			if (!bFound)
+			{
+				Context.AddError(FText::Format(LOCTEXT("BlendSpaceSampleNotInSequences",
+					"AnimSequenceTransformProvider {0} blend space {1} sample sequence {2} is not in the Sequences array."),
+					FText::FromString(GetPathName()),
+					FText::FromString(BlendSpace->GetPathName()),
+					FText::FromString(SampleSequence->GetPathName())
+				));
+				Result = EDataValidationResult::Invalid;
+			}
+		}
+	}
+
 	return Result;
 }
 
@@ -888,29 +1129,36 @@ UAnimSequenceTransformProviderDataInstance* UAnimSequenceTransformProviderDataIn
 
 void UAnimSequenceTransformProviderDataInstance::SetParentProvider(UAnimSequenceTransformProviderData* ProviderData)
 {
-	Super::SetParentProvider(ProviderData);
+	check(ProviderData);
+	ParentProvider = ProviderData;
+	SkinnedAsset = ProviderData->SkinnedAsset;
+	Sequences = ProviderData->Sequences;
+	Layers = ProviderData->Layers;
+	BlendSpaces = ProviderData->BlendSpaces;
+
+#if WITH_EDITORONLY_DATA
+	SequenceList = ProviderData->SequenceList;
+	LayerStack = ProviderData->LayerStack;
+	BlendSpaceList = ProviderData->BlendSpaceList;
+#endif
 
-	if (ProviderData)
-	{
-		SkinnedAsset = ProviderData->SkinnedAsset;
-		bEnabled     = ProviderData->bEnabled;
-	}
+	InitLayerPools();
 }
 
 bool UAnimSequenceTransformProviderDataInstance::IsEnabled() const
 {
-	return Tracks.GetNumTracks() != 0 && Super::IsEnabled();
+	return !TrackPools.IsEmpty() && TrackPools[0].GetNumTracks() != 0 && Super::IsEnabled();
 }
 
 uint32 UAnimSequenceTransformProviderDataInstance::GetUniqueAnimationCount() const
 {
-	return Tracks.GetNumTracks();
+	return !TrackPools.IsEmpty() ? TrackPools[0].GetNumTracks() : 0;
 }
 
 uint32 UAnimSequenceTransformProviderDataInstance::GetSkinningDataOffset(int32 InstanceIndex, const FTransform& ComponentTransform, const FSkinnedMeshInstanceData& InstanceData) const
 {
 	const uint32 AnimationIndex = InstanceData.AnimationIndex;
-	if (AnimationIndex < uint32(Tracks.GetNumTracks()))
+	if (!TrackPools.IsEmpty() && AnimationIndex < uint32(TrackPools[0].GetNumTracks()))
 	{
 		return AnimationIndex * 2u;
 	}
@@ -920,31 +1168,31 @@ uint32 UAnimSequenceTransformProviderDataInstance::GetSkinningDataOffset(int32 I
 
 bool UAnimSequenceTransformProviderDataInstance::HasAnimationBounds() const
 {
-	return ParentProvider->CachedData.SequenceBounds.Num() > 0;
+	return ParentProvider && ParentProvider->HasAnimationBounds();
 }
 
 bool UAnimSequenceTransformProviderDataInstance::GetAnimationBounds(uint32 TrackIndex, FRenderBounds& OutBounds) const
 {
-	if (Tracks.IsActiveIndex(TrackIndex))
+	if (!ParentProvider)
 	{
-		const int32 SequenceIndex = Tracks.GetData(TrackIndex).GetSequenceIndex();
-		if (SequenceIndex < ParentProvider->CachedData.SequenceBounds.Num())
-		{
-			OutBounds = ParentProvider->CachedData.SequenceBounds[SequenceIndex];
-			return true;
-		}
+		return false;
 	}
-	return false;
+	return ParentProvider->GetAnimationBounds(TrackIndex, OutBounds);
+}
+
+UTransformProviderData* UAnimSequenceTransformProviderDataInstance::GetRootTransformProvider()
+{
+	return ParentProvider;
 }
 
 bool UAnimSequenceTransformProviderDataInstance::IsCompiling() const
 {
-    if (ParentProvider && ParentProvider->IsCompiling())
-    {
-        return true;
-    }
+	if (ParentProvider && ParentProvider->IsCompiling())
+	{
+		return true;
+	}
 
-    return Super::IsCompiling();
+	return Super::IsCompiling();
 }
 
 FTransformProviderRenderProxy* UAnimSequenceTransformProviderDataInstance::CreateRenderProxy(FInstancedSkinningSceneExtensionProxy* SceneExtensionProxy) const
@@ -985,6 +1233,133 @@ void UAnimSequenceTransformProviderDataInstance::MarkOwnerRenderStateDirty()
 	}
 }
 
+void UAnimSequenceTransformProviderDataInstance::InitLayerPools(bool bMarkDirty)
+{
+	const int32 NumLayers = Layers.Num();
+	check(NumLayers >= 1);
+
+	TrackPools.Reserve(NumLayers);
+
+	while (TrackPools.Num() < NumLayers)
+	{
+		FAnimSequenceTrackPool& Pool = TrackPools.Emplace_GetRef();
+
+		if (TrackPools.Num() > 1)
+		{
+			Pool.InitFrom(TrackPools[0], [](int32) { return FAnimSequenceTrackPackedData{}; });
+		}
+	}
+
+	while (TrackPools.Num() > NumLayers)
+	{
+		TrackPools.Pop();
+	}
+
+	bTrackPoolsDirty = true;
+	BuildBlendSpaceCaches();
+
+	if (bMarkDirty)
+	{
+		MarkOwnerRenderStateDirty();
+	}
+}
+
+bool UAnimSequenceTransformProviderDataInstance::IsValidTrackAndLayer(int32 TrackIndex, int32 LayerIndex) const
+{
+	return TrackPools.IsValidIndex(LayerIndex) && TrackPools[LayerIndex].IsActiveIndex(TrackIndex);
+}
+
+void UAnimSequenceTransformProviderDataInstance::BuildBlendSpaceCaches()
+{
+	BlendSpaceCaches.SetNum(BlendSpaces.Num());
+
+	for (int32 BlendSpaceIndex = 0; BlendSpaceIndex < BlendSpaces.Num(); BlendSpaceIndex++)
+	{
+		FBlendSpaceCache& Cache = BlendSpaceCaches[BlendSpaceIndex];
+		Cache.bValid = false;
+		Cache.SampleToSequenceIndex.Reset();
+
+		const UBlendSpace* BlendSpace = BlendSpaces[BlendSpaceIndex].BlendSpace;
+		if (!BlendSpace)
+		{
+			continue;
+		}
+
+		const TArray<FBlendSample>& Samples = BlendSpace->GetBlendSamples();
+		Cache.SampleToSequenceIndex.SetNum(Samples.Num());
+
+		int32 NumSamplesFound = 0;
+
+		for (int32 SampleIndex = 0; SampleIndex < Samples.Num(); SampleIndex++)
+		{
+			const UAnimSequence* SampleSequence = Cast<UAnimSequence>(Samples[SampleIndex].Animation);
+			Cache.SampleToSequenceIndex[SampleIndex] = INDEX_NONE;
+
+			if (SampleSequence)
+			{
+				for (int32 SequenceIndex = 0; SequenceIndex < Sequences.Num(); SequenceIndex++)
+				{
+					if (Sequences[SequenceIndex].Sequence == SampleSequence)
+					{
+						Cache.SampleToSequenceIndex[SampleIndex] = SequenceIndex;
+						NumSamplesFound++;
+						break;
+					}
+				}
+			}
+		}
+
+		Cache.bValid = (NumSamplesFound == Samples.Num());
+	}
+}
+
+void UAnimSequenceTransformProviderDataInstance::ResolveBlendSpaceSamples(FAnimSequenceTrackPackedData& Data) const
+{
+	const uint8 BlendSpaceIndex = Data.GetBlendSpaceIndex();
+
+	if (!BlendSpaceCaches.IsValidIndex(BlendSpaceIndex) || !BlendSpaceCaches[BlendSpaceIndex].bValid)
+	{
+		Data.SetNumSamples(0);
+		return;
+	}
+
+	const FBlendSpaceCache& Cache = BlendSpaceCaches[BlendSpaceIndex];
+	const UBlendSpace* BlendSpace = BlendSpaces.IsValidIndex(BlendSpaceIndex) ? BlendSpaces[BlendSpaceIndex].BlendSpace.Get() : nullptr;
+	if (!BlendSpace)
+	{
+		Data.SetNumSamples(0);
+		return;
+	}
+
+	const FVector2f Position = Data.GetBlendSpacePosition();
+	TArray<FBlendSampleData> SampleResults;
+	int32 TriangulationIndex = 0;
+	BlendSpace->GetSamplesFromBlendInput(FVector(Position.X, Position.Y, 0.0f), SampleResults, TriangulationIndex, false);
+
+	if (SampleResults.IsEmpty())
+	{
+		Data.SetNumSamples(0);
+		return;
+	}
+
+	const int32 NumSamples = FMath::Min(SampleResults.Num(), 3);
+	Data.SetNumSamples(NumSamples);
+
+	for (int32 SampleIndex = 0; SampleIndex < NumSamples; SampleIndex++)
+	{
+		const FBlendSampleData& Sample = SampleResults[SampleIndex];
+		const int32 BlendSampleIndex = Sample.SampleDataIndex;
+
+		uint16 SequenceIndex = 0;
+		if (Cache.SampleToSequenceIndex.IsValidIndex(BlendSampleIndex))
+		{
+			SequenceIndex = Cache.SampleToSequenceIndex[BlendSampleIndex];
+		}
+
+		Data.SetSampleData(SampleIndex, SequenceIndex, FFloat16(Sample.GetClampedWeight()));
+	}
+}
+
 void UAnimSequenceTransformProviderDataInstance::SubmitChanges()
 {
 	check(RenderData);
@@ -992,31 +1367,83 @@ void UAnimSequenceTransformProviderDataInstance::SubmitChanges()
 	// SubmitChanges is called during the ISKM instance data update. Guard against multiple concurrent updates.
 	UE::TScopeLock Lock(SubmitChangesMutex);
 
+	check(TrackPools.Num() == Layers.Num());
+
 	if (bNeedsPostLoadInit)
 	{
 		const double WorldTime = GetWorldTime();
 
-		Tracks.EnumerateActiveTracks([&](int32 TrackIndex)
+		for (FAnimSequenceTrackPool& Pool : TrackPools)
 		{
-			Tracks.GetData(TrackIndex).PostLoad(WorldTime);
-			Tracks.SetDirty(TrackIndex);
-		});
+			Pool.EnumerateActiveTracks([&](int32 TrackIndex)
+			{
+				Pool.GetData(TrackIndex).PostLoad(WorldTime);
+				Pool.SetDirty(TrackIndex);
+			});
+		}
 
 		bNeedsPostLoadInit = false;
 	}
 
 	bPendingRenderStateDirty = false;
 
-	if (Tracks.GetNumDirtyTracks() > 0)
+	// Resolve blend space triangulation for dirty blend space tracks
+	for (int32 LayerIndex = 0; LayerIndex < TrackPools.Num(); LayerIndex++)
 	{
-		FAnimSequenceTrackPool::FPatch Patch = Tracks.Finalize([] (FAnimSequenceTrackPackedData& Data)
+		FAnimSequenceTrackPool& Pool = TrackPools[LayerIndex];
+		Pool.EnumerateActiveTracks([&](int32 TrackIndex)
 		{
-			Data.ResetPatchFlags();
+			if (Pool.IsDirtyIndex(TrackIndex))
+			{
+				FAnimSequenceTrackPackedData& Data = Pool.GetData(TrackIndex);
+				if (Data.IsBlendSpace())
+				{
+					ResolveBlendSpaceSamples(Data);
+				}
+			}
 		});
+	}
+
+	const int32 NumLayers = Layers.Num();
+	bool bAnyDirty = bTrackPoolsDirty;
+	for (const FAnimSequenceTrackPool& Pool : TrackPools)
+	{
+		bAnyDirty |= Pool.GetNumDirtyTracks() > 0;
+	}
 
-		ENQUEUE_RENDER_COMMAND(AnimSequenceDataInstanceSubmit)([&RenderData = *RenderData, Patch = MoveTemp(Patch), WorldTime = GetWorldTime()](FRHICommandList&) mutable
+	if (bAnyDirty)
+	{
+		TArray<FAnimSequenceTrackPool::FPatch> Patches;
+		Patches.SetNumZeroed(NumLayers);
+
+		TArray<FAnimSequenceTransformProviderRenderLayer> PatchLayers;
+		PatchLayers.SetNum(NumLayers);
+
+		for (int32 LayerIndex = 0; LayerIndex < NumLayers; LayerIndex++)
+		{
+			if (TrackPools[LayerIndex].GetNumDirtyTracks() > 0)
+			{
+				Patches[LayerIndex] = TrackPools[LayerIndex].Finalize([](FAnimSequenceTrackPackedData& Data)
+				{
+					Data.ResetPatchFlags();
+				});
+			}
+
+			PatchLayers[LayerIndex].Layer  = Layers[LayerIndex];
+			PatchLayers[LayerIndex].Weight = Layers[LayerIndex].Weight;
+		}
+
+		bTrackPoolsDirty = false;
+
+		ENQUEUE_RENDER_COMMAND(AnimSequenceDataInstanceSubmit)(
+			[
+				  &RenderData  = *RenderData
+				, Patches      = MoveTemp(Patches)
+				, InLayers     = MoveTemp(PatchLayers)
+				, WorldTime    = GetWorldTime()
+			](FRHICommandList&) mutable
 		{
-			RenderData.Patch(MoveTemp(Patch), WorldTime);
+			RenderData.Patch(MoveTemp(Patches), MoveTemp(InLayers), WorldTime);
 		});
 	}
 }
@@ -1025,49 +1452,111 @@ void UAnimSequenceTransformProviderDataInstance::Serialize(FArchive& Ar)
 {
 	Ar.UsingCustomVersion(FFortniteMainBranchObjectVersion::GUID);
 	Super::Serialize(Ar);
-	Tracks.Serialize(Ar);
+
+	if (Ar.CustomVer(FFortniteMainBranchObjectVersion::GUID) >= FFortniteMainBranchObjectVersion::AnimSequenceTransformProviderLayers)
+	{
+		int32 NumLayers = TrackPools.Num();
+		Ar << NumLayers;
+
+		if (Ar.IsLoading())
+		{
+			TrackPools.SetNum(NumLayers);
+		}
+
+		for (int32 LayerIndex = 0; LayerIndex < NumLayers; LayerIndex++)
+		{
+			TrackPools[LayerIndex].Serialize(Ar);
+		}
+	}
+	else
+	{
+		FAnimSequenceTrackPool LegacyPool;
+		LegacyPool.Serialize(Ar);
+
+		TrackPools.SetNum(1);
+		TrackPools[0] = MoveTemp(LegacyPool);
+	}
 
 	if (Ar.IsLoading())
 	{
+		InitLayerPools(false);
 		bNeedsPostLoadInit = true;
 	}
 }
 
-int32 UAnimSequenceTransformProviderDataInstance::AllocateAutoPlayTrack(const FAnimSequenceTrackAutoPlayData& Data)
+///////////////////////////////////////////////////////////////////////////////////////////////////
+
+void UAnimSequenceTransformProviderDataInstance::SetLayerWeight(int32 TrackIndex, int32 LayerIndex, float Weight)
 {
-	if (Sequences.IsValidIndex(Data.SequenceIndex))
+	if (!IsValidTrackAndLayer(TrackIndex, LayerIndex))
 	{
-		int32 TrackIndex = Tracks.AllocateTrack(FAnimSequenceTrackPackedData::Pack(Data, GetWorldTime()));
-		if (TrackIndex != INDEX_NONE)
-		{
-			MarkOwnerRenderStateDirty();
-		}
-		return TrackIndex;
+		return;
 	}
 
-	return INDEX_NONE;
+	TrackPools[LayerIndex].GetData(TrackIndex).SetLayerWeight(FMath::Clamp(Weight, 0.0f, 1.0f));
+	TrackPools[LayerIndex].SetDirty(TrackIndex);
+	MarkOwnerRenderStateDirty();
 }
 
-int32 UAnimSequenceTransformProviderDataInstance::AllocateManualTrack(const FAnimSequenceTrackManualData& Data)
+float UAnimSequenceTransformProviderDataInstance::GetLayerWeight(int32 TrackIndex, int32 LayerIndex) const
 {
-	if (Sequences.IsValidIndex(Data.SequenceIndex))
+	if (!IsValidTrackAndLayer(TrackIndex, LayerIndex))
+	{
+		return 1.0f;
+	}
+
+	return TrackPools[LayerIndex].GetData(TrackIndex).GetLayerWeight();
+}
+
+///////////////////////////////////////////////////////////////////////////////////////////////////
+
+int32 UAnimSequenceTransformProviderDataInstance::AllocateTrack()
+{
+	check(!TrackPools.IsEmpty());
+	FAnimSequenceTrackPackedData DefaultData{};
+	const int32 TrackIndex = TrackPools[0].AllocateTrack(DefaultData);
+
+	if (TrackIndex != INDEX_NONE)
 	{
-		int32 TrackIndex = Tracks.AllocateTrack(FAnimSequenceTrackPackedData::Pack(Data, GetSequencePlayLength(Data.SequenceIndex)));
-		if (TrackIndex != INDEX_NONE)
+		for (int32 LayerIndex = 1; LayerIndex < TrackPools.Num(); LayerIndex++)
 		{
-			MarkOwnerRenderStateDirty();
+			TrackPools[LayerIndex].AllocateTrackAt(TrackIndex, DefaultData);
 		}
-		return TrackIndex;
+		MarkOwnerRenderStateDirty();
 	}
 
-	return INDEX_NONE;
+	return TrackIndex;
+}
+
+bool UAnimSequenceTransformProviderDataInstance::DeallocateTrack(int32 TrackIndex)
+{
+	bool bDeallocated = false;
+
+	for (FAnimSequenceTrackPool& Pool : TrackPools)
+	{
+		bDeallocated |= Pool.DeallocateTrack(TrackIndex);
+	}
+
+	if (bDeallocated)
+	{
+		MarkOwnerRenderStateDirty();
+	}
+
+	return bDeallocated;
 }
 
-bool UAnimSequenceTransformProviderDataInstance::UpdateAutoPlayTrack(int32 TrackIndex, const FAnimSequenceTrackAutoPlayData& Data)
+bool UAnimSequenceTransformProviderDataInstance::SetAutoPlayData(int32 TrackIndex, int32 LayerIndex, const FAnimSequenceTrackAutoPlayData& Data)
 {
-	if (Sequences.IsValidIndex(Data.SequenceIndex))
+	if (!TrackPools.IsValidIndex(LayerIndex) || !Sequences.IsValidIndex(Data.SequenceIndex))
+	{
+		return false;
+	}
+
+	FAnimSequenceTrackPool& Pool = TrackPools[LayerIndex];
+
+	if (Pool.IsActiveIndex(TrackIndex))
 	{
-		if (Tracks.UpdateTrack(TrackIndex, FAnimSequenceTrackPackedData::Pack(Data, GetWorldTime())))
+		if (Pool.UpdateTrack(TrackIndex, FAnimSequenceTrackPackedData::Pack(Data, GetWorldTime())))
 		{
 			MarkOwnerRenderStateDirty();
 			return true;
@@ -1077,11 +1566,18 @@ bool UAnimSequenceTransformProviderDataInstance::UpdateAutoPlayTrack(int32 Track
 	return false;
 }
 
-bool UAnimSequenceTransformProviderDataInstance::UpdateManualTrack(int32 TrackIndex, const FAnimSequenceTrackManualData& Data)
+bool UAnimSequenceTransformProviderDataInstance::SetManualData(int32 TrackIndex, int32 LayerIndex, const FAnimSequenceTrackManualData& Data)
 {
-	if (Sequences.IsValidIndex(Data.SequenceIndex))
+	if (!TrackPools.IsValidIndex(LayerIndex) || !Sequences.IsValidIndex(Data.SequenceIndex))
 	{
-		if (Tracks.UpdateTrack(TrackIndex, FAnimSequenceTrackPackedData::Pack(Data, GetSequencePlayLength(Data.SequenceIndex))))
+		return false;
+	}
+
+	FAnimSequenceTrackPool& Pool = TrackPools[LayerIndex];
+
+	if (Pool.IsActiveIndex(TrackIndex))
+	{
+		if (Pool.UpdateTrack(TrackIndex, FAnimSequenceTrackPackedData::Pack(Data, GetSequencePlayLength(Data.SequenceIndex))))
 		{
 			MarkOwnerRenderStateDirty();
 			return true;
@@ -1091,22 +1587,18 @@ bool UAnimSequenceTransformProviderDataInstance::UpdateManualTrack(int32 TrackIn
 	return false;
 }
 
-bool UAnimSequenceTransformProviderDataInstance::DeallocateTrack(int32 TrackIndex)
+bool UAnimSequenceTransformProviderDataInstance::SetPlayRate(int32 TrackIndex, int32 LayerIndex, float PlayRate)
 {
-	if (Tracks.DeallocateTrack(TrackIndex))
+	if (!IsValidTrackAndLayer(TrackIndex, LayerIndex))
 	{
-		MarkOwnerRenderStateDirty();
-		return true;
+		return false;
 	}
-	return false;
-}
 
-bool UAnimSequenceTransformProviderDataInstance::SetLoopMode(int32 TrackIndex, EAnimSequenceTrackLoopMode LoopMode)
-{
-	if (Tracks.IsActiveIndex(TrackIndex))
+	FAnimSequenceTrackPackedData& Data = TrackPools[LayerIndex].GetData(TrackIndex);
+	if (Data.IsAutoPlay() || Data.IsBlendSpace())
 	{
-		Tracks.GetData(TrackIndex).SetLoopMode(LoopMode);
-		Tracks.SetDirty(TrackIndex);
+		Data.SetPlayRate(PlayRate, GetWorldTime());
+		TrackPools[LayerIndex].SetDirty(TrackIndex);
 		MarkOwnerRenderStateDirty();
 		return true;
 	}
@@ -1114,63 +1606,99 @@ bool UAnimSequenceTransformProviderDataInstance::SetLoopMode(int32 TrackIndex, E
 	return false;
 }
 
-bool UAnimSequenceTransformProviderDataInstance::SetPlayRate(int32 TrackIndex, float PlayRate)
+float UAnimSequenceTransformProviderDataInstance::GetPlayRate(int32 TrackIndex, int32 LayerIndex) const
 {
-	if (Tracks.IsActiveIndex(TrackIndex))
+	if (IsValidTrackAndLayer(TrackIndex, LayerIndex))
 	{
-		FAnimSequenceTrackPackedData& Data = Tracks.GetData(TrackIndex);
-		if (Data.IsAutoPlay())
-		{
-			Data.SetPlayRate(PlayRate, GetWorldTime());
-			Tracks.SetDirty(TrackIndex);
-			MarkOwnerRenderStateDirty();
-			return true;
-		}
+		return TrackPools[LayerIndex].GetData(TrackIndex).GetPlayRate();
 	}
 
-	return false;
+	return 1.0f;
 }
 
-float UAnimSequenceTransformProviderDataInstance::GetPlayRate(int32 TrackIndex) const
+bool UAnimSequenceTransformProviderDataInstance::SetLoopMode(int32 TrackIndex, int32 LayerIndex, EAnimSequenceTrackLoopMode LoopMode)
 {
-	if (Tracks.IsActiveIndex(TrackIndex))
+	if (!IsValidTrackAndLayer(TrackIndex, LayerIndex))
 	{
-		return Tracks.GetData(TrackIndex).GetPlayRate();
+		return false;
 	}
 
-	return 1.0f;
+	TrackPools[LayerIndex].GetData(TrackIndex).SetLoopMode(LoopMode);
+	TrackPools[LayerIndex].SetDirty(TrackIndex);
+	MarkOwnerRenderStateDirty();
+	return true;
 }
 
-EAnimSequenceTrackLoopMode UAnimSequenceTransformProviderDataInstance::GetLoopMode(int32 TrackIndex) const
+EAnimSequenceTrackLoopMode UAnimSequenceTransformProviderDataInstance::GetLoopMode(int32 TrackIndex, int32 LayerIndex) const
 {
-	if (Tracks.IsActiveIndex(TrackIndex))
+	if (IsValidTrackAndLayer(TrackIndex, LayerIndex))
 	{
-		return Tracks.GetData(TrackIndex).GetLoopMode();
+		return TrackPools[LayerIndex].GetData(TrackIndex).GetLoopMode();
 	}
 
 	return EAnimSequenceTrackLoopMode::Loop;
 }
 
-int32 UAnimSequenceTransformProviderDataInstance::GetSequenceIndex(int32 TrackIndex) const
+int32 UAnimSequenceTransformProviderDataInstance::GetSequenceIndex(int32 TrackIndex, int32 LayerIndex) const
 {
-	if (Tracks.IsActiveIndex(TrackIndex))
+	if (IsValidTrackAndLayer(TrackIndex, LayerIndex))
 	{
-		return Tracks.GetData(TrackIndex).GetSequenceIndex();
+		return TrackPools[LayerIndex].GetData(TrackIndex).GetSequenceIndex();
 	}
 
 	return INDEX_NONE;
 }
 
-float UAnimSequenceTransformProviderDataInstance::GetPosition(int32 TrackIndex) const
+float UAnimSequenceTransformProviderDataInstance::GetPosition(int32 TrackIndex, int32 LayerIndex) const
 {
-	if (Tracks.IsActiveIndex(TrackIndex))
+	if (!IsValidTrackAndLayer(TrackIndex, LayerIndex))
 	{
-		const FAnimSequenceTrackPackedData& PackedData = Tracks.GetData(TrackIndex);
-		const float SequenceLength = GetSequencePlayLength(PackedData.GetSequenceIndex());
-		return PackedData.GetPosition(GetWorldTime(), SequenceLength);
+		return 0.0f;
 	}
 
-	return 0.0f;
+	const FAnimSequenceTrackPackedData& Data = TrackPools[LayerIndex].GetData(TrackIndex);
+	const int32 SeqIndex = Data.GetSequenceIndex();
+
+	if (SeqIndex == INDEX_NONE || !Sequences.IsValidIndex(SeqIndex))
+	{
+		return 0.0f;
+	}
+
+	const double RawTime = Data.GetTimePosition(GetWorldTime());
+	const float PlayLength = GetSequencePlayLength(SeqIndex);
+	return FAnimSequenceTrackPackedData::WrapPosition(Data.GetLoopMode(), static_cast<float>(RawTime), PlayLength);
+}
+
+bool UAnimSequenceTransformProviderDataInstance::SetBlendSpaceData(int32 TrackIndex, int32 LayerIndex, const FAnimSequenceTrackBlendSpaceData& InData)
+{
+	if (!IsValidTrackAndLayer(TrackIndex, LayerIndex))
+	{
+		return false;
+	}
+
+	const int32 BlendSpaceIndex = InData.BlendSpaceIndex;
+	if (!BlendSpaces.IsValidIndex(BlendSpaceIndex) || !BlendSpaceCaches.IsValidIndex(BlendSpaceIndex) || !BlendSpaceCaches[BlendSpaceIndex].bValid || BlendSpaceCaches[BlendSpaceIndex].SampleToSequenceIndex.IsEmpty())
+	{
+		return false;
+	}
+
+	TrackPools[LayerIndex].UpdateTrack(TrackIndex, FAnimSequenceTrackPackedData::Pack(InData, GetWorldTime()));
+	MarkOwnerRenderStateDirty();
+	return true;
+}
+
+FVector2f UAnimSequenceTransformProviderDataInstance::GetBlendSpacePosition(int32 TrackIndex, int32 LayerIndex) const
+{
+	if (IsValidTrackAndLayer(TrackIndex, LayerIndex))
+	{
+		const FAnimSequenceTrackPackedData& PackedData = TrackPools[LayerIndex].GetData(TrackIndex);
+		if (PackedData.IsBlendSpace())
+		{
+			return PackedData.GetBlendSpacePosition();
+		}
+	}
+
+	return FVector2f::ZeroVector;
 }
 
 #undef LOCTEXT_NAMESPACE
\ No newline at end of file
diff --git a/Engine/Source/Runtime/Engine/Public/Animation/TransformProviderData.h b/Engine/Source/Runtime/Engine/Public/Animation/TransformProviderData.h
--- a/Engine/Source/Runtime/Engine/Public/Animation/TransformProviderData.h
+++ b/Engine/Source/Runtime/Engine/Public/Animation/TransformProviderData.h
@@ -33,7 +33,7 @@ class UTransformProviderData : public UObject
 public:
 	virtual bool IsEnabled() const
 	{
-		return bEnabled;
+		return true;
 	}
 
 	virtual const FGuid& GetTransformProviderID() const
@@ -96,8 +96,4 @@ public:
 		return true;
 	}
 #endif
-
-public:
-	UPROPERTY(EditAnywhere, Category = TransformProvider)
-	bool bEnabled = true;
 };
diff --git a/Engine/Source/Runtime/Renderer/Private/Skinning/AnimSequenceTransformProvider.cpp b/Engine/Source/Runtime/Renderer/Private/Skinning/AnimSequenceTransformProvider.cpp
--- a/Engine/Source/Runtime/Renderer/Private/Skinning/AnimSequenceTransformProvider.cpp
+++ b/Engine/Source/Runtime/Renderer/Private/Skinning/AnimSequenceTransformProvider.cpp
@@ -8,6 +8,7 @@
 #include "Animation/SkeletonRemapping.h"
 #include "Animation/SkeletonRemappingRegistry.h"
 #include "Animation/Skeleton.h"
+#include "Animation/BlendProfile.h"
 #include "Rendering/SkeletalMeshRenderData.h"
 #include "AnimationRuntime.h"
 #include "AnimEncoding.h"
@@ -25,6 +26,7 @@
 #include "SkeletalRenderPublic.h"
 #include "SkinningDefinitions.h"
 
+IMPLEMENT_SCENE_EXTENSION(FAnimSequenceTransformProvider);
 
 DECLARE_STATS_GROUP(TEXT("AnimSequenceTransformProvider"), STATGROUP_AnimSequenceTransformProvider, STATCAT_Advanced);
 
@@ -33,11 +35,8 @@ DECLARE_DWORD_ACCUMULATOR_STAT(TEXT("Mappings"), STAT_AnimSequenceTransformProvi
 
 DECLARE_MEMORY_STAT(TEXT("Sequence Transform Buffer"), STAT_AnimSequenceTransformProvider_SequenceTransformMemory, STATGROUP_AnimSequenceTransformProvider);
 DECLARE_MEMORY_STAT(TEXT("Retargeting Data Buffer"), STAT_AnimSequenceTransformProvider_RetargetingDataMemory, STATGROUP_AnimSequenceTransformProvider);
-DECLARE_MEMORY_STAT(TEXT("Leaf Bone Path Buffer"), STAT_AnimSequenceTransformProvider_LeafBonePathMemory, STATGROUP_AnimSequenceTransformProvider);
 
 DECLARE_DWORD_COUNTER_STAT(TEXT("Active Tracks"), STAT_AnimSequenceTransformProvider_ActiveTracks, STATGROUP_AnimSequenceTransformProvider);
-DECLARE_DWORD_COUNTER_STAT(TEXT("Evaluate SinglePass Blocks"), STAT_AnimSequenceTransformProvider_EvaluateSinglePassBlocks, STATGROUP_AnimSequenceTransformProvider);
-DECLARE_DWORD_COUNTER_STAT(TEXT("Evaluate SinglePass Blend Blocks"), STAT_AnimSequenceTransformProvider_EvaluateSinglePassBlendBlocks, STATGROUP_AnimSequenceTransformProvider);
 DECLARE_DWORD_COUNTER_STAT(TEXT("Evaluate Blocks"), STAT_AnimSequenceTransformProvider_EvaluateBlocks, STATGROUP_AnimSequenceTransformProvider);
 DECLARE_DWORD_COUNTER_STAT(TEXT("Evaluate Blend Blocks"), STAT_AnimSequenceTransformProvider_EvaluateBlendBlocks, STATGROUP_AnimSequenceTransformProvider);
 DECLARE_DWORD_COUNTER_STAT(TEXT("Current Transform Updates"), STAT_AnimSequenceTransformProvider_CurrentTransformUpdates, STATGROUP_AnimSequenceTransformProvider);
@@ -46,10 +45,7 @@ DECLARE_DWORD_COUNTER_STAT(TEXT("Sequence Frame Uploads"), STAT_AnimSequenceTran
 
 DECLARE_GPU_STAT(AnimSequenceTransformProvider);
 
-static FGuid AnimSequenceGPUProviderId(ANIM_SEQUENCE_GPU_TRANSFORM_PROVIDER_GUID);
-static FGuid AnimSequenceCPUProviderId(ANIM_SEQUENCE_CPU_TRANSFORM_PROVIDER_GUID);
-
-IMPLEMENT_SCENE_EXTENSION(FAnimSequenceTransformProvider);
+static FGuid AnimSequenceProviderId(ANIM_SEQUENCE_GPU_TRANSFORM_PROVIDER_GUID);
 
 static TAutoConsoleVariable<float> CVarAnimSequenceTransformProviderTimeScale(
 	TEXT("r.AnimSequenceTransformProvider.TimeScale"),
@@ -58,16 +54,8 @@ static TAutoConsoleVariable<float> CVarAnimSequenceTransformProviderTimeScale(
 	ECVF_Scalability | ECVF_RenderThreadSafe
 );
 
-static TAutoConsoleVariable<bool> CVarAnimSequenceTransformProviderGPU(
-	TEXT("r.AnimSequenceTransformProvider.GPU"),
-	true,
-	TEXT(""),
-	FConsoleVariableDelegate::CreateLambda([](IConsoleVariable* InVariable)
-	{
-		FGlobalComponentRecreateRenderStateContext Context;
-	}),
-	ECVF_Scalability | ECVF_RenderThreadSafe
-);
+static constexpr uint32 Float4sPerTransformFull = sizeof(FAnimSequenceTransformProvider::FBoneTransformWithUniformScale) / sizeof(FVector4f);
+static_assert(Float4sPerTransformFull == 2, "FBoneTransformWithUniformScale must be exactly 2 float4s");
 
 static TAutoConsoleVariable<int32> CVarAnimSequenceTransformProviderUploadBudget(
 	TEXT("r.AnimSequenceTransformProvider.UploadBudget"),
@@ -113,6 +101,14 @@ static TAutoConsoleVariable<int32> CVarAnimSequenceTransformProviderForceSampleR
 	ECVF_Scalability | ECVF_RenderThreadSafe
 );
 
+static TAutoConsoleVariable<bool> CVarAnimSequenceTransformProviderHalfPrecision(
+	TEXT("r.AnimSequenceTransformProvider.HalfPrecision"),
+	true,
+	TEXT("Pack sequence transforms as 16-bit floats. Halves sequence buffer memory with minor precision loss.\n")
+	TEXT("Read-only: must be set at startup. Requires shader recompile in editor, cook for client builds."),
+	ECVF_ReadOnly
+);
+
 FFrameRate GetAnimSequenceSampleRate(const UAnimSequence* AnimSequence)
 {
 	const int32 ForcedSampleRate = CVarAnimSequenceTransformProviderForceSampleRate.GetValueOnRenderThread();
@@ -160,180 +156,36 @@ void FAnimSequenceTransformProvider::FSequence::GetKeyIndicesFromTime(int32& Out
 	OutAlpha = OutKey1 > OutKey0 ? (ExactFrame - (float)OutKey0) / (float)Step : 0.0f;
 }
 
-class FEvaluateAnimSequenceCS : public FGlobalShader
-{
-public:
-	DECLARE_GLOBAL_SHADER(FEvaluateAnimSequenceCS);
-	SHADER_USE_PARAMETER_STRUCT(FEvaluateAnimSequenceCS, FGlobalShader);
-
-	static constexpr uint32 BonesPerGroup = 64u;
-
-	static bool ShouldCompilePermutation(const FGlobalShaderPermutationParameters& Parameters)
-	{
-		return true;
-	}
-
-	static void ModifyCompilationEnvironment(const FGlobalShaderPermutationParameters& Parameters, FShaderCompilerEnvironment& OutEnvironment)
-	{
-		FGlobalShader::ModifyCompilationEnvironment(Parameters, OutEnvironment);
-		OutEnvironment.SetDefine(TEXT("BONES_PER_GROUP"), BonesPerGroup);
-
-		OutEnvironment.CompilerFlags.Add(CFLAG_WarningsAsErrors);
-		OutEnvironment.CompilerFlags.Add(CFLAG_HLSL2021);
-	}
-
-	BEGIN_SHADER_PARAMETER_STRUCT(FParameters, )
-		SHADER_PARAMETER(uint32, HeaderBaseOffset)
-		SHADER_PARAMETER_RDG_BUFFER_SRV(ByteAddressBuffer, RetargetingDataBuffer)
-		SHADER_PARAMETER_RDG_BUFFER_SRV(ByteAddressBuffer, SequenceTransformBuffer)
-		SHADER_PARAMETER_RDG_BUFFER_SRV(ByteAddressBuffer, EvaluateBlockHeaderBuffer)
-		SHADER_PARAMETER_RDG_BUFFER_UAV(RWByteAddressBuffer, LocalSpaceTransformBufferUAV)
-	END_SHADER_PARAMETER_STRUCT()
-};
-IMPLEMENT_GLOBAL_SHADER(FEvaluateAnimSequenceCS, "/Engine/Private/Skinning/AnimSequenceTransformProvider.usf", "EvaluateAnimSequenceCS", SF_Compute);
-
-class FEvaluateAnimSequenceBlendCS : public FGlobalShader
-{
-public:
-	DECLARE_GLOBAL_SHADER(FEvaluateAnimSequenceBlendCS);
-	SHADER_USE_PARAMETER_STRUCT(FEvaluateAnimSequenceBlendCS, FGlobalShader);
-
-	static constexpr uint32 BonesPerGroup = 64u;
-
-	static bool ShouldCompilePermutation(const FGlobalShaderPermutationParameters& Parameters)
-	{
-		return true;
-	}
-
-	static void ModifyCompilationEnvironment(const FGlobalShaderPermutationParameters& Parameters, FShaderCompilerEnvironment& OutEnvironment)
-	{
-		FGlobalShader::ModifyCompilationEnvironment(Parameters, OutEnvironment);
-		OutEnvironment.SetDefine(TEXT("BONES_PER_GROUP"), BonesPerGroup);
-
-		OutEnvironment.CompilerFlags.Add(CFLAG_WarningsAsErrors);
-		OutEnvironment.CompilerFlags.Add(CFLAG_HLSL2021);
-	}
-
-	BEGIN_SHADER_PARAMETER_STRUCT(FParameters, )
-		SHADER_PARAMETER(uint32, HeaderBaseOffset)
-		SHADER_PARAMETER_RDG_BUFFER_SRV(ByteAddressBuffer, RetargetingDataBuffer)
-		SHADER_PARAMETER_RDG_BUFFER_SRV(ByteAddressBuffer, SequenceTransformBuffer)
-		SHADER_PARAMETER_RDG_BUFFER_SRV(ByteAddressBuffer, EvaluateBlockHeaderBuffer)
-		SHADER_PARAMETER_RDG_BUFFER_UAV(RWByteAddressBuffer, LocalSpaceTransformBufferUAV)
-	END_SHADER_PARAMETER_STRUCT()
-};
-IMPLEMENT_GLOBAL_SHADER(FEvaluateAnimSequenceBlendCS, "/Engine/Private/Skinning/AnimSequenceTransformProvider.usf", "EvaluateAnimSequenceBlendCS", SF_Compute);
-
-class FConcatenateLocalToGlobalCS : public FGlobalShader
-{
-public:
-	DECLARE_GLOBAL_SHADER(FConcatenateLocalToGlobalCS);
-	SHADER_USE_PARAMETER_STRUCT(FConcatenateLocalToGlobalCS, FGlobalShader);
-
-	static constexpr uint32 BonesPerGroup = FEvaluateAnimSequenceCS::BonesPerGroup;
-
-	static bool ShouldCompilePermutation(const FGlobalShaderPermutationParameters& Parameters)
-	{
-		return true;
-	}
-
-	static void ModifyCompilationEnvironment(const FGlobalShaderPermutationParameters& Parameters, FShaderCompilerEnvironment& OutEnvironment)
-	{
-		FGlobalShader::ModifyCompilationEnvironment(Parameters, OutEnvironment);
-		OutEnvironment.SetDefine(TEXT("BONES_PER_GROUP"), BonesPerGroup);
-
-		OutEnvironment.CompilerFlags.Add(CFLAG_WarningsAsErrors);
-		OutEnvironment.CompilerFlags.Add(CFLAG_HLSL2021);
-	}
-
-	BEGIN_SHADER_PARAMETER_STRUCT(FParameters, )
-		SHADER_PARAMETER(uint32, HeaderBaseOffset)
-		SHADER_PARAMETER_RDG_BUFFER_SRV(ByteAddressBuffer, LocalSpaceTransformBuffer)
-		SHADER_PARAMETER_RDG_BUFFER_SRV(ByteAddressBuffer, ConcatenateBlockHeaderBuffer)
-		SHADER_PARAMETER_RDG_BUFFER_SRV(ByteAddressBuffer, LeafBonePathBuffer)
-		SHADER_PARAMETER_RDG_BUFFER_UAV(RWByteAddressBuffer, GlobalSpaceTransformBufferUAV)
-	END_SHADER_PARAMETER_STRUCT()
-};
-IMPLEMENT_GLOBAL_SHADER(FConcatenateLocalToGlobalCS, "/Engine/Private/Skinning/AnimSequenceTransformProvider.usf", "ConcatenateLocalToGlobalCS", SF_Compute);
-
-class FScatterTransformsCS : public FGlobalShader
-{
-public:
-	DECLARE_GLOBAL_SHADER(FScatterTransformsCS);
-	SHADER_USE_PARAMETER_STRUCT(FScatterTransformsCS, FGlobalShader);
-
-	static constexpr uint32 BonesPerGroup = FEvaluateAnimSequenceCS::BonesPerGroup;
-
-	static bool ShouldCompilePermutation(const FGlobalShaderPermutationParameters& Parameters)
-	{
-		return true;
-	}
-
-	static void ModifyCompilationEnvironment(const FGlobalShaderPermutationParameters& Parameters, FShaderCompilerEnvironment& OutEnvironment)
-	{
-		FGlobalShader::ModifyCompilationEnvironment(Parameters, OutEnvironment);
-		OutEnvironment.SetDefine(TEXT("BONES_PER_GROUP"), BonesPerGroup);
-
-		OutEnvironment.CompilerFlags.Add(CFLAG_WarningsAsErrors);
-		OutEnvironment.CompilerFlags.Add(CFLAG_HLSL2021);
-	}
-
-	BEGIN_SHADER_PARAMETER_STRUCT(FParameters, )
-		SHADER_PARAMETER(uint32, HeaderBaseOffset)
-		SHADER_PARAMETER_RDG_BUFFER_SRV(ByteAddressBuffer, GlobalSpaceTransformBuffer)
-		SHADER_PARAMETER_RDG_BUFFER_SRV(ByteAddressBuffer, SequenceTransformBuffer)
-		SHADER_PARAMETER_RDG_BUFFER_SRV(ByteAddressBuffer, BoneMapBuffer)
-		SHADER_PARAMETER_RDG_BUFFER_SRV(ByteAddressBuffer, ScatterBlockHeaderBuffer)
-		SHADER_PARAMETER_RDG_COMPRESSED_BONE_TRANSFORM_UAV(SkinningTransformBufferUAV)
-	END_SHADER_PARAMETER_STRUCT()
-};
-IMPLEMENT_GLOBAL_SHADER(FScatterTransformsCS, "/Engine/Private/Skinning/AnimSequenceTransformProvider.usf", "ScatterTransformsCS", SF_Compute);
-
 class FEvaluateConcatenateScatterCS : public FGlobalShader
 {
 public:
 	DECLARE_GLOBAL_SHADER(FEvaluateConcatenateScatterCS);
 	SHADER_USE_PARAMETER_STRUCT(FEvaluateConcatenateScatterCS, FGlobalShader);
 
-	class FBonesPerGroupDim : SHADER_PERMUTATION_SPARSE_INT("BONES_PER_GROUP", 64, 128, 256);
-	using FPermutationDomain = TShaderPermutationDomain<FBonesPerGroupDim>;
+	class FBonesPerGroupDim : SHADER_PERMUTATION_SPARSE_INT("BONES_PER_GROUP", 64, 128, 256, 512, 1024);
+	class FSingleLayerDim : SHADER_PERMUTATION_BOOL("SINGLE_LAYER");
+	class FMeshSpaceBlend : SHADER_PERMUTATION_BOOL("HAS_MESH_SPACE_LAYERS");
+	using FPermutationDomain = TShaderPermutationDomain<FBonesPerGroupDim, FSingleLayerDim, FMeshSpaceBlend>;
 
 	static bool ShouldCompilePermutation(const FGlobalShaderPermutationParameters& Parameters)
 	{
-		return true;
-	}
-
-	static void ModifyCompilationEnvironment(const FGlobalShaderPermutationParameters& Parameters, FShaderCompilerEnvironment& OutEnvironment)
-	{
-		FGlobalShader::ModifyCompilationEnvironment(Parameters, OutEnvironment);
+		FPermutationDomain PermutationVector(Parameters.PermutationId);
+		const int32 BonesPerGroup = PermutationVector.Get<FBonesPerGroupDim>();
+		const bool bSingleLayer = PermutationVector.Get<FSingleLayerDim>();
+		const bool bMeshSpace = PermutationVector.Get<FMeshSpaceBlend>();
 
-		OutEnvironment.CompilerFlags.Add(CFLAG_WarningsAsErrors);
-		OutEnvironment.CompilerFlags.Add(CFLAG_HLSL2021);
-	}
-
-	BEGIN_SHADER_PARAMETER_STRUCT(FParameters, )
-		SHADER_PARAMETER(uint32, NumTotalHeaders)
-		SHADER_PARAMETER(uint32, HeaderBaseOffset)
-		SHADER_PARAMETER_RDG_BUFFER_SRV(ByteAddressBuffer, RetargetingDataBuffer)
-		SHADER_PARAMETER_RDG_BUFFER_SRV(ByteAddressBuffer, SequenceTransformBuffer)
-		SHADER_PARAMETER_RDG_BUFFER_SRV(ByteAddressBuffer, EvaluateBlockHeaderBuffer)
-		SHADER_PARAMETER_RDG_BUFFER_SRV(ByteAddressBuffer, BoneMapBuffer)
-		SHADER_PARAMETER_RDG_COMPRESSED_BONE_TRANSFORM_UAV(SkinningTransformBufferUAV)
-	END_SHADER_PARAMETER_STRUCT()
-};
-IMPLEMENT_GLOBAL_SHADER(FEvaluateConcatenateScatterCS, "/Engine/Private/Skinning/AnimSequenceTransformProvider.usf", "EvaluateConcatenateScatterCS", SF_Compute);
-
-class FEvaluateBlendConcatenateScatterCS : public FGlobalShader
-{
-public:
-	DECLARE_GLOBAL_SHADER(FEvaluateBlendConcatenateScatterCS);
-	SHADER_USE_PARAMETER_STRUCT(FEvaluateBlendConcatenateScatterCS, FGlobalShader);
+		// Single-layer can't have mesh-space (layer 0 is always local-space override)
+		if (bSingleLayer && bMeshSpace)
+		{
+			return false;
+		}
 
-	class FBonesPerGroupDim : SHADER_PERMUTATION_SPARSE_INT("BONES_PER_GROUP", 64, 128, 256);
-	using FPermutationDomain = TShaderPermutationDomain<FBonesPerGroupDim>;
+		// Mesh-space requires 2x groupshared
+		if (bMeshSpace && BonesPerGroup > SKINNING_MAX_BONES_PER_GROUP_MESHSPACE)
+		{
+			return false;
+		}
 
-	static bool ShouldCompilePermutation(const FGlobalShaderPermutationParameters& Parameters)
-	{
 		return true;
 	}
 
@@ -343,33 +195,63 @@ public:
 
 		OutEnvironment.CompilerFlags.Add(CFLAG_WarningsAsErrors);
 		OutEnvironment.CompilerFlags.Add(CFLAG_HLSL2021);
+		OutEnvironment.SetDefine(TEXT("HALF_PRECISION_SEQUENCES"), CVarAnimSequenceTransformProviderHalfPrecision.GetValueOnAnyThread() ? 1 : 0);
 	}
 
 	BEGIN_SHADER_PARAMETER_STRUCT(FParameters, )
 		SHADER_PARAMETER(uint32, NumTotalHeaders)
 		SHADER_PARAMETER(uint32, HeaderBaseOffset)
+		SHADER_PARAMETER(uint32, SlotsPerInstance)
+		SHADER_PARAMETER(uint32, NumCompactBones)
+		SHADER_PARAMETER(uint32, NumLayers)
+		SHADER_PARAMETER(uint32, MeshSpaceMask)
+		SHADER_PARAMETER(uint32, LayerModeMask)
+
 		SHADER_PARAMETER_RDG_BUFFER_SRV(ByteAddressBuffer, RetargetingDataBuffer)
-		SHADER_PARAMETER_RDG_BUFFER_SRV(ByteAddressBuffer, SequenceTransformBuffer)
-		SHADER_PARAMETER_RDG_BUFFER_SRV(ByteAddressBuffer, EvaluateBlockHeaderBuffer)
+		SHADER_PARAMETER_RDG_BUFFER_SRV(StructuredBuffer<FVector4f>, SequenceTransformBuffer)
+		SHADER_PARAMETER_RDG_BUFFER_SRV(StructuredBuffer<FEvaluateLayerHeader>, LayerHeaderBuffer)
+		SHADER_PARAMETER_RDG_BUFFER_SRV(StructuredBuffer<FSamplePoseData>, SamplePoseBuffer)
 		SHADER_PARAMETER_RDG_BUFFER_SRV(ByteAddressBuffer, BoneMapBuffer)
+		SHADER_PARAMETER_RDG_BUFFER_SRV(ByteAddressBuffer, BoneMaskBuffer)
 		SHADER_PARAMETER_RDG_COMPRESSED_BONE_TRANSFORM_UAV(SkinningTransformBufferUAV)
 	END_SHADER_PARAMETER_STRUCT()
 };
-IMPLEMENT_GLOBAL_SHADER(FEvaluateBlendConcatenateScatterCS, "/Engine/Private/Skinning/AnimSequenceTransformProvider.usf", "EvaluateBlendConcatenateScatterCS", SF_Compute);
+IMPLEMENT_GLOBAL_SHADER(FEvaluateConcatenateScatterCS, "/Engine/Private/Skinning/AnimSequenceTransformProvider.usf", "EvaluateConcatenateScatterCS", SF_Compute);
 
 FAnimSequenceTransformProvider::FProxyUserData::FProxyUserData(FAnimSequenceTransformProviderProxy* Proxy)
-	: Tracks(Proxy->GetTracks())
+	: RenderData(Proxy->GetRenderData())
 	, SceneExtensionProxy(*Proxy->GetSceneExtensionProxy())
-	, TransformStorageMode(SceneExtensionProxy.GetBoneTransformStorageMode())
-{}
+{
+	check(RenderData);
+}
 
 FAnimSequenceTransformProvider::FAnimSequenceTransformProvider(FScene& InScene)
 	: ISceneExtension(InScene)
 	, SequenceTransformPersistentBuffer(0, TEXT("AnimSequenceTransformProvider.SequenceTransformBuffer"))
 	, RetargetingDataPersistentBuffer(0, TEXT("AnimSequenceTransformProvider.RetargetingDataBuffer"))
-	, LeafBonePathPersistentBuffer(0, TEXT("AnimSequenceTransformProvider.LeafBonePathBuffer"))
+	, BoneHierarchyPersistentBuffer(0, TEXT("AnimSequenceTransformProvider.BoneHierarchyBuffer"))
+	, BoneMaskPersistentBuffer(1024u, TEXT("AnimSequenceTransformProvider.BoneMaskBuffer"))
+	, bHalfPrecisionSequences(CVarAnimSequenceTransformProviderHalfPrecision.GetValueOnAnyThread())
 {}
 
+void FAnimSequenceTransformProvider::WriteBoneTransform(TArrayView<FVector4f>& DstFloat4s, uint32 Index, const FBoneTransformWithUniformScale& Transform) const
+{
+	if (bHalfPrecisionSequences)
+	{
+		FVector4f Packed;
+		Packed.X = BitCast<float>(uint32(FFloat16(Transform.Rotation[0]).Encoded) | (uint32(FFloat16(Transform.Rotation[1]).Encoded) << 16));
+		Packed.Y = BitCast<float>(uint32(FFloat16(Transform.Rotation[2]).Encoded) | (uint32(FFloat16(Transform.Rotation[3]).Encoded) << 16));
+		Packed.Z = BitCast<float>(uint32(FFloat16(Transform.Translation[0]).Encoded) | (uint32(FFloat16(Transform.Translation[1]).Encoded) << 16));
+		Packed.W = BitCast<float>(uint32(FFloat16(Transform.Translation[2]).Encoded) | (uint32(FFloat16(Transform.Scale).Encoded) << 16));
+		DstFloat4s[Index] = Packed;
+	}
+	else
+	{
+		DstFloat4s[Index * 2]     = FVector4f(Transform.Rotation[0], Transform.Rotation[1], Transform.Rotation[2], Transform.Rotation[3]);
+		DstFloat4s[Index * 2 + 1] = FVector4f(Transform.Translation[0], Transform.Translation[1], Transform.Translation[2], Transform.Scale);
+	}
+}
+
 bool FAnimSequenceTransformProvider::ShouldCreateExtension(FScene& InScene)
 {
 	return IsGPUSkinSceneExtensionEnabled();
@@ -379,17 +261,9 @@ void FAnimSequenceTransformProvider::InitExtension(FScene& InScene)
 {
 	if (auto TransformProvider = InScene.GetExtensionPtr<FSkinningTransformProvider>())
 	{
-		// Register GPU transform provider
-		TransformProvider->RegisterProvider(
-			AnimSequenceGPUProviderId,
-			FSkinningTransformProvider::FOnProvideTransforms::CreateRaw(this, &FAnimSequenceTransformProvider::ProvideGPUTransforms),
-			false /* Use skeleton batching */
-		);
-
-		// Register CPU transform provider
 		TransformProvider->RegisterProvider(
-			AnimSequenceCPUProviderId,
-			FSkinningTransformProvider::FOnProvideTransforms::CreateRaw(this, &FAnimSequenceTransformProvider::ProvideCPUTransforms),
+			AnimSequenceProviderId,
+			FSkinningTransformProvider::FOnProvideTransforms::CreateRaw(this, &FAnimSequenceTransformProvider::ProvideTransforms),
 			false /* Use skeleton batching */
 		);
 	}
@@ -613,147 +487,32 @@ void FAnimSequenceTransformProvider::BuildRetargetingTables(
 	}
 }
 
-void FAnimSequenceTransformProvider::BuildLeafBonePathData(FLeafBonePathData& OutData, const FBoneTables& BoneTables)
-{
-	const uint32 NumBones = BoneTables.NumCompactBones;
-
-	struct FBoneData
-	{
-		uint16 Depth    : 14 = 0;
-		uint16 bIsLeaf  : 1  = 1;
-		uint16 bIsOwned : 1  = 0;
-	};
-
-	TArray<FBoneData, FConcurrentLinearArrayAllocator> BoneDatas;
-	BoneDatas.SetNum(NumBones);
-	uint32 NumLeafBones = NumBones;
-
-	for (uint32 ChildIndex = 1; ChildIndex < NumBones; ++ChildIndex)
-	{
-		const uint32 ParentIndex = BoneTables.CompactToParent[ChildIndex];
-		if (ParentIndex != INVALID_BONE_INDEX)
-		{
-			FBoneData& ParentBone = BoneDatas[ParentIndex];
-			FBoneData& ChildBone  = BoneDatas[ChildIndex];
-
-			if (ParentBone.bIsLeaf)
-			{
-				ParentBone.bIsLeaf = 0;
-				NumLeafBones--;
-			}
-
-			ChildBone.Depth = ParentBone.Depth + 1;
-		}
-	}
-
-	uint32 NumPathNodes = 0;
-	for (uint32 BoneIndex = 0; BoneIndex < NumBones; ++BoneIndex)
-	{
-		if (BoneDatas[BoneIndex].bIsLeaf)
-		{
-			NumPathNodes += BoneDatas[BoneIndex].Depth;
-		}
-	}
-
-	const uint32 OffsetTableSize = NumLeafBones;
-	const uint32 HeadersSize = NumLeafBones;
-	const uint32 TotalSize = OffsetTableSize + HeadersSize + NumPathNodes;
-
-	OutData.Buffer.SetNumUninitialized(TotalSize);
-	OutData.NumLeafBones = NumLeafBones;
-
-	uint32* OffsetTable = OutData.Buffer.GetData();
-	uint32* PathData = OutData.Buffer.GetData() + OffsetTableSize;
-
-	int32 LeafIndex = 0;
-	uint32 PathDataOffset = OffsetTableSize * sizeof(uint32);
-
-	for (uint32 BoneIndex = 0; BoneIndex < NumBones; ++BoneIndex)
-	{
-		if (!BoneDatas[BoneIndex].bIsLeaf)
-		{
-			continue;
-		}
-
-		const uint16 Depth = BoneDatas[BoneIndex].Depth;
-
-		OffsetTable[LeafIndex] = PathDataOffset;
-
-		FBoneData& RootBone = BoneDatas[0];
-		const bool bIsRootOwned = !RootBone.bIsOwned;
-		RootBone.bIsOwned = 1;
-
-		UE::HLSL::FLeafBonePathHeader* Header = reinterpret_cast<UE::HLSL::FLeafBonePathHeader*>(PathData);
-		Header->PathDepth = Depth;
-		Header->bIsRootOwned = bIsRootOwned;
-		PathData++;
-		PathDataOffset += sizeof(uint32);
-
-		// Write path nodes in reverse (root -> leaf), skipping root
-		if (Depth > 0)
-		{
-			UE::HLSL::FLeafBonePathNode* Nodes = reinterpret_cast<UE::HLSL::FLeafBonePathNode*>(PathData);
-
-			int32 CurrentNodeIndex = Depth - 1;
-			int32 CurrentBoneIndex = BoneIndex;
-
-			while (CurrentNodeIndex >= 0)
-			{
-				check(CurrentBoneIndex != INVALID_BONE_INDEX);
-
-				FBoneData& CurrentBone = BoneDatas[CurrentBoneIndex];
-
-				UE::HLSL::FLeafBonePathNode& Node = Nodes[CurrentNodeIndex];
-				Node.BoneIndex = CurrentBoneIndex;
-				Node.bIsLeaf   = CurrentBone.bIsLeaf;
-				Node.bIsOwned  = !CurrentBone.bIsOwned;
-				CurrentBone.bIsOwned = 1;
-
-				CurrentBoneIndex = BoneTables.CompactToParent[CurrentBoneIndex];
-				CurrentNodeIndex--;
-			}
-
-			PathData += Depth;
-			PathDataOffset += Depth * sizeof(uint32);
-		}
-
-		LeafIndex++;
-	}
-}
-
 void FAnimSequenceTransformProvider::RegisterProxy(FAnimSequenceTransformProviderProxy* Proxy)
 {
 	TConstArrayView<FAnimSequenceTransformProviderSequence> ProviderSequences = Proxy->GetSequences();
 	FInstancedSkinningSceneExtensionProxy* ExtensionProxy = Proxy->GetSceneExtensionProxy();
 	const USkinnedAsset* SkinnedAsset = ExtensionProxy->GetSkinnedAsset();
-	const FSkeletalMeshRenderData* RenderData = SkinnedAsset->GetResourceForRendering();
+	const FSkeletalMeshRenderData* MeshRenderData = SkinnedAsset->GetResourceForRendering();
 	const FSkeletalMeshObject* MeshObject = ExtensionProxy->GetMeshObject();
 	TConstArrayView<FMatrix44f> InvGlobalRefPoseTransforms = SkinnedAsset->GetRefBasesInvMatrix();
 
 	FProxyUserData* ProxyUserData = new FProxyUserData(Proxy);
+	check(Proxy->GetNumLayers() >= 1);
 	ProxyUserData->Sequences.Reserve(ProviderSequences.Num());
 	ProxyUserData->SequenceToMappingTable.Reserve(ProviderSequences.Num());
 	ProxyUserData->NumMeshBones = InvGlobalRefPoseTransforms.Num();
 	ProxyUserData->MinLODLevel = MeshObject->MinLODLevel;
 
 	const int32 MinLODIndex = MeshObject->MinLODLevel;
-	const int32 NumLODs = RenderData->LODRenderData.Num();
+	const int32 NumLODs = MeshRenderData->LODRenderData.Num();
 	ProxyUserData->LODs.SetNum(NumLODs);
 	ProxyUserData->SequenceToMappingTable.Reserve(ProviderSequences.Num() * NumLODs);
 
 	for (int32 LODIndex = MinLODIndex; LODIndex < NumLODs; ++LODIndex)
 	{
 		FProxyUserData::FLOD& LOD = ProxyUserData->LODs[LODIndex];
-		BuildBoneTables(LOD.BoneTables, SkinnedAsset, RenderData->LODRenderData[LODIndex], ExtensionProxy->HasSocketsInBoneMap());
-
-		// Single pass mode skips building leaf bone paths in favor of doing everything in one compute dispatch.
-		ComputeSinglePassGroupParams(LOD.BoneTables.NumCompactBones, LOD.GroupSize, LOD.InstancesPerGroup, LOD.SlotsPerInstance);
-		const bool bSinglePass = LOD.GroupSize > 0 && (RenderData->HasUnifiedBoneMap() || ExtensionProxy->GetBoneTransformStorageMode() == EBoneTransformStorageMode::Direct);
-		if (!bSinglePass)
-		{
-			BuildLeafBonePathData(LOD.LeafBonePathData, LOD.BoneTables);
-		}
-		LOD.bSinglePass = bSinglePass;
+		BuildBoneTables(LOD.BoneTables, SkinnedAsset, MeshRenderData->LODRenderData[LODIndex], ExtensionProxy->HasSocketsInBoneMap());
+		ComputeGroupParams(LOD.BoneTables.NumCompactBones, LOD.GroupSize, LOD.InstancesPerGroup, LOD.SlotsPerInstance);
 	}
 
 	for (const FAnimSequenceTransformProviderSequence& ProviderSequence : ProviderSequences)
@@ -772,7 +531,7 @@ void FAnimSequenceTransformProvider::RegisterProxy(FAnimSequenceTransformProvide
 			Sequence->Sequence               = AnimSequence;
 			Sequence->PlayLength             = AnimSequence->GetPlayLength();
 			Sequence->SampleRate             = GetAnimSequenceSampleRate(AnimSequence);
-			Sequence->NumFrames              = Sequence->SampleRate.AsFrameTime(Sequence->PlayLength).RoundToFrame().Value;
+			Sequence->NumFrames              = FMath::Min(Sequence->SampleRate.AsFrameTime(Sequence->PlayLength).RoundToFrame().Value, SKINNING_MAX_KEY_INDEX);
 			Sequence->NumAnimBones           = SourceSkeleton->GetReferenceSkeleton().GetNum();
 			Sequence->NumSampleRateLevels    = GetMaxSampleRateLevels(Sequence->NumFrames);
 			Sequence->CurrentSampleRateLevel = Sequence->NumSampleRateLevels;
@@ -832,6 +591,53 @@ void FAnimSequenceTransformProvider::RegisterProxy(FAnimSequenceTransformProvide
 		ProxyUserData->Sequences.Emplace(Sequence.Get());
 	}
 
+	const USkeleton* Skeleton = SkinnedAsset->GetSkeleton();
+	FAnimSequenceTransformProviderRenderData& RenderData = *ProxyUserData->GetRenderData();
+	const int32 NumLayers = RenderData.GetNumLayers();
+
+	for (FProxyUserData::FLOD& LOD : ProxyUserData->LODs)
+	{
+		LOD.BoneMaskBufferOffsets.SetNum(NumLayers);
+		for (int32 LayerIndex = 0; LayerIndex < NumLayers; LayerIndex++)
+		{
+			LOD.BoneMaskBufferOffsets[LayerIndex] = FProxyUserData::InvalidBoneMaskOffset;
+		}
+	}
+
+	ProxyUserData->BlendProfiles.SetNum(NumLayers);
+
+	for (int32 LayerIndex = 0; LayerIndex < NumLayers; LayerIndex++)
+	{
+		const FAnimSequenceTransformProviderLayer& Layer = RenderData.GetLayers()[LayerIndex].Layer;
+		if (Layer.BoneMaskProfileName.IsNone())
+		{
+			continue;
+		}
+
+		const UBlendProfile* BlendProfile = Skeleton ? const_cast<USkeleton*>(Skeleton)->GetBlendProfile(Layer.BoneMaskProfileName) : nullptr;
+		if (!BlendProfile || BlendProfile->GetMode() != EBlendProfileMode::BlendMask)
+		{
+			continue;
+		}
+
+		ProxyUserData->BlendProfiles[LayerIndex] = BlendProfile;
+
+		for (int32 LODIndex = ProxyUserData->MinLODLevel; LODIndex < ProxyUserData->LODs.Num(); LODIndex++)
+		{
+			FProxyUserData::FLOD& LOD = ProxyUserData->LODs[LODIndex];
+			if (!LOD.IsValid())
+			{
+				continue;
+			}
+
+			const uint32 NumCompactBones = LOD.BoneTables.NumCompactBones;
+			const uint32 AllocationOffset = BoneMaskAllocator.Allocate(NumCompactBones);
+			LOD.BoneMaskBufferOffsets[LayerIndex] = AllocationOffset * sizeof(float);
+		}
+
+		ProxyUserDatasToUpload.AddUnique(ProxyUserData);
+	}
+
 	ExtensionProxy->SetUserData(ProxyUserData);
 }
 
@@ -845,10 +651,21 @@ void FAnimSequenceTransformProvider::UnregisterProxy(FAnimSequenceTransformProvi
 
 	for (FProxyUserData::FLOD& LOD : ProxyUserData->LODs)
 	{
-		if (LOD.LeafBonePathBufferAllocation.IsValid())
+		for (int32 LayerIndex = 0; LayerIndex < LOD.BoneMaskBufferOffsets.Num(); LayerIndex++)
+		{
+			if (LOD.BoneMaskBufferOffsets[LayerIndex] != FProxyUserData::InvalidBoneMaskOffset)
+			{
+				BoneMaskAllocator.Free(LOD.BoneMaskBufferOffsets[LayerIndex] / sizeof(float), LOD.BoneTables.NumCompactBones);
+			}
+		}
+	}
+
+	for (FProxyUserData::FLOD& LOD : ProxyUserData->LODs)
+	{
+		if (LOD.BoneMapBufferAllocation.IsValid())
 		{
-			LeafBonePathAllocator.Free(LOD.LeafBonePathBufferAllocation.Offset, LOD.LeafBonePathBufferAllocation.Count);
-			LOD.LeafBonePathBufferAllocation = {};
+			BoneHierarchyAllocator.Free(LOD.BoneMapBufferAllocation.Offset, LOD.BoneMapBufferAllocation.Count);
+			LOD.BoneMapBufferAllocation = {};
 		}
 	}
 
@@ -906,11 +723,10 @@ void FAnimSequenceTransformProvider::UnregisterProxy(FAnimSequenceTransformProvi
 		}
 	}
 
+	ProxyUserDatasToUpload.RemoveSingleSwap(ProxyUserData);
 	delete ProxyUserData;
 }
 
-
-
 template <typename T>
 FRDGBuffer* CreateByteAddressBuffer(FRDGBuilder& GraphBuilder, uint32 MaxNumElements, const TCHAR* Name)
 {
@@ -941,116 +757,483 @@ void Unlock(FRHICommandList& RHICmdList, FRDGBuffer* Buffer)
 	}
 }
 
-struct FSinglePassDispatchGroup
+struct FAnimSequenceTransformProvider::FTrackUpdateRequest
 {
-	// Identity (used for TSet equality/hash)
-	uint32 NumCompactBones         = 0;
-
-	// Group parameters
-	uint32 GroupSize               = 0;
-	uint32 InstancesPerGroup       = 1;
-	uint32 SlotsPerInstance        = 0;
+	uint8 bNeedsCurrent : 1;
+	uint8 UpdateMode    : 7;
 
-	// Header counts (accumulated during counting phase)
-	uint32 NumEvaluateHeaders      = 0;
-	uint32 NumBlendEvaluateHeaders = 0;
-
-	// Offsets into shared header buffers (assigned during buffer creation phase)
-	uint32 EvaluateHeaderBaseOffset      = 0;
-	uint32 BlendEvaluateHeaderBaseOffset = 0;
-
-	// Write cursors (used during fill phase)
-	uint32 NextEvaluateWriteIndex      = 0;
-	uint32 NextBlendEvaluateWriteIndex = 0;
-
-	friend bool operator==(const FSinglePassDispatchGroup& A, const FSinglePassDispatchGroup& B)
+	void SetUpdateMode(EPreviousBoneTransformUpdateMode InUpdateMode)
 	{
-		return A.NumCompactBones == B.NumCompactBones;
+		UpdateMode = static_cast<uint8>(InUpdateMode);
 	}
 
-	friend uint32 GetTypeHash(const FSinglePassDispatchGroup& Group)
+	EPreviousBoneTransformUpdateMode GetUpdateMode() const
 	{
-		return GetTypeHash(Group.NumCompactBones);
+		return static_cast<EPreviousBoneTransformUpdateMode>(UpdateMode);
 	}
 };
 
-struct FMultiPassDispatchGroup
+struct FAnimSequenceTransformProvider::FDispatchGroupKey
 {
-	// Identity (used for TSet equality/hash)
-	uint32 NumCompactBones         = 0;
-	uint32 NumLeafBones            = 0;
-	uint32 MaxBoneTransformCount   = 0;
-
-	// Header counts (accumulated during counting phase)
-	uint32 NumEvaluateHeaders      = 0;
-	uint32 NumBlendEvaluateHeaders = 0;
-
-	// Offsets into shared header buffers (assigned during buffer creation phase)
-	uint32 EvaluateHeaderBaseOffset      = 0;
-	uint32 BlendEvaluateHeaderBaseOffset = 0;
-	uint32 ConcatenateHeaderBaseOffset   = 0;
-	uint32 ScatterHeaderBaseOffset       = 0;
-
-	// Write cursors (used during fill phase)
-	uint32 NextEvaluateWriteIndex      = 0;
-	uint32 NextBlendEvaluateWriteIndex = 0;
-	uint32 NextConcatenateWriteIndex   = 0;
-	uint32 NextScatterWriteIndex       = 0;
-
-	friend bool operator==(const FMultiPassDispatchGroup& A, const FMultiPassDispatchGroup& B)
+	uint32 NumCompactBones  = 0;
+	uint32 NumLayers        = 0;
+	uint32 MeshSpaceMask    = 0;
+	uint32 LayerModeMask    = 0;
+
+	friend bool operator==(const FDispatchGroupKey& A, const FDispatchGroupKey& B)
 	{
-		return A.NumCompactBones == B.NumCompactBones
-			&& A.NumLeafBones == B.NumLeafBones
-			&& A.MaxBoneTransformCount == B.MaxBoneTransformCount;
+		return A.NumCompactBones == B.NumCompactBones && A.NumLayers == B.NumLayers
+			&& A.MeshSpaceMask == B.MeshSpaceMask && A.LayerModeMask == B.LayerModeMask;
 	}
 
-	friend uint32 GetTypeHash(const FMultiPassDispatchGroup& Group)
+	friend uint32 GetTypeHash(const FDispatchGroupKey& Key)
 	{
-		return HashCombine(
-			Group.NumCompactBones,
-			HashCombine(Group.NumLeafBones, Group.MaxBoneTransformCount));
+		uint32 Hash = GetTypeHash(Key.NumCompactBones);
+		Hash = HashCombine(Hash, GetTypeHash(Key.NumLayers));
+		Hash = HashCombine(Hash, GetTypeHash(Key.MeshSpaceMask));
+		Hash = HashCombine(Hash, GetTypeHash(Key.LayerModeMask));
+		return Hash;
 	}
 };
 
-void FAnimSequenceTransformProvider::ProvideGPUTransforms(FSkinningTransformProvider::FProviderContext& Context)
+struct FAnimSequenceTransformProvider::FDispatchGroup
 {
-	TRACE_CPUPROFILER_EVENT_SCOPE(FAnimSequenceTransformProvider::ProvideGPUTransforms);
-	FRDGBuilder& GraphBuilder = Context.GraphBuilder;
-
-	SequenceTransformAllocator.Consolidate();
-	RetargetingDataAllocator.Consolidate();
-	LeafBonePathAllocator.Consolidate();
+	uint32 GroupSize                      = 0;
+	uint32 InstancesPerGroup              = 1;
+	uint32 SlotsPerInstance               = 0;
+	uint32 NumCompactBones                = 0;
+	uint32 NumLayers                      = 0;
+	uint32 MeshSpaceMask                  = 0;
+	uint32 LayerModeMask                  = 0;
+
+	uint32 NumHeaders                     = 0;
+	uint32 NumSamples                     = 0;
+	uint32 HeaderBaseOffset               = 0;
+	uint32 SampleBaseOffset               = 0;
+	uint32 NextHeaderWriteIndex           = 0;
+	uint32 NextSampleWriteIndex           = 0;
+};
 
-	uint32 MaxTotalActiveTracks = 0;
+void FAnimSequenceTransformProvider::UploadSequences(FRDGScatterUploader& ScatterUploader, int32& NumTotalFramesRequested)
+{
+	TRACE_CPUPROFILER_EVENT_SCOPE(FAnimSequenceTransformProvider::UploadSequences);
 
-	for (const FSkinningTransformProvider::FProviderIndirection Indirection : Context.Indirections)
+	for (const FMapping* Mapping : MappingsToUpload)
 	{
-		FInstancedSkinningSceneExtensionProxy* ExtensionProxy = static_cast<FInstancedSkinningSceneExtensionProxy*>(Context.Proxies[Indirection.Index]);
-		FProxyUserData& ProxyUserData = static_cast<FProxyUserData&>(ExtensionProxy->GetUserDataChecked());
-		MaxTotalActiveTracks += ProxyUserData.Tracks.GetNumActiveTracks();
-	}
+		TArrayView<FVector4f> DstFloat4s = ScatterUploader.Add_GetRef<FVector4f>(Mapping->MappingTransformBufferAllocation.Offset, Mapping->MappingTransformBufferAllocation.Count);
 
-	if (!MaxTotalActiveTracks)
-	{
-		return;
-	}
+		for (uint32 CompactBoneIndex = 0; CompactBoneIndex < Mapping->NumCompactBones; ++CompactBoneIndex)
+		{
+			FBoneTransformWithUniformScale InvRefPose;
+			InvRefPose.SetTransform(Mapping->InvGlobalRefPoseTransforms[CompactBoneIndex]);
+			WriteBoneTransform(DstFloat4s, CompactBoneIndex, InvRefPose);
 
-	struct FTrackUpdateRequest
+			const uint32 RetargetTransformIndex = Mapping->NumCompactBones + CompactBoneIndex * (Mapping->RetargetingTables.bRequiresRefPoseRetarget ? 2 : 1);
+			FBoneTransformWithUniformScale RetargetTransform;
+			RetargetTransform.SetTransform(Mapping->RetargetingTables.Transform[CompactBoneIndex]);
+			WriteBoneTransform(DstFloat4s, RetargetTransformIndex, RetargetTransform);
+
+			if (Mapping->RetargetingTables.bRequiresRefPoseRetarget)
+			{
+				const FQuat4f& Q0 = Mapping->RetargetingTables.RefPoseRetargetQuats[CompactBoneIndex].Get<0>();
+				const FQuat4f& Q1 = Mapping->RetargetingTables.RefPoseRetargetQuats[CompactBoneIndex].Get<1>();
+
+				FBoneTransformWithUniformScale RefPoseTransform;
+				RefPoseTransform.Rotation[0] = Q0.X;
+				RefPoseTransform.Rotation[1] = Q0.Y;
+				RefPoseTransform.Rotation[2] = Q0.Z;
+				RefPoseTransform.Rotation[3] = Q0.W;
+				RefPoseTransform.Translation[0] = Q1.X;
+				RefPoseTransform.Translation[1] = Q1.Y;
+				RefPoseTransform.Translation[2] = Q1.Z;
+				RefPoseTransform.Scale = Q1.W;
+				WriteBoneTransform(DstFloat4s, RetargetTransformIndex + 1, RefPoseTransform);
+			}
+		}
+	}
+
+	BoneTrackArray TrackToBoneIndexMap;
+	FMemMark Mark(FMemStack::Get());
+	TArray<FTransform, TMemStackAllocator<>> AnimLocalPoses;
+	uint32 NumSequenceFrameUploads = 0;
+	int32 SequenceIndex = 0;
+
+	const int32 UploadBudget = FMath::Max<int32>(0, CVarAnimSequenceTransformProviderUploadBudget.GetValueOnRenderThread() - NumTotalFramesRequested);
+	const int32 UploadBudgetDivisor = FMath::Max(1, SequencesToUpload.Num());
+	const int32 FramesPerSequence = UploadBudget / UploadBudgetDivisor;
+	int32 NumFramesLeftover       = UploadBudget % UploadBudgetDivisor;
+
+	while (SequenceIndex < SequencesToUpload.Num())
+	{
+		FSequence& Sequence = *SequencesToUpload[SequenceIndex];
+
+		int32 NumBudgetFramesToUpload = FramesPerSequence + (NumFramesLeftover > 0 ? 1 : 0);
+		NumFramesLeftover = FMath::Max(NumFramesLeftover - 1, 0);
+
+		while (Sequence.NextFrameToUpload < Sequence.NumFrames && NumBudgetFramesToUpload > 0)
+		{
+			const int32 FrameStep = 1u << Sequence.TargetSampleRateLevel;
+
+			if (Sequence.FramesUploaded[Sequence.NextFrameToUpload] || Sequence.FramesRequested[Sequence.NextFrameToUpload])
+			{
+				Sequence.NextFrameToUpload += FrameStep;
+				continue;
+			}
+
+			Sequence.FramesRequested[Sequence.NextFrameToUpload] = true;
+			Sequence.NumFramesRequested++;
+			Sequence.NextFrameToUpload += FrameStep;
+			NumBudgetFramesToUpload--;
+		}
+
+		if (Sequence.NextFrameToUpload >= Sequence.NumFrames && Sequence.TargetSampleRateLevel > 0)
+		{
+			Sequence.TargetSampleRateLevel--;
+			Sequence.CurrentSampleRateLevel--;
+			Sequence.NextFrameToUpload = 1u << Sequence.TargetSampleRateLevel;
+		}
+
+		if (Sequence.NumFramesRequested > 0)
+		{
+			const UAnimSequence* AnimSequence = Sequence.Sequence;
+			UAnimSequence::FScopedCompressedAnimSequence CompressedAnimSequence = AnimSequence->GetCompressedData();
+			const FCompressedAnimSequence& PlatformCompressedData = CompressedAnimSequence.Get();
+
+			const TArray<FTransform>& AnimRefLocalPoses = AnimSequence->GetSkeleton()->GetRefLocalPoses();
+			check(Sequence.NumAnimBones <= (uint16)AnimRefLocalPoses.Num());
+
+			const int32 NumTracks = PlatformCompressedData.CompressedTrackToSkeletonMapTable.Num();
+			TrackToBoneIndexMap.Reset(NumTracks);
+
+			for (int32 TrackIndex = 0; TrackIndex < NumTracks; ++TrackIndex)
+			{
+				const uint16 BoneIndex = PlatformCompressedData.GetSkeletonIndexFromTrackIndex(TrackIndex);
+
+				if (BoneIndex < Sequence.NumAnimBones)
+				{
+					TrackToBoneIndexMap.Add(BoneTrackPair(BoneIndex, TrackIndex));
+				}
+			}
+
+			FAnimSequenceDecompressionContext DecompressionContext(
+				  Sequence.SampleRate
+				, Sequence.NumFrames
+				, AnimSequence->Interpolation
+				, AnimSequence->GetRetargetTransformsSourceName()
+				, *PlatformCompressedData.CompressedDataStructure
+				, AnimRefLocalPoses
+				, PlatformCompressedData.CompressedTrackToSkeletonMapTable
+				, AnimSequence->GetSkeleton()
+				, AnimSequence->IsValidAdditive()
+				, AnimSequence->GetAdditiveAnimType()
+			);
+
+			AnimLocalPoses.Reset(AnimRefLocalPoses.Num());
+			AnimLocalPoses.Append(AnimRefLocalPoses);
+
+			for (TConstSetBitIterator<> BitIt(Sequence.FramesRequested); BitIt; ++BitIt)
+			{
+				TRACE_CPUPROFILER_EVENT_SCOPE(FAnimSequenceTransformProvider::DecompressPose);
+				const int32 FrameIndex = BitIt.GetIndex();
+
+				const double SeekTime = Sequence.SampleRate.AsSeconds(FFrameTime(FrameIndex));
+				DecompressionContext.Seek(SeekTime);
+
+				if (PlatformCompressedData.BoneCompressionCodec)
+				{
+					TArrayView<FTransform> AnimLocalPosesView(AnimLocalPoses);
+					PlatformCompressedData.BoneCompressionCodec->DecompressPose(DecompressionContext, TrackToBoneIndexMap, TrackToBoneIndexMap, TrackToBoneIndexMap, AnimLocalPosesView);
+				}
+
+				const uint32 Float4sPerBone = GetFloat4sPerBone();
+				const uint32 FrameFloat4Offset = Sequence.BufferAllocation.Offset + Sequence.NumAnimBones * FrameIndex * Float4sPerBone;
+				TArrayView<FVector4f> DstFloat4s = ScatterUploader.Add_GetRef<FVector4f>(FrameFloat4Offset, Sequence.NumAnimBones * Float4sPerBone);
+
+				for (uint16 BoneIndex = 0; BoneIndex < Sequence.NumAnimBones; ++BoneIndex)
+				{
+					FBoneTransformWithUniformScale Transform;
+					Transform.SetTransform(FTransform3f(AnimLocalPoses[BoneIndex]));
+					WriteBoneTransform(DstFloat4s, BoneIndex, Transform);
+				}
+
+				Sequence.NumFramesUploaded++;
+				Sequence.FramesUploaded[FrameIndex] = true;
+				Sequence.FramesRequested[FrameIndex] = false;
+				NumSequenceFrameUploads++;
+			}
+
+			Sequence.NumFramesRequested = 0;
+		}
+
+		if (Sequence.NumFramesUploaded == Sequence.NumFrames)
+		{
+			SequencesToUpload.RemoveAtSwap(SequenceIndex, EAllowShrinking::No);
+		}
+		else
+		{
+			SequenceIndex++;
+		}
+	}
+
+	INC_DWORD_STAT_BY(STAT_AnimSequenceTransformProvider_SequenceFrameUploads, NumSequenceFrameUploads);
+}
+
+
+void FAnimSequenceTransformProvider::UploadBlockHeaders(
+	FRHICommandList& RHICmdList,
+	double WorldTime,
+	float DeltaTime,
+	TConstArrayView<FSkinningTransformProvider::FProviderIndirection> Indirections,
+	TConstArrayView<FSkinningSceneExtensionProxy*> Proxies,
+	int32& NumTotalFramesRequested,
+	TConstArrayView<FTrackUpdateRequest> TrackUpdateRequests,
+	FRDGBuffer* LayerHeaderBuffer,
+	FRDGBuffer* SamplePoseBuffer,
+	uint32 NumTotalEntries,
+	uint32 NumTotalSamples,
+	TMap<FDispatchGroupKey, FDispatchGroup>& DispatchGroups)
+{
+	TRACE_CPUPROFILER_EVENT_SCOPE(FAnimSequenceTransformProvider::UploadBlockHeaders);
+
+	UE::HLSL::FEvaluateLayerHeader* HeaderUpload = reinterpret_cast<UE::HLSL::FEvaluateLayerHeader*>(
+		RHICmdList.LockBuffer(LayerHeaderBuffer->GetRHIUnchecked(), 0, NumTotalEntries * sizeof(UE::HLSL::FEvaluateLayerHeader), RLM_WriteOnly));
+
+	UE::HLSL::FSamplePoseData* SampleUpload = reinterpret_cast<UE::HLSL::FSamplePoseData*>(
+		RHICmdList.LockBuffer(SamplePoseBuffer->GetRHIUnchecked(), 0, NumTotalSamples * sizeof(UE::HLSL::FSamplePoseData), RLM_WriteOnly));
+
+	uint32 NumActiveTracks = 0;
+	uint32 NumCurrentTransformUpdates = 0;
+	uint32 NumPreviousTransformUpdates = 0;
+
+	struct FWriteSamplesResult
 	{
-		uint8 bNeedsCurrent : 1;
-		uint8 UpdateMode    : 7;
+		uint32 NumSamples = 0;
+		uint32 InvGlobalRefPoseTransformBufferOffset = 0;
+	};
+
+	const auto WriteSamples = [&](
+		const FProxyUserData& ProxyUserData,
+		const FAnimSequenceTrackPackedData& PackedData,
+		const FAnimSequenceTrackRenderData& RenderData,
+		uint32 LODIndex,
+		FDispatchGroup& Group,
+		bool bPrevious) -> FWriteSamplesResult
+	{
+		FWriteSamplesResult Result;
+
+		const double TimePosition = bPrevious ? RenderData.GetPreviousTimePosition() : PackedData.GetTimePosition(WorldTime);
 
-		void SetUpdateMode(EPreviousBoneTransformUpdateMode InUpdateMode)
+		if (PackedData.IsBlendSpace())
 		{
-			UpdateMode = static_cast<uint8>(InUpdateMode);
+			Result.NumSamples = PackedData.GetNumSamples();
+
+			if (Result.NumSamples == 0)
+			{
+				return Result;
+			}
+
+			for (uint32 SampleIndex = 0; SampleIndex < Result.NumSamples; SampleIndex++)
+			{
+				const uint16 SequenceIndex = PackedData.GetSampleSequenceIndex(SampleIndex);
+				const FTrackData TrackData = ProxyUserData.GetTrackDataForSequence(SequenceIndex, TimePosition, LODIndex, PackedData.GetLoopMode());
+
+				int32 KeyIndex0, KeyIndex1;
+				float Alpha;
+				TrackData.Sequence.GetKeyIndicesFromTime(KeyIndex0, KeyIndex1, Alpha, TrackData.Position);
+				NumTotalFramesRequested += TrackData.Sequence.RequestFrame(KeyIndex0);
+				NumTotalFramesRequested += TrackData.Sequence.RequestFrame(KeyIndex1);
+
+				UE::HLSL::FSamplePoseData& Sample = SampleUpload[Group.SampleBaseOffset + Group.NextSampleWriteIndex++];
+				Sample.Alpha                            = Alpha;
+				Sample.KeyIndex0                        = KeyIndex0;
+				Sample.KeyIndex1                        = KeyIndex1;
+				Sample.SequenceTransformBufferOffset    = TrackData.Sequence.BufferAllocation.Offset;
+				Sample.RetargetingTransformBufferOffset = TrackData.Mapping.RetargetingTransformBufferOffset;
+				Sample.RetargetingDataBufferOffset      = TrackData.Mapping.RetargetingDataBufferAllocation.Offset * sizeof(UE::HLSL::FBoneRetargetingData);
+				Sample.NumAnimBones                     = TrackData.Sequence.NumAnimBones;
+				Sample.Weight                           = PackedData.GetSampleWeight(SampleIndex).GetFloat();
+
+				if (SampleIndex == 0)
+				{
+					Result.InvGlobalRefPoseTransformBufferOffset = TrackData.Mapping.InvGlobalRefPoseTransformBufferOffset;
+				}
+			}
 		}
+		else
+		{
+			const uint16 SequenceIndex = static_cast<uint16>(PackedData.GetSequenceIndex());
+			const FTrackData TrackData = ProxyUserData.GetTrackDataForSequence(SequenceIndex, TimePosition, LODIndex, PackedData.GetLoopMode());
+
+			int32 KeyIndex0, KeyIndex1;
+			float Alpha;
+			TrackData.Sequence.GetKeyIndicesFromTime(KeyIndex0, KeyIndex1, Alpha, TrackData.Position);
+			NumTotalFramesRequested += TrackData.Sequence.RequestFrame(KeyIndex0);
+			NumTotalFramesRequested += TrackData.Sequence.RequestFrame(KeyIndex1);
+
+			UE::HLSL::FSamplePoseData& Sample = SampleUpload[Group.SampleBaseOffset + Group.NextSampleWriteIndex++];
+			Sample.Alpha                            = Alpha;
+			Sample.KeyIndex0                        = KeyIndex0;
+			Sample.KeyIndex1                        = KeyIndex1;
+			Sample.SequenceTransformBufferOffset    = TrackData.Sequence.BufferAllocation.Offset;
+			Sample.RetargetingTransformBufferOffset = TrackData.Mapping.RetargetingTransformBufferOffset;
+			Sample.RetargetingDataBufferOffset      = TrackData.Mapping.RetargetingDataBufferAllocation.Offset * sizeof(UE::HLSL::FBoneRetargetingData);
+			Sample.NumAnimBones                     = TrackData.Sequence.NumAnimBones;
+			Sample.Weight                           = 1.0f;
+
+			Result.NumSamples = 1;
+			Result.InvGlobalRefPoseTransformBufferOffset = TrackData.Mapping.InvGlobalRefPoseTransformBufferOffset;
+		}
+
+		return Result;
+	};
 
-		EPreviousBoneTransformUpdateMode GetUpdateMode() const
+	const auto WriteTrackLayers = [&](
+		const FProxyUserData& ProxyUserData,
+		const FProxyUserData::FLOD& LOD,
+		FAnimSequenceTransformProviderRenderData& RenderData,
+		TConstArrayView<uint32> BoneMaskOffsets,
+		FDispatchGroup& Group,
+		int32 LODIndex,
+		int32 TrackIndex,
+		uint32 DstOffset1,
+		uint32 DstOffset2,
+		bool bPrevious,
+		bool bTick)
+	{
+		const uint32 BoneMapOffset = LOD.BoneMapBufferOffset * sizeof(uint32);
+
+		for (int32 LayerIndex = 0; LayerIndex < RenderData.GetNumLayers(); LayerIndex++)
 		{
-			return static_cast<EPreviousBoneTransformUpdateMode>(UpdateMode);
+			const uint32 CurrentHeaderIndex = Group.HeaderBaseOffset + Group.NextHeaderWriteIndex++;
+			UE::HLSL::FEvaluateLayerHeader& Header = HeaderUpload[CurrentHeaderIndex];
+
+			Header.BoneMapBufferOffset       = BoneMapOffset;
+			Header.DstTransformBufferOffset1 = DstOffset1;
+			Header.DstTransformBufferOffset2 = DstOffset2;
+
+			if (RenderData.IsTrackActive(TrackIndex, LayerIndex))
+			{
+				FAnimSequenceTransformProviderRenderLayer& Layer = RenderData.GetLayers()[LayerIndex];
+				FAnimSequenceTrackRenderData& Data = Layer.Tracks.GetData(TrackIndex);
+
+				const uint32 SampleOffset = Group.SampleBaseOffset + Group.NextSampleWriteIndex;
+				const float BlendWeight = Data.GetBlendWeight(bPrevious);
+
+				const FWriteSamplesResult SourceResult = WriteSamples(ProxyUserData, Data.GetSource(), Data, LODIndex, Group, bPrevious);
+				FWriteSamplesResult TargetResult;
+
+				if (BlendWeight > 0.0f)
+				{
+					TargetResult = WriteSamples(ProxyUserData, Data.GetTarget(), Data, LODIndex, Group, bPrevious);
+				}
+
+				Header.SampleDataOffset = SampleOffset;
+				Header.BlendWeight      = BlendWeight;
+				Header.LayerWeight      = Layer.Weight * Data.GetTarget().GetLayerWeight();
+				Header.InvGlobalRefPoseTransformBufferOffset = SourceResult.InvGlobalRefPoseTransformBufferOffset;
+				SetBoneMaskAndSampleCounts(Header, BoneMaskOffsets[LayerIndex], SourceResult.NumSamples, TargetResult.NumSamples);
+
+				if (bTick)
+				{
+					Data.Tick(DeltaTime, WorldTime);
+				}
+			}
+			else
+			{
+				Header.SampleDataOffset                      = 0;
+				Header.BlendWeight                           = 0.0f;
+				Header.LayerWeight                           = 0.0f;
+				Header.InvGlobalRefPoseTransformBufferOffset = 0;
+				SetBoneMaskAndSampleCounts(Header, SKINNING_INVALID_BONE_MASK_OFFSET, 0, 0);
+			}
 		}
 	};
-	
+
+	uint32 UpdateResultIndex = 0;
+
+	for (const FSkinningTransformProvider::FProviderIndirection Indirection : Indirections)
+	{
+		FInstancedSkinningSceneExtensionProxy* ExtensionProxy = static_cast<FInstancedSkinningSceneExtensionProxy*>(Proxies[Indirection.Index]);
+		FProxyUserData& ProxyUserData = static_cast<FProxyUserData&>(ExtensionProxy->GetUserDataChecked());
+		FAnimSequenceTransformProviderRenderData& RenderData = *ProxyUserData.GetRenderData();
+		const uint32 LODIndex = ExtensionProxy->GetLOD();
+		const uint32 NumTransforms = ExtensionProxy->GetMaxBoneTransformCount();
+		FProxyUserData::FLOD& LOD = ProxyUserData.LODs[LODIndex];
+		FAnimSequenceTrackRenderPool& Tracks = RenderData.GetLayers()[0].Tracks;
+		TConstArrayView<uint32> BoneMaskOffsets = LOD.BoneMaskBufferOffsets;
+
+		FDispatchGroupKey GroupKey;
+		GroupKey.NumCompactBones = LOD.BoneTables.NumCompactBones;
+		GroupKey.NumLayers       = RenderData.GetNumLayers();
+		GroupKey.MeshSpaceMask   = RenderData.GetMeshSpaceMask();
+		GroupKey.LayerModeMask   = RenderData.GetLayerModeMask();
+		FDispatchGroup& DispatchGroup = DispatchGroups.FindChecked(GroupKey);
+
+		const uint32 CurrentTransformOffset  = Indirection.CurrentTransformOffset;
+		const uint32 PreviousTransformOffset = Indirection.PreviousTransformOffset;
+
+		check(Tracks.GetNumActiveTracks() <= (int32)ExtensionProxy->GetUniqueAnimationCount());
+
+		Tracks.EnumerateActiveTracks([&] (int32 TrackIndex)
+		{
+			const FTrackUpdateRequest UpdateRequest = TrackUpdateRequests[UpdateResultIndex++];
+			const uint32 TrackTransformOffset         = TrackIndex * NumTransforms * 2 * sizeof(FCompressedBoneTransform);
+			const uint32 TrackCurrentTransformOffset  = TrackTransformOffset + CurrentTransformOffset;
+			const uint32 TrackPreviousTransformOffset = TrackTransformOffset + PreviousTransformOffset;
+			const bool bNeedsPrevious                 = UpdateRequest.GetUpdateMode() == EPreviousBoneTransformUpdateMode::UpdatePrevious;
+
+			if (UpdateRequest.bNeedsCurrent)
+			{
+				const bool bTick = !bNeedsPrevious;
+				uint32 OutputTransformOffset1 = TrackCurrentTransformOffset;
+				uint32 OutputTransformOffset2 = UpdateRequest.GetUpdateMode() == EPreviousBoneTransformUpdateMode::DuplicateCurrentToPrevious ? TrackPreviousTransformOffset : TrackCurrentTransformOffset;
+				WriteTrackLayers(ProxyUserData, LOD, RenderData, BoneMaskOffsets, DispatchGroup, LODIndex, TrackIndex, OutputTransformOffset1, OutputTransformOffset2, false, bTick);
+				NumCurrentTransformUpdates++;
+			}
+
+			if (bNeedsPrevious)
+			{
+				const bool bTick = true;
+				WriteTrackLayers(ProxyUserData, LOD, RenderData, BoneMaskOffsets, DispatchGroup, LODIndex, TrackIndex, TrackPreviousTransformOffset, TrackPreviousTransformOffset, true, bTick);
+				NumPreviousTransformUpdates++;
+			}
+			NumActiveTracks++;
+		});
+	}
+
+	RHICmdList.UnlockBuffer(LayerHeaderBuffer->GetRHIUnchecked());
+	RHICmdList.UnlockBuffer(SamplePoseBuffer->GetRHIUnchecked());
+
+	INC_DWORD_STAT_BY(STAT_AnimSequenceTransformProvider_ActiveTracks, NumActiveTracks);
+	INC_DWORD_STAT_BY(STAT_AnimSequenceTransformProvider_CurrentTransformUpdates, NumCurrentTransformUpdates);
+	INC_DWORD_STAT_BY(STAT_AnimSequenceTransformProvider_PreviousTransformUpdates, NumPreviousTransformUpdates);
+}
+
+void FAnimSequenceTransformProvider::ProvideTransforms(FSkinningTransformProvider::FProviderContext& Context)
+{
+	TRACE_CPUPROFILER_EVENT_SCOPE(FAnimSequenceTransformProvider::ProvideTransforms);
+	FRDGBuilder& GraphBuilder = Context.GraphBuilder;
+
+	SequenceTransformAllocator.Consolidate();
+	RetargetingDataAllocator.Consolidate();
+	BoneHierarchyAllocator.Consolidate();
+	BoneMaskAllocator.Consolidate();
+
+	uint32 MaxTotalActiveTracks = 0;
+
+	for (const FSkinningTransformProvider::FProviderIndirection Indirection : Context.Indirections)
+	{
+		FInstancedSkinningSceneExtensionProxy* ExtensionProxy = static_cast<FInstancedSkinningSceneExtensionProxy*>(Context.Proxies[Indirection.Index]);
+		FProxyUserData& ProxyUserData = static_cast<FProxyUserData&>(ExtensionProxy->GetUserDataChecked());
+		FAnimSequenceTransformProviderRenderData& RenderData = *ProxyUserData.GetRenderData();
+		MaxTotalActiveTracks += RenderData.GetLayers()[0].Tracks.GetNumActiveTracks();
+	}
+
+	if (!MaxTotalActiveTracks)
+	{
+		return;
+	}
+
 	TArray<FTrackUpdateRequest, SceneRenderingAllocator> TrackUpdateRequests;
 	TrackUpdateRequests.Reserve(MaxTotalActiveTracks);
 
@@ -1067,10 +1250,8 @@ void FAnimSequenceTransformProvider::ProvideGPUTransforms(FSkinningTransformProv
 	const float DeltaTime = Context.CurrentTime.GetDeltaWorldTimeSeconds() * GlobalTimeScale;
 
 	uint32 NumTotalSubTracks = 0;
-	uint32 NumTotalIntermediateTransforms = 0;
 
-	auto& SinglePassDispatchGroups = *GraphBuilder.AllocObject<TSet<FSinglePassDispatchGroup>>();
-	auto& MultiPassDispatchGroups  = *GraphBuilder.AllocObject<TSet<FMultiPassDispatchGroup>>();
+	auto& DispatchGroups = *GraphBuilder.AllocObject<TMap<FDispatchGroupKey, FDispatchGroup>>();
 
 	for (const FSkinningTransformProvider::FProviderIndirection Indirection : Context.Indirections)
 	{
@@ -1078,76 +1259,68 @@ void FAnimSequenceTransformProvider::ProvideGPUTransforms(FSkinningTransformProv
 		FProxyUserData& ProxyUserData = static_cast<FProxyUserData&>(ExtensionProxy->GetUserDataChecked());
 		const uint32 LODIndex = ExtensionProxy->GetLOD();
 		FProxyUserData::FLOD& LOD = ProxyUserData.LODs[LODIndex];
-		FAnimSequenceTrackRenderPool& Tracks = ProxyUserData.Tracks;
-		ProxyUserData.TrackHistories.SetNum(Tracks.GetNumTracks());
+		FAnimSequenceTransformProviderRenderData& RenderData = *ProxyUserData.GetRenderData();
+		FAnimSequenceTrackRenderPool& Tracks = RenderData.GetLayers()[0].Tracks;
 
 		const uint32 CurrentTransformSlot = Indirection.CurrentTransformOffset < Indirection.PreviousTransformOffset ? 0 : 1;
-
-		// Resolve group once per proxy, not per track
-		uint32* NumDispatchGroupEvaluateHeaders = nullptr;
-		uint32* NumDispatchGroupBlendEvaluateHeaders = nullptr;
-
-		if (LOD.bSinglePass)
-		{
-			FSinglePassDispatchGroup LookupKey;
-			LookupKey.NumCompactBones   = LOD.BoneTables.NumCompactBones;
-			LookupKey.GroupSize         = LOD.GroupSize;
-			LookupKey.InstancesPerGroup = LOD.InstancesPerGroup;
-			LookupKey.SlotsPerInstance  = LOD.SlotsPerInstance;
-
-			FSinglePassDispatchGroup& DispatchGroup = SinglePassDispatchGroups.FindOrAdd(LookupKey);
-			NumDispatchGroupEvaluateHeaders      = &DispatchGroup.NumEvaluateHeaders;
-			NumDispatchGroupBlendEvaluateHeaders = &DispatchGroup.NumBlendEvaluateHeaders;
-		}
-		else
+		const uint32 NumLayers = RenderData.GetNumLayers();
+		const uint32 MeshSpaceMask = RenderData.GetMeshSpaceMask();
+		const uint32 LayerModeMask = RenderData.GetLayerModeMask();
+
+		FDispatchGroupKey GroupKey;
+		GroupKey.NumCompactBones = LOD.BoneTables.NumCompactBones;
+		GroupKey.NumLayers       = NumLayers;
+		GroupKey.MeshSpaceMask   = MeshSpaceMask;
+		GroupKey.LayerModeMask   = LayerModeMask;
+
+		FDispatchGroup& DispatchGroup = DispatchGroups.FindOrAdd(GroupKey);
+		if (DispatchGroup.GroupSize == 0)
 		{
-			FMultiPassDispatchGroup LookupKey;
-			LookupKey.NumCompactBones       = LOD.BoneTables.NumCompactBones;
-			LookupKey.NumLeafBones          = LOD.LeafBonePathData.NumLeafBones;
-			LookupKey.MaxBoneTransformCount = ExtensionProxy->GetMaxBoneTransformCount();
-
-			FMultiPassDispatchGroup& DispatchGroup = MultiPassDispatchGroups.FindOrAdd(LookupKey);
-			NumDispatchGroupEvaluateHeaders      = &DispatchGroup.NumEvaluateHeaders;
-			NumDispatchGroupBlendEvaluateHeaders = &DispatchGroup.NumBlendEvaluateHeaders;
+			DispatchGroup.NumCompactBones   = LOD.BoneTables.NumCompactBones;
+			DispatchGroup.NumLayers         = NumLayers;
+			DispatchGroup.MeshSpaceMask     = MeshSpaceMask;
+			DispatchGroup.LayerModeMask     = LayerModeMask;
+			DispatchGroup.GroupSize         = LOD.GroupSize;
+			DispatchGroup.InstancesPerGroup = LOD.InstancesPerGroup;
+			DispatchGroup.SlotsPerInstance  = LOD.SlotsPerInstance;
 		}
 
 		uint32 NumSubTracks = 0;
 
 		Tracks.EnumerateActiveTracks([&] (int32 TrackIndex)
 		{
-			const FAnimSequenceTrackRenderData& Data = Tracks.GetData(TrackIndex);
-			const FTrackData TargetTrack             = ProxyUserData.GetTrackData(Data.GetTarget(), WorldTime, LODIndex);
-			const float BlendWeight                  = Data.GetBlendWeight();
-
-			// Dirty previous transforms means the whole allocation is dirty. Rebuild everything.
-			const bool bForceUpdateAll = EnumHasAnyFlags(Indirection.DirtyBoneTransforms, EDirtyBoneTransforms::Previous);
-			FTrackHistory::FUpdateResult UpdateResult = ProxyUserData.TrackHistories[TrackIndex].Update(TargetTrack, CurrentTransformSlot, bForceUpdateAll);
-
 			FTrackUpdateRequest UpdateRequest;
 			UpdateRequest.SetUpdateMode(EPreviousBoneTransformUpdateMode::None);
-			UpdateRequest.bNeedsCurrent = UpdateResult.bNeedsCurrent;
+			const bool bForceUpdateAll = EnumHasAnyFlags(Indirection.DirtyBoneTransforms, EDirtyBoneTransforms::Previous);
 
-			if (UpdateResult.bNeedsPrevious)
+			UpdateRequest.bNeedsCurrent = true;
+
+			if (bForceUpdateAll)
 			{
-				UpdateRequest.UpdateMode = uint8(UpdateResult.bNeedsCurrent && UpdateResult.bNeedsPrevious && TargetTrack.Position == TargetTrack.PreviousPosition
-					? EPreviousBoneTransformUpdateMode::DuplicateCurrentToPrevious
-					: EPreviousBoneTransformUpdateMode::UpdatePrevious);
+				UpdateRequest.SetUpdateMode(EPreviousBoneTransformUpdateMode::UpdatePrevious);
 			}
 
-			const int32 NumLocalSubTracks = UpdateResult.bNeedsCurrent + (UpdateRequest.GetUpdateMode() == EPreviousBoneTransformUpdateMode::UpdatePrevious ? 1 : 0);
+			for (int32 LayerIndex = 0; LayerIndex < (int32)NumLayers; LayerIndex++)
+			{
+				if (RenderData.IsTrackActive(TrackIndex, LayerIndex))
+				{
+					FAnimSequenceTransformProviderRenderLayer& Layer = RenderData.GetLayers()[LayerIndex];
+					FAnimSequenceTrackRenderData& Data = Layer.Tracks.GetData(TrackIndex);
 
-			NumSubTracks                   += NumLocalSubTracks;
-			NumTotalIntermediateTransforms += LOD.bSinglePass ? 0 : NumLocalSubTracks * LOD.BoneTables.NumCompactBones;
+					if (!Data.HasPreviousPosition())
+					{
+						UpdateRequest.SetUpdateMode(EPreviousBoneTransformUpdateMode::DuplicateCurrentToPrevious);
+					}
 
-			if (BlendWeight > 0.0f)
-			{
-				*NumDispatchGroupBlendEvaluateHeaders += NumLocalSubTracks;
-			}
-			else
-			{
-				*NumDispatchGroupEvaluateHeaders += NumLocalSubTracks;
+					DispatchGroup.NumSamples += Data.GetTotalSampleCount(UpdateRequest.GetUpdateMode());
+				}
 			}
 
+			const int32 NumLocalSubTracks = UpdateRequest.bNeedsCurrent + (UpdateRequest.GetUpdateMode() == EPreviousBoneTransformUpdateMode::UpdatePrevious ? 1 : 0);
+
+			NumSubTracks += NumLocalSubTracks;
+			DispatchGroup.NumHeaders += NumLocalSubTracks;
+
 			TrackUpdateRequests.Emplace(UpdateRequest);
 		});
 
@@ -1177,7 +1350,7 @@ void FAnimSequenceTransformProvider::ProvideGPUTransforms(FSkinningTransformProv
 
 	uint32 NumTotalSequenceTransforms = 0;
 	uint32 NumTotalRetargetingDatas = 0;
-	uint32 NumTotalLeafBonePathElements = 0;
+	uint32 NumTotalBoneHierarchyElements = 0;
 
 	for (FMapping* Mapping : MappingsToUpload)
 	{
@@ -1186,16 +1359,16 @@ void FAnimSequenceTransformProvider::ProvideGPUTransforms(FSkinningTransformProv
 		const uint32 NumTransforms            = Mapping->NumCompactBones * NumRetargetingTransforms + NumInvGlobalRefPoses;
 		const uint32 NumRetargetingDatas      = Mapping->NumCompactBones;
 
-		Mapping->MappingTransformBufferAllocation.Count  = NumTransforms;
-		Mapping->MappingTransformBufferAllocation.Offset = SequenceTransformAllocator.Allocate(NumTransforms);
+		Mapping->MappingTransformBufferAllocation.Count  = NumTransforms * GetFloat4sPerBone();
+		Mapping->MappingTransformBufferAllocation.Offset = SequenceTransformAllocator.Allocate(NumTransforms * GetFloat4sPerBone());
 
 		Mapping->InvGlobalRefPoseTransformBufferOffset = Mapping->MappingTransformBufferAllocation.Offset;
-		Mapping->RetargetingTransformBufferOffset      = Mapping->MappingTransformBufferAllocation.Offset + NumInvGlobalRefPoses;
+		Mapping->RetargetingTransformBufferOffset      = Mapping->MappingTransformBufferAllocation.Offset + NumInvGlobalRefPoses * GetFloat4sPerBone();
 
 		Mapping->RetargetingDataBufferAllocation.Count = NumRetargetingDatas;
 		Mapping->RetargetingDataBufferAllocation.Offset = RetargetingDataAllocator.Allocate(NumRetargetingDatas);
 
-		NumTotalSequenceTransforms += NumTransforms;
+		NumTotalSequenceTransforms += NumTransforms * GetFloat4sPerBone();
 		NumTotalRetargetingDatas  += NumRetargetingDatas;
 	}
 
@@ -1207,75 +1380,60 @@ void FAnimSequenceTransformProvider::ProvideGPUTransforms(FSkinningTransformProv
 		});
 	}
 
+	if (!ProxyUserDatasToUpload.IsEmpty())
+	{
+		GraphBuilder.AddPostExecuteCallback([this]
+		{
+			ProxyUserDatasToUpload.Empty();
+		});
+	}
+
 	for (FSequence* Sequence : SequencesToUpload)
 	{
-		const int32 NumTransforms = Sequence->NumAnimBones * Sequence->NumFrames;
+		const int32 NumFloat4s = Sequence->NumAnimBones * Sequence->NumFrames * GetFloat4sPerBone();
 
 		if (!Sequence->BufferAllocation.IsValid())
 		{
-			Sequence->BufferAllocation.Count  = NumTransforms;
-			Sequence->BufferAllocation.Offset = SequenceTransformAllocator.Allocate(NumTransforms);
+			Sequence->BufferAllocation.Count  = NumFloat4s;
+			Sequence->BufferAllocation.Offset = SequenceTransformAllocator.Allocate(NumFloat4s);
 		}
 
-		NumTotalSequenceTransforms += NumTransforms;
+		NumTotalSequenceTransforms += NumFloat4s;
 	}
 
 	for (const FLODUploadRequest& Request : LODUploadRequests)
 	{
 		FProxyUserData::FLOD& LOD = Request.LOD;
-		const int32 NumLeafBonePaths = LOD.LeafBonePathData.Buffer.Num();
-
-		// When in single pass mode the bone map is in compact bone space, otherwise it's in transform space.
-		const uint32 NumBoneMapEntries = LOD.bSinglePass ? LOD.BoneTables.NumCompactBones : Request.ExtensionProxy.GetMaxBoneTransformCount();
-
-		const uint32 NumElements = NumLeafBonePaths + NumBoneMapEntries;
-		LOD.LeafBonePathBufferAllocation.Count = NumElements;
-		LOD.LeafBonePathBufferAllocation.Offset = LeafBonePathAllocator.Allocate(NumElements);
-
-		LOD.BoneMapBufferOffset = LOD.LeafBonePathBufferAllocation.Offset + NumLeafBonePaths;
-
-		NumTotalLeafBonePathElements += NumElements;
+		const uint32 NumBoneMapEntries = LOD.BoneTables.NumCompactBones;
+		LOD.BoneMapBufferAllocation.Count = NumBoneMapEntries;
+		LOD.BoneMapBufferAllocation.Offset = BoneHierarchyAllocator.Allocate(NumBoneMapEntries);
+		LOD.BoneMapBufferOffset = LOD.BoneMapBufferAllocation.Offset;
+		NumTotalBoneHierarchyElements += NumBoneMapEntries;
 	}
 
 	int32& NumTotalFramesRequested = *GraphBuilder.AllocPOD<int32>();
 	NumTotalFramesRequested = 0;
 
-	// Accumulate totals and assign per-group base offsets into shared buffers
-	uint32 NumTotalSinglePassEvaluateHeaders      = 0;
-	uint32 NumTotalSinglePassBlendEvaluateHeaders = 0;
-	uint32 NumTotalMultiPassEvaluateHeaders       = 0;
-	uint32 NumTotalMultiPassBlendEvaluateHeaders  = 0;
-	uint32 NumTotalConcatenateHeaders             = 0;
-	uint32 NumTotalScatterHeaders                 = 0;
+	uint32 NumTotalEntries = 0;
+	uint32 NumTotalSamples = 0;
 
-	for (FSinglePassDispatchGroup& Group : SinglePassDispatchGroups)
+	for (auto& [Key, Group] : DispatchGroups)
 	{
-		Group.EvaluateHeaderBaseOffset          = NumTotalSinglePassEvaluateHeaders;
-		Group.BlendEvaluateHeaderBaseOffset     = NumTotalSinglePassBlendEvaluateHeaders;
-		NumTotalSinglePassEvaluateHeaders      += Group.NumEvaluateHeaders;
-		NumTotalSinglePassBlendEvaluateHeaders += Group.NumBlendEvaluateHeaders;
+		Group.HeaderBaseOffset = NumTotalEntries;
+		Group.SampleBaseOffset = NumTotalSamples;
+		NumTotalEntries += Group.NumHeaders * Group.NumLayers;
+		NumTotalSamples += Group.NumSamples;
 	}
 
-	for (FMultiPassDispatchGroup& Group : MultiPassDispatchGroups)
-	{
-		const uint32 NumGroupTotalHeaders = Group.NumEvaluateHeaders + Group.NumBlendEvaluateHeaders;
-		Group.EvaluateHeaderBaseOffset         = NumTotalMultiPassEvaluateHeaders;
-		Group.BlendEvaluateHeaderBaseOffset    = NumTotalMultiPassBlendEvaluateHeaders;
-		Group.ConcatenateHeaderBaseOffset      = NumTotalConcatenateHeaders;
-		Group.ScatterHeaderBaseOffset          = NumTotalScatterHeaders;
-		NumTotalMultiPassEvaluateHeaders      += Group.NumEvaluateHeaders;
-		NumTotalMultiPassBlendEvaluateHeaders += Group.NumBlendEvaluateHeaders;
-		NumTotalConcatenateHeaders            += NumGroupTotalHeaders;
-		NumTotalScatterHeaders                += NumGroupTotalHeaders;
-	}
+	FRDGBuffer* LayerHeaderBuffer = GraphBuilder.CreateBuffer(
+		FRDGBufferDesc::CreateStructuredDesc(sizeof(UE::HLSL::FEvaluateLayerHeader), NumTotalEntries),
+		TEXT("AnimSequenceTransformProvider.LayerHeaders"));
+	GraphBuilder.ConvertToExternalBuffer(LayerHeaderBuffer);
 
-	// Create shared buffers and lock
-	FRDGBuffer* SinglePassEvaluateHeaderBuffer      = CreateByteAddressBuffer<UE::HLSL::FEvaluateConcatenateScatterBlockHeader>(GraphBuilder, NumTotalSinglePassEvaluateHeaders, TEXT("AnimSequenceTransformProvider.SinglePassEvaluateHeaders"));
-	FRDGBuffer* SinglePassBlendEvaluateHeaderBuffer = CreateByteAddressBuffer<UE::HLSL::FEvaluateBlendConcatenateScatterBlockHeader>(GraphBuilder, NumTotalSinglePassBlendEvaluateHeaders, TEXT("AnimSequenceTransformProvider.SinglePassBlendEvaluateHeaders"));
-	FRDGBuffer* MultiPassEvaluateHeaderBuffer       = CreateByteAddressBuffer<UE::HLSL::FEvaluateAnimSequenceBlockHeader>(GraphBuilder, NumTotalMultiPassEvaluateHeaders, TEXT("AnimSequenceTransformProvider.EvaluateHeaders"));
-	FRDGBuffer* MultiPassBlendEvaluateHeaderBuffer  = CreateByteAddressBuffer<UE::HLSL::FEvaluateAnimSequenceBlendBlockHeader>(GraphBuilder, NumTotalMultiPassBlendEvaluateHeaders, TEXT("AnimSequenceTransformProvider.BlendEvaluateHeaders"));
-	FRDGBuffer* ConcatenateHeaderBuffer             = CreateByteAddressBuffer<UE::HLSL::FConcatenateBlockHeader>(GraphBuilder, NumTotalConcatenateHeaders, TEXT("AnimSequenceTransformProvider.ConcatenateHeaders"));
-	FRDGBuffer* ScatterHeaderBuffer                 = CreateByteAddressBuffer<UE::HLSL::FScatterTransformsBlockHeader>(GraphBuilder, NumTotalScatterHeaders, TEXT("AnimSequenceTransformProvider.ScatterHeaders"));
+	FRDGBuffer* SamplePoseBuffer = GraphBuilder.CreateBuffer(
+		FRDGBufferDesc::CreateStructuredDesc(sizeof(UE::HLSL::FSamplePoseData), NumTotalSamples),
+		TEXT("AnimSequenceTransformProvider.SamplePoseData"));
+	GraphBuilder.ConvertToExternalBuffer(SamplePoseBuffer);
 
 	UE::Tasks::FTask ProcessTask = GraphBuilder.AddCommandListSetupTask(
 		[
@@ -1285,315 +1443,22 @@ void FAnimSequenceTransformProvider::ProvideGPUTransforms(FSkinningTransformProv
 			, Proxies = Context.Proxies
 			, &NumTotalFramesRequested
 			, TrackUpdateRequests = MoveTemp(TrackUpdateRequests)
-			, SinglePassEvaluateHeaderBuffer
-			, SinglePassBlendEvaluateHeaderBuffer
-			, MultiPassEvaluateHeaderBuffer
-			, MultiPassBlendEvaluateHeaderBuffer
-			, ConcatenateHeaderBuffer
-			, ScatterHeaderBuffer
-			, &SinglePassDispatchGroups
-			, &MultiPassDispatchGroups
-
+			, LayerHeaderBuffer
+			, SamplePoseBuffer
+			, NumTotalEntries
+			, NumTotalSamples
+			, &DispatchGroups
 		] (FRHICommandList& RHICmdList) mutable
 	{
-		TRACE_CPUPROFILER_EVENT_SCOPE(FAnimSequenceTransformProvider::EvaluateAnimSequences);
-
-		// Lock all header buffers on the command list timeline.
-		TArrayView<UE::HLSL::FEvaluateConcatenateScatterBlockHeader>      SinglePassEvaluateUpload      = Lock<UE::HLSL::FEvaluateConcatenateScatterBlockHeader>(RHICmdList, SinglePassEvaluateHeaderBuffer);
-		TArrayView<UE::HLSL::FEvaluateBlendConcatenateScatterBlockHeader> SinglePassBlendEvaluateUpload = Lock<UE::HLSL::FEvaluateBlendConcatenateScatterBlockHeader>(RHICmdList, SinglePassBlendEvaluateHeaderBuffer);
-		TArrayView<UE::HLSL::FEvaluateAnimSequenceBlockHeader>            MultiPassEvaluateUpload       = Lock<UE::HLSL::FEvaluateAnimSequenceBlockHeader>(RHICmdList, MultiPassEvaluateHeaderBuffer);
-		TArrayView<UE::HLSL::FEvaluateAnimSequenceBlendBlockHeader>       MultiPassBlendEvaluateUpload  = Lock<UE::HLSL::FEvaluateAnimSequenceBlendBlockHeader>(RHICmdList, MultiPassBlendEvaluateHeaderBuffer);
-		TArrayView<UE::HLSL::FConcatenateBlockHeader>                     ConcatenateUpload             = Lock<UE::HLSL::FConcatenateBlockHeader>(RHICmdList, ConcatenateHeaderBuffer);
-		TArrayView<UE::HLSL::FScatterTransformsBlockHeader>               ScatterUpload                 = Lock<UE::HLSL::FScatterTransformsBlockHeader>(RHICmdList, ScatterHeaderBuffer);
-
-		static const uint32 MultiPassBonesPerGroup = FEvaluateAnimSequenceCS::BonesPerGroup;
-		uint32 IntermediateTransformOffset = 0;
-		uint32 NumActiveTracks = 0;
-		uint32 NumSinglePassBlocks = 0;
-		uint32 NumSinglePassBlendBlocks = 0;
-		uint32 NumMultiPassBlocks = 0;
-		uint32 NumMultiPassBlendBlocks = 0;
-		uint32 NumCurrentTransformUpdates = 0;
-		uint32 NumPreviousTransformUpdates = 0;
-
-		const auto EvaluateAnimSequence = [&](
-			const FProxyUserData::FLOD& LOD,
-			void* InDispatchGroup,
-			const FMapping& SourceMapping,
-			const FMapping& TargetMapping,
-			FSequence& SourceSequence,
-			FSequence& TargetSequence,
-			float SourceTime,
-			float TargetTime,
-			float BlendWeight,
-			uint32 MaxBoneTransformCount,
-			uint32 DstTransformBufferOffset1,
-			uint32 DstTransformBufferOffset2)
-		{
-			check(SourceMapping.NumCompactBones == TargetMapping.NumCompactBones);
-			const uint32 SourceSequenceTransformBufferOffset    = SourceSequence.BufferAllocation.Offset * sizeof(FBoneTransformWithScale);
-			const uint32 TargetSequenceTransformBufferOffset    = TargetSequence.BufferAllocation.Offset * sizeof(FBoneTransformWithScale);
-			const uint32 SourceRetargetingTransformBufferOffset = SourceMapping.RetargetingTransformBufferOffset * sizeof(FBoneTransformWithScale);
-			const uint32 TargetRetargetingTransformBufferOffset = TargetMapping.RetargetingTransformBufferOffset * sizeof(FBoneTransformWithScale);
-			const uint32 SourceRetargetingDataBufferOffset      = SourceMapping.RetargetingDataBufferAllocation.Offset * sizeof(UE::HLSL::FBoneRetargetingData);
-			const uint32 TargetRetargetingDataBufferOffset      = TargetMapping.RetargetingDataBufferAllocation.Offset * sizeof(UE::HLSL::FBoneRetargetingData);
-			const uint32 InvGlobalRefPoseTransformBufferOffset  = SourceMapping.InvGlobalRefPoseTransformBufferOffset * sizeof(FBoneTransformWithScale);
-			const uint32 LeafBonePathBufferOffset               = LOD.LeafBonePathBufferAllocation.Offset * sizeof(uint32);
-			const uint32 BoneMapBufferOffset                    = LOD.BoneMapBufferOffset * sizeof(uint32);
-
-			int32 SourceKeyIndex0;
-			int32 SourceKeyIndex1;
-			float SourceAlpha;
-			SourceSequence.GetKeyIndicesFromTime(SourceKeyIndex0, SourceKeyIndex1, SourceAlpha, SourceTime);
-
-			NumTotalFramesRequested += SourceSequence.RequestFrame(SourceKeyIndex0);
-			NumTotalFramesRequested += SourceSequence.RequestFrame(SourceKeyIndex1);
-
-			int32 TargetKeyIndex0 = SourceKeyIndex0;
-			int32 TargetKeyIndex1 = SourceKeyIndex1;
-			float TargetAlpha = SourceAlpha;
-
-			if (BlendWeight > 0.0f)
-			{
-				TargetSequence.GetKeyIndicesFromTime(TargetKeyIndex0, TargetKeyIndex1, TargetAlpha, TargetTime);
-
-				NumTotalFramesRequested += TargetSequence.RequestFrame(TargetKeyIndex0);
-				NumTotalFramesRequested += TargetSequence.RequestFrame(TargetKeyIndex1);
-			}
-
-			if (LOD.bSinglePass)
-			{
-				FSinglePassDispatchGroup& Group = *static_cast<FSinglePassDispatchGroup*>(InDispatchGroup);
-
-				if (BlendWeight > 0.0f)
-				{
-					UE::HLSL::FEvaluateBlendConcatenateScatterBlockHeader& BlockHeader = SinglePassBlendEvaluateUpload[Group.BlendEvaluateHeaderBaseOffset + Group.NextBlendEvaluateWriteIndex++];
-
-					BlockHeader.SourceAlpha                            = SourceAlpha;
-					BlockHeader.TargetAlpha                            = TargetAlpha;
-					BlockHeader.BlendWeight                            = BlendWeight;
-					BlockHeader.SourceKeyIndex0                        = SourceKeyIndex0;
-					BlockHeader.SourceKeyIndex1                        = SourceKeyIndex1;
-					BlockHeader.TargetKeyIndex0                        = TargetKeyIndex0;
-					BlockHeader.TargetKeyIndex1                        = TargetKeyIndex1;
-					BlockHeader.SourceSequenceTransformBufferOffset    = SourceSequenceTransformBufferOffset;
-					BlockHeader.TargetSequenceTransformBufferOffset    = TargetSequenceTransformBufferOffset;
-					BlockHeader.TargetRetargetingTransformBufferOffset = TargetRetargetingTransformBufferOffset;
-					BlockHeader.SourceRetargetingTransformBufferOffset = SourceRetargetingTransformBufferOffset;
-					BlockHeader.TargetRetargetingDataBufferOffset      = TargetRetargetingDataBufferOffset;
-					BlockHeader.SourceRetargetingDataBufferOffset      = SourceRetargetingDataBufferOffset;
-					BlockHeader.InvGlobalRefPoseTransformBufferOffset  = InvGlobalRefPoseTransformBufferOffset;
-					BlockHeader.BoneMapBufferOffset                    = BoneMapBufferOffset;
-					BlockHeader.DstTransformBufferOffset1              = DstTransformBufferOffset1;
-					BlockHeader.DstTransformBufferOffset2              = DstTransformBufferOffset2;
-					BlockHeader.SourceNumAnimBones                     = SourceSequence.NumAnimBones;
-					BlockHeader.TargetNumAnimBones                     = TargetSequence.NumAnimBones;
-					BlockHeader.NumCompactBones                        = LOD.BoneTables.NumCompactBones;
-					BlockHeader.SlotsPerInstance                       = LOD.SlotsPerInstance;
-
-					NumSinglePassBlendBlocks++;
-				}
-				else
-				{
-					UE::HLSL::FEvaluateConcatenateScatterBlockHeader& BlockHeader = SinglePassEvaluateUpload[Group.EvaluateHeaderBaseOffset + Group.NextEvaluateWriteIndex++];
-
-					BlockHeader.Alpha                                 = SourceAlpha;
-					BlockHeader.KeyIndex0                             = SourceKeyIndex0;
-					BlockHeader.KeyIndex1                             = SourceKeyIndex1;
-					BlockHeader.SequenceTransformBufferOffset         = SourceSequenceTransformBufferOffset;
-					BlockHeader.RetargetingTransformBufferOffset      = SourceRetargetingTransformBufferOffset;
-					BlockHeader.RetargetingDataBufferOffset           = SourceRetargetingDataBufferOffset;
-					BlockHeader.InvGlobalRefPoseTransformBufferOffset = InvGlobalRefPoseTransformBufferOffset;
-					BlockHeader.BoneMapBufferOffset                   = BoneMapBufferOffset;
-					BlockHeader.DstTransformBufferOffset1             = DstTransformBufferOffset1;
-					BlockHeader.DstTransformBufferOffset2             = DstTransformBufferOffset2;
-					BlockHeader.NumAnimBones                          = SourceSequence.NumAnimBones;
-					BlockHeader.NumCompactBones                       = LOD.BoneTables.NumCompactBones;
-					BlockHeader.SlotsPerInstance                      = LOD.SlotsPerInstance;
-					BlockHeader.Padding                               = 0;
-
-					NumSinglePassBlocks++;
-				}
-			}
-			else
-			{
-				FMultiPassDispatchGroup& Group = *static_cast<FMultiPassDispatchGroup*>(InDispatchGroup);
-
-				// Evaluate: single header per instance
-				if (BlendWeight > 0.0f)
-				{
-					UE::HLSL::FEvaluateAnimSequenceBlendBlockHeader& BlockHeader = MultiPassBlendEvaluateUpload[Group.BlendEvaluateHeaderBaseOffset + Group.NextBlendEvaluateWriteIndex++];
-
-					BlockHeader.SourceAlpha                            = SourceAlpha;
-					BlockHeader.TargetAlpha                            = TargetAlpha;
-					BlockHeader.BlendWeight                            = BlendWeight;
-					BlockHeader.SourceKeyIndex0                        = SourceKeyIndex0;
-					BlockHeader.SourceKeyIndex1                        = SourceKeyIndex1;
-					BlockHeader.TargetKeyIndex0                        = TargetKeyIndex0;
-					BlockHeader.TargetKeyIndex1                        = TargetKeyIndex1;
-					BlockHeader.SourceSequenceTransformBufferOffset    = SourceSequenceTransformBufferOffset;
-					BlockHeader.TargetSequenceTransformBufferOffset    = TargetSequenceTransformBufferOffset;
-					BlockHeader.TargetRetargetingTransformBufferOffset = TargetRetargetingTransformBufferOffset;
-					BlockHeader.SourceRetargetingTransformBufferOffset = SourceRetargetingTransformBufferOffset;
-					BlockHeader.TargetRetargetingDataBufferOffset      = TargetRetargetingDataBufferOffset;
-					BlockHeader.SourceRetargetingDataBufferOffset      = SourceRetargetingDataBufferOffset;
-					BlockHeader.SourceNumAnimBones                     = SourceSequence.NumAnimBones;
-					BlockHeader.TargetNumAnimBones                     = TargetSequence.NumAnimBones;
-					BlockHeader.DstTransformBufferOffset               = IntermediateTransformOffset;
-					BlockHeader.NumCompactBones                        = TargetMapping.NumCompactBones;
-					BlockHeader.Padding                                = 0;
-
-					NumMultiPassBlendBlocks++;
-				}
-				else
-				{
-					UE::HLSL::FEvaluateAnimSequenceBlockHeader& BlockHeader = MultiPassEvaluateUpload[Group.EvaluateHeaderBaseOffset + Group.NextEvaluateWriteIndex++];
-
-					BlockHeader.SequenceTransformBufferOffset    = SourceSequenceTransformBufferOffset;
-					BlockHeader.RetargetingTransformBufferOffset = SourceRetargetingTransformBufferOffset;
-					BlockHeader.RetargetingDataBufferOffset      = SourceRetargetingDataBufferOffset;
-					BlockHeader.DstTransformBufferOffset         = IntermediateTransformOffset;
-					BlockHeader.NumAnimBones                     = SourceSequence.NumAnimBones;
-					BlockHeader.KeyIndex0                        = SourceKeyIndex0;
-					BlockHeader.KeyIndex1                        = SourceKeyIndex1;
-					BlockHeader.Alpha                            = SourceAlpha;
-					BlockHeader.NumCompactBones                  = TargetMapping.NumCompactBones;
-					BlockHeader.Padding                          = 0;
-
-					NumMultiPassBlocks++;
-				}
-
-				// Concatenate: single header per instance
-				{
-					UE::HLSL::FConcatenateBlockHeader& ConcatHeader = ConcatenateUpload[Group.ConcatenateHeaderBaseOffset + Group.NextConcatenateWriteIndex++];
-
-					ConcatHeader.LeafBonePathBufferOffset = LeafBonePathBufferOffset;
-					ConcatHeader.TransformBufferOffset    = IntermediateTransformOffset;
-					ConcatHeader.NumPrefetchBones         = FMath::Min<uint32>(LOD.BoneTables.NumCompactBones, MultiPassBonesPerGroup);
-					ConcatHeader.NumLeafBones             = LOD.LeafBonePathData.NumLeafBones;
-					ConcatHeader.Padding                  = 0;
-				}
-
-				// Scatter: single header per instance
-				{
-					UE::HLSL::FScatterTransformsBlockHeader& ScatterHeader = ScatterUpload[Group.ScatterHeaderBaseOffset + Group.NextScatterWriteIndex++];
-
-					ScatterHeader.InvGlobalRefPoseTransformBufferOffset = InvGlobalRefPoseTransformBufferOffset;
-					ScatterHeader.SrcTransformBufferOffset              = IntermediateTransformOffset;
-					ScatterHeader.DstTransformBufferOffset1             = DstTransformBufferOffset1;
-					ScatterHeader.DstTransformBufferOffset2             = DstTransformBufferOffset2;
-					ScatterHeader.BoneMapBufferOffset                   = BoneMapBufferOffset;
-					ScatterHeader.MaxBoneTransformCount                 = MaxBoneTransformCount;
-					ScatterHeader.Padding                               = 0;
-				}
-
-				IntermediateTransformOffset += SourceMapping.NumCompactBones * sizeof(FBoneTransformWithScale);
-			}
-		};
-
-		uint32 UpdateResultIndex = 0;
-
-		for (const FSkinningTransformProvider::FProviderIndirection Indirection : Indirections)
-		{
-			FInstancedSkinningSceneExtensionProxy* ExtensionProxy = static_cast<FInstancedSkinningSceneExtensionProxy*>(Proxies[Indirection.Index]);
-			FProxyUserData& ProxyUserData = static_cast<FProxyUserData&>(ExtensionProxy->GetUserDataChecked());
-			const uint32 LODIndex = ExtensionProxy->GetLOD();
-			FProxyUserData::FLOD& LOD = ProxyUserData.LODs[LODIndex];
-			FAnimSequenceTrackRenderPool& Tracks = ProxyUserData.Tracks;
-
-			// Look up the dispatch group for this LOD's bone configuration.
-			void* DispatchGroup = nullptr;
-			if (LOD.bSinglePass)
-			{
-				FSinglePassDispatchGroup LookupKey;
-				LookupKey.NumCompactBones = LOD.BoneTables.NumCompactBones;
-				DispatchGroup = SinglePassDispatchGroups.Find(LookupKey);
-			}
-			else
-			{
-				FMultiPassDispatchGroup LookupKey;
-				LookupKey.NumCompactBones       = LOD.BoneTables.NumCompactBones;
-				LookupKey.NumLeafBones          = LOD.LeafBonePathData.NumLeafBones;
-				LookupKey.MaxBoneTransformCount = ExtensionProxy->GetMaxBoneTransformCount();
-				DispatchGroup = MultiPassDispatchGroups.Find(LookupKey);
-			}
-			check(DispatchGroup != nullptr);
-
-			const uint32 CurrentTransformOffset  = Indirection.CurrentTransformOffset;
-			const uint32 PreviousTransformOffset = Indirection.PreviousTransformOffset;
-
-			Tracks.EnumerateActiveTracks([&] (int32 TrackIndex)
-			{
-				FAnimSequenceTrackRenderData& Data = Tracks.GetData(TrackIndex);
-				const FTrackData SourceTrack       = ProxyUserData.GetTrackData(Data.GetSource(), WorldTime, LODIndex);
-				const FTrackData TargetTrack       = ProxyUserData.GetTrackData(Data.GetTarget(), WorldTime, LODIndex);
-				const float BlendWeight            = Data.GetBlendWeight();
-
-				const FTrackUpdateRequest UpdateRequest   = TrackUpdateRequests[UpdateResultIndex++];
-				const uint32 NumTransforms                = ExtensionProxy->GetMaxBoneTransformCount();
-				const uint32 TrackTransformOffset         = TrackIndex * NumTransforms * 2 * sizeof(FCompressedBoneTransform);
-				const uint32 TrackCurrentTransformOffset  = TrackTransformOffset + CurrentTransformOffset;
-				const uint32 TrackPreviousTransformOffset = TrackTransformOffset + PreviousTransformOffset;
-
-				if (UpdateRequest.bNeedsCurrent)
-				{
-					uint32 OutputTransformOffset1 = TrackCurrentTransformOffset;
-					uint32 OutputTransformOffset2 = UpdateRequest.GetUpdateMode() == EPreviousBoneTransformUpdateMode::DuplicateCurrentToPrevious ? TrackPreviousTransformOffset : TrackCurrentTransformOffset;
-
-					EvaluateAnimSequence(
-						LOD, DispatchGroup,
-						SourceTrack.Mapping, TargetTrack.Mapping,
-						SourceTrack.Sequence, TargetTrack.Sequence,
-						SourceTrack.Position, TargetTrack.Position,
-						BlendWeight,
-						NumTransforms,
-						OutputTransformOffset1, OutputTransformOffset2);
-
-					NumCurrentTransformUpdates++;
-				}
-
-				if (UpdateRequest.GetUpdateMode() == EPreviousBoneTransformUpdateMode::UpdatePrevious)
-				{
-					EvaluateAnimSequence(
-						LOD, DispatchGroup,
-						SourceTrack.Mapping, TargetTrack.Mapping,
-						SourceTrack.Sequence, TargetTrack.Sequence,
-						SourceTrack.PreviousPosition, TargetTrack.PreviousPosition,
-						BlendWeight,
-						NumTransforms,
-						TrackPreviousTransformOffset, TrackPreviousTransformOffset);
-
-					NumPreviousTransformUpdates++;
-				}
-
-				Data.Tick(DeltaTime, SourceTrack.Position, TargetTrack.Position);
-				NumActiveTracks++;
-			});
-		}
-
-		Unlock(RHICmdList, SinglePassEvaluateHeaderBuffer);
-		Unlock(RHICmdList, SinglePassBlendEvaluateHeaderBuffer);
-		Unlock(RHICmdList, MultiPassEvaluateHeaderBuffer);
-		Unlock(RHICmdList, MultiPassBlendEvaluateHeaderBuffer);
-		Unlock(RHICmdList, ConcatenateHeaderBuffer);
-		Unlock(RHICmdList, ScatterHeaderBuffer);
-
-		INC_DWORD_STAT_BY(STAT_AnimSequenceTransformProvider_ActiveTracks, NumActiveTracks);
-		INC_DWORD_STAT_BY(STAT_AnimSequenceTransformProvider_EvaluateSinglePassBlocks, NumSinglePassBlocks);
-		INC_DWORD_STAT_BY(STAT_AnimSequenceTransformProvider_EvaluateSinglePassBlendBlocks, NumSinglePassBlendBlocks);
-		INC_DWORD_STAT_BY(STAT_AnimSequenceTransformProvider_EvaluateBlocks, NumMultiPassBlocks);
-		INC_DWORD_STAT_BY(STAT_AnimSequenceTransformProvider_EvaluateBlendBlocks, NumMultiPassBlendBlocks);
-		INC_DWORD_STAT_BY(STAT_AnimSequenceTransformProvider_CurrentTransformUpdates, NumCurrentTransformUpdates);
-		INC_DWORD_STAT_BY(STAT_AnimSequenceTransformProvider_PreviousTransformUpdates, NumPreviousTransformUpdates);
+		UploadBlockHeaders(RHICmdList, WorldTime, DeltaTime, Indirections, Proxies, NumTotalFramesRequested, TrackUpdateRequests, LayerHeaderBuffer, SamplePoseBuffer, NumTotalEntries, NumTotalSamples, DispatchGroups);
 	});
 
 	FRDGBuffer* SequenceTransformBuffer = SequenceTransformPersistentBuffer.ResizeBufferIfNeeded(GraphBuilder, SequenceTransformAllocator.GetMaxSize());
 	FRDGBuffer* RetargetingDataBuffer = RetargetingDataPersistentBuffer.ResizeBufferIfNeeded(GraphBuilder, RetargetingDataAllocator.GetMaxSize());
-	FRDGBuffer* LeafBonePathBuffer = LeafBonePathPersistentBuffer.ResizeBufferIfNeeded(GraphBuilder, LeafBonePathAllocator.GetMaxSize());
-
+	FRDGBuffer* BoneHierarchyBuffer = BoneHierarchyPersistentBuffer.ResizeBufferIfNeeded(GraphBuilder, BoneHierarchyAllocator.GetMaxSize());
+	FRDGBuffer* BoneMaskBuffer = BoneMaskPersistentBuffer.ResizeBufferIfNeeded(GraphBuilder, BoneMaskAllocator.GetMaxSize());
 	SET_MEMORY_STAT(STAT_AnimSequenceTransformProvider_SequenceTransformMemory, SequenceTransformBuffer->Desc.GetSize());
 	SET_MEMORY_STAT(STAT_AnimSequenceTransformProvider_RetargetingDataMemory, RetargetingDataBuffer->Desc.GetSize());
-	SET_MEMORY_STAT(STAT_AnimSequenceTransformProvider_LeafBonePathMemory, LeafBonePathBuffer->Desc.GetSize());
 
 	bool bExecuteScatterUpload = false;
 
@@ -1601,277 +1466,129 @@ void FAnimSequenceTransformProvider::ProvideGPUTransforms(FSkinningTransformProv
 	{
 		Context.ScatterUploadBuilder.AddPrerequisiteTask(ProcessTask);
 
-		Context.ScatterUploadBuilder.AddPass(GraphBuilder, SequenceTransformUploadBuffer, SequenceTransformBuffer, NumTotalSequenceTransforms, sizeof(FBoneTransformWithScale), TEXT("AnimSequenceTransformProvider.Sequences"),
+		Context.ScatterUploadBuilder.AddPass(GraphBuilder, SequenceTransformUploadBuffer, SequenceTransformBuffer, NumTotalSequenceTransforms, sizeof(FVector4f), TEXT("AnimSequenceTransformProvider.Sequences"),
 			[this, &NumTotalFramesRequested] (FRDGScatterUploader& ScatterUploader)
 		{
-			TRACE_CPUPROFILER_EVENT_SCOPE(FAnimSequenceTransformProvider::UploadSequences);
+			UploadSequences(ScatterUploader, NumTotalFramesRequested);
+		});
+
+		bExecuteScatterUpload = true;
+	}
+
+	if (NumTotalRetargetingDatas > 0)
+	{
+		Context.ScatterUploadBuilder.AddPass(GraphBuilder, RetargetingDataUploadBuffer, RetargetingDataBuffer, NumTotalRetargetingDatas, sizeof(UE::HLSL::FBoneRetargetingData), TEXT("AnimSequenceTransformProvider.RetargetingData"),
+			[this] (FRDGScatterUploader& ScatterUploader)
+		{
+			TRACE_CPUPROFILER_EVENT_SCOPE(FAnimSequenceTransformProvider::UploadRetargetingData);
 
 			for (const FMapping* Mapping : MappingsToUpload)
 			{
-				TArrayView<FBoneTransformWithScale> DstTransforms = ScatterUploader.Add_GetRef<FBoneTransformWithScale>(Mapping->MappingTransformBufferAllocation.Offset, Mapping->MappingTransformBufferAllocation.Count);
+				TArrayView<UE::HLSL::FBoneRetargetingData> DstRetargetingDatas = ScatterUploader.Add_GetRef<UE::HLSL::FBoneRetargetingData>(Mapping->RetargetingDataBufferAllocation.Offset, Mapping->RetargetingDataBufferAllocation.Count);
 
 				for (uint32 CompactBoneIndex = 0; CompactBoneIndex < Mapping->NumCompactBones; ++CompactBoneIndex)
 				{
-					DstTransforms[CompactBoneIndex].SetTransform(Mapping->InvGlobalRefPoseTransforms[CompactBoneIndex]);
-
-					const uint32 RetargetTransformIndex = Mapping->NumCompactBones + CompactBoneIndex * (Mapping->RetargetingTables.bRequiresRefPoseRetarget ? 2 : 1);
-					DstTransforms[RetargetTransformIndex].SetTransform(Mapping->RetargetingTables.Transform[CompactBoneIndex]);
-
-					if (Mapping->RetargetingTables.bRequiresRefPoseRetarget)
-					{
-						FBoneTransformWithScale& RefPoseSlot = DstTransforms[RetargetTransformIndex + 1];
-						const FQuat4f& Q0 = Mapping->RetargetingTables.RefPoseRetargetQuats[CompactBoneIndex].Get<0>();
-						const FQuat4f& Q1 = Mapping->RetargetingTables.RefPoseRetargetQuats[CompactBoneIndex].Get<1>();
-
-						// Pack Q0 into Rotation
-						RefPoseSlot.Rotation[0] = Q0.X;
-						RefPoseSlot.Rotation[1] = Q0.Y;
-						RefPoseSlot.Rotation[2] = Q0.Z;
-						RefPoseSlot.Rotation[3] = Q0.W;
-
-						// Pack Q1 into Translation + Scale
-						RefPoseSlot.Translation[0] = Q1.X;
-						RefPoseSlot.Translation[1] = Q1.Y;
-						RefPoseSlot.Translation[2] = Q1.Z;
-						RefPoseSlot.Scale          = Q1.W;
-					}
-				}
-			}
-
-			BoneTrackArray TrackToBoneIndexMap;
-			FMemMark Mark(FMemStack::Get());
-			TArray<FTransform, TMemStackAllocator<>> AnimLocalPoses;
-			uint32 NumSequenceFrameUploads = 0;
-			int32 SequenceIndex = 0;
+					UE::HLSL::FBoneRetargetingData& RetargetingData = DstRetargetingDatas[CompactBoneIndex];
+					RetargetingData.RetargetMode  = Mapping->RetargetingTables.Mode[CompactBoneIndex];
+					RetargetingData.AnimBoneIndex = Mapping->RetargetingTables.CompactToAnim[CompactBoneIndex];
+					RetargetingData.bRequiresRefPoseRetarget = Mapping->RetargetingTables.bRequiresRefPoseRetarget;
+				}
+			}
+		});
 
-			const int32 UploadBudget = FMath::Max<int32>(0, CVarAnimSequenceTransformProviderUploadBudget.GetValueOnRenderThread() - NumTotalFramesRequested);
-			const int32 UploadBudgetDivisor = FMath::Max(1, SequencesToUpload.Num());
-			const int32 FramesPerSequence = UploadBudget / UploadBudgetDivisor;
-			int32 NumFramesLeftover       = UploadBudget % UploadBudgetDivisor;
+		bExecuteScatterUpload = true;
+	}
 
-			while (SequenceIndex < SequencesToUpload.Num())
+	if (!ProxyUserDatasToUpload.IsEmpty() && BoneMaskBuffer)
+	{
+		uint32 NumTotalBoneMaskElements = 0;
+		for (const FProxyUserData* UserData : ProxyUserDatasToUpload)
+		{
+			for (const FProxyUserData::FLOD& LOD : UserData->LODs)
 			{
-				FSequence& Sequence = *SequencesToUpload[SequenceIndex];
-
-				int32 NumBudgetFramesToUpload = FramesPerSequence + (NumFramesLeftover > 0 ? 1 : 0);
-				NumFramesLeftover = FMath::Max(NumFramesLeftover - 1, 0);
-
-				while (Sequence.NextFrameToUpload < Sequence.NumFrames && NumBudgetFramesToUpload > 0)
+				for (int32 LayerIndex = 0; LayerIndex < LOD.BoneMaskBufferOffsets.Num(); LayerIndex++)
 				{
-					const int32 FrameStep = 1u << Sequence.TargetSampleRateLevel;
-
-					if (Sequence.FramesUploaded[Sequence.NextFrameToUpload] || Sequence.FramesRequested[Sequence.NextFrameToUpload])
+					if (LOD.BoneMaskBufferOffsets[LayerIndex] != FProxyUserData::InvalidBoneMaskOffset)
 					{
-						Sequence.NextFrameToUpload += FrameStep;
-						continue;
+						NumTotalBoneMaskElements += LOD.BoneTables.NumCompactBones;
 					}
-
-					Sequence.FramesRequested[Sequence.NextFrameToUpload] = true;
-					Sequence.NumFramesRequested++;
-					Sequence.NextFrameToUpload += FrameStep;
-					NumBudgetFramesToUpload--;
-				}
-
-				if (Sequence.NextFrameToUpload >= Sequence.NumFrames && Sequence.TargetSampleRateLevel > 0)
-				{
-					Sequence.TargetSampleRateLevel--;
-					Sequence.CurrentSampleRateLevel--;
-					Sequence.NextFrameToUpload = 1u << Sequence.TargetSampleRateLevel;
 				}
+			}
+		}
 
-				if (Sequence.NumFramesRequested > 0)
+		if (NumTotalBoneMaskElements > 0)
+		{
+			Context.ScatterUploadBuilder.AddPass(GraphBuilder, BoneMaskUploadBuffer, BoneMaskBuffer, NumTotalBoneMaskElements, sizeof(float), TEXT("AnimSequenceTransformProvider.BoneMasks"),
+				[this](FRDGScatterUploader& ScatterUploader)
+			{
+				for (const FProxyUserData* UserData : ProxyUserDatasToUpload)
 				{
-					const UAnimSequence* AnimSequence = Sequence.Sequence;
-					UAnimSequence::FScopedCompressedAnimSequence CompressedAnimSequence = AnimSequence->GetCompressedData();
-					const FCompressedAnimSequence& PlatformCompressedData = CompressedAnimSequence.Get();
-
-					const TArray<FTransform>& AnimRefLocalPoses = AnimSequence->GetSkeleton()->GetRefLocalPoses();
-					check(Sequence.NumAnimBones <= (uint16)AnimRefLocalPoses.Num());
-
-					const int32 NumTracks = PlatformCompressedData.CompressedTrackToSkeletonMapTable.Num();
-					TrackToBoneIndexMap.Reset(NumTracks);
-
-					for (int32 TrackIndex = 0; TrackIndex < NumTracks; ++TrackIndex)
+					for (const FProxyUserData::FLOD& LOD : UserData->LODs)
 					{
-						const uint16 BoneIndex = PlatformCompressedData.GetSkeletonIndexFromTrackIndex(TrackIndex);
+						const uint32 NumCompactBones = LOD.BoneTables.NumCompactBones;
 
-						// We only care about raw bones.
-						if (BoneIndex < Sequence.NumAnimBones)
+						for (int32 LayerIndex = 0; LayerIndex < LOD.BoneMaskBufferOffsets.Num(); LayerIndex++)
 						{
-							TrackToBoneIndexMap.Add(BoneTrackPair(BoneIndex, TrackIndex));
-						}
-					}
-
-					FAnimSequenceDecompressionContext DecompressionContext(
-						  Sequence.SampleRate
-						, Sequence.NumFrames
-						, AnimSequence->Interpolation
-						, AnimSequence->GetRetargetTransformsSourceName()
-						, *PlatformCompressedData.CompressedDataStructure
-						, AnimRefLocalPoses
-						, PlatformCompressedData.CompressedTrackToSkeletonMapTable
-						, AnimSequence->GetSkeleton()
-						, AnimSequence->IsValidAdditive()
-						, AnimSequence->GetAdditiveAnimType()
-					);
-
-					AnimLocalPoses.Reset(AnimRefLocalPoses.Num());
-					AnimLocalPoses.Append(AnimRefLocalPoses);
-
-					for (TConstSetBitIterator<> BitIt(Sequence.FramesRequested); BitIt; ++BitIt)
-					{
-						TRACE_CPUPROFILER_EVENT_SCOPE(FAnimSequenceTransformProvider::DecompressPose);
-						const int32 FrameIndex = BitIt.GetIndex();
-
-						const double SeekTime = Sequence.SampleRate.AsSeconds(FFrameTime(FrameIndex));
-						DecompressionContext.Seek(SeekTime);
+							const uint32 MaskOffset = LOD.BoneMaskBufferOffsets[LayerIndex];
+							if (MaskOffset == FProxyUserData::InvalidBoneMaskOffset)
+							{
+								continue;
+							}
 
-						if (PlatformCompressedData.BoneCompressionCodec)
-						{
-							TArrayView<FTransform> AnimLocalPosesView(AnimLocalPoses);
-							PlatformCompressedData.BoneCompressionCodec->DecompressPose(DecompressionContext, TrackToBoneIndexMap, TrackToBoneIndexMap, TrackToBoneIndexMap, AnimLocalPosesView);
-						}
+							const UBlendProfile* BlendMask = UserData->BlendProfiles[LayerIndex];
+							check(BlendMask);
 
-						TArrayView<FBoneTransformWithScale> DstTransforms = ScatterUploader.Add_GetRef<FBoneTransformWithScale>(Sequence.BufferAllocation.Offset + Sequence.NumAnimBones * FrameIndex, Sequence.NumAnimBones);
+							const uint32 AllocationIndex = MaskOffset / sizeof(float);
+							TArrayView<float> DstWeights = ScatterUploader.Add_GetRef<float>(AllocationIndex, NumCompactBones);
 
-						for (uint16 BoneIndex = 0; BoneIndex < Sequence.NumAnimBones; ++BoneIndex)
-						{
-							DstTransforms[BoneIndex].SetTransform(FTransform3f(AnimLocalPoses[BoneIndex]));
+							for (uint32 CompactBoneIndex = 0; CompactBoneIndex < NumCompactBones; CompactBoneIndex++)
+							{
+								const uint16 TargetSkeletonBoneIndex = LOD.BoneTables.CompactToTargetSkeleton[CompactBoneIndex];
+								DstWeights[CompactBoneIndex] = FMath::Clamp(BlendMask->GetBoneBlendScale(TargetSkeletonBoneIndex), 0.0f, 1.0f);
+							}
 						}
-
-						Sequence.NumFramesUploaded++;
-						Sequence.FramesUploaded[FrameIndex] = true;
-						Sequence.FramesRequested[FrameIndex] = false;
-						NumSequenceFrameUploads++;
 					}
-
-					Sequence.NumFramesRequested = 0;
 				}
+			});
 
-				if (Sequence.NumFramesUploaded == Sequence.NumFrames)
-				{
-					SequencesToUpload.RemoveAtSwap(SequenceIndex, EAllowShrinking::No);
-				}
-				else
-				{
-					SequenceIndex++;
-				}
-			}
-
-			INC_DWORD_STAT_BY(STAT_AnimSequenceTransformProvider_SequenceFrameUploads, NumSequenceFrameUploads);
-		});
-
-		bExecuteScatterUpload = true;
-	}
-
-	if (NumTotalRetargetingDatas > 0)
-	{
-		Context.ScatterUploadBuilder.AddPass(GraphBuilder, RetargetingDataUploadBuffer, RetargetingDataBuffer, NumTotalRetargetingDatas, sizeof(UE::HLSL::FBoneRetargetingData), TEXT("AnimSequenceTransformProvider.RetargetingData"),
-			[this] (FRDGScatterUploader& ScatterUploader)
-		{
-			TRACE_CPUPROFILER_EVENT_SCOPE(FAnimSequenceTransformProvider::UploadRetargetingData);
-
-			for (const FMapping* Mapping : MappingsToUpload)
-			{
-				TArrayView<UE::HLSL::FBoneRetargetingData> DstRetargetingDatas = ScatterUploader.Add_GetRef<UE::HLSL::FBoneRetargetingData>(Mapping->RetargetingDataBufferAllocation.Offset, Mapping->RetargetingDataBufferAllocation.Count);
-
-				for (uint32 CompactBoneIndex = 0; CompactBoneIndex < Mapping->NumCompactBones; ++CompactBoneIndex)
-				{
-					UE::HLSL::FBoneRetargetingData& RetargetingData = DstRetargetingDatas[CompactBoneIndex];
-					RetargetingData.RetargetMode  = Mapping->RetargetingTables.Mode[CompactBoneIndex];
-					RetargetingData.AnimBoneIndex = Mapping->RetargetingTables.CompactToAnim[CompactBoneIndex];
-					RetargetingData.bRequiresRefPoseRetarget = Mapping->RetargetingTables.bRequiresRefPoseRetarget;
-				}
-			}
-		});
-
-		bExecuteScatterUpload = true;
+			bExecuteScatterUpload = true;
+		}
 	}
 
-	if (NumTotalLeafBonePathElements > 0)
+	if (!LODUploadRequests.IsEmpty())
 	{
-		Context.ScatterUploadBuilder.AddPass(GraphBuilder, LeafBonePathUploadBuffer, LeafBonePathBuffer, NumTotalLeafBonePathElements, sizeof(uint32), TEXT("AnimSequenceTransformProvider.LeafBonePaths"),
+		Context.ScatterUploadBuilder.AddPass(GraphBuilder, BoneHierarchyUploadBuffer, BoneHierarchyBuffer, NumTotalBoneHierarchyElements, sizeof(UE::HLSL::FBoneHierarchyData), TEXT("AnimSequenceTransformProvider.BoneHierarchy"),
 			[this, LODUploadRequests = MoveTemp(LODUploadRequests)](FRDGScatterUploader& ScatterUploader)
 		{
-			TRACE_CPUPROFILER_EVENT_SCOPE(FAnimSequenceTransformProvider::UploadLeafBonePaths);
-			TArray<uint16, SceneRenderingAllocator> CompactToTransformStorage;
+			TArray<uint16, SceneRenderingAllocator> CompactToTransform;
 
 			for (const FLODUploadRequest& Request : LODUploadRequests)
 			{
 				FProxyUserData::FLOD& LOD = Request.LOD;
 				TConstArrayView<FBoneIndexType> BoneMap = Request.ExtensionProxy.GetBoneMap();
 				const FBoneTables& BoneTables = LOD.BoneTables;
-				const int32 NumTransformBones = Request.ExtensionProxy.GetMaxBoneTransformCount();
 
-				if (LOD.bSinglePass)
-				{
-					TArrayView<UE::HLSL::FBoneHierarchyData> DstBuffer = ScatterUploader.Add_GetRef<UE::HLSL::FBoneHierarchyData>(LOD.LeafBonePathBufferAllocation.Offset, LOD.LeafBonePathBufferAllocation.Count);
-					TConstArrayView<uint16> CompactToTransform;
+				TArrayView<UE::HLSL::FBoneHierarchyData> DstBuffer = ScatterUploader.Add_GetRef<UE::HLSL::FBoneHierarchyData>(LOD.BoneMapBufferAllocation.Offset, LOD.BoneMapBufferAllocation.Count);
 
-					if (!BoneMap.IsEmpty())
-					{
-						CompactToTransformStorage.Init(INVALID_BONE_INDEX, BoneTables.NumCompactBones);
+				CompactToTransform.Init(INVALID_BONE_INDEX, BoneTables.NumCompactBones);
 
-						for (uint32 TransformIndex = 0; TransformIndex < (uint32)BoneMap.Num(); ++TransformIndex)
-						{
-							const uint16 CompactBoneIndex = BoneTables.MeshToCompact[BoneMap[TransformIndex]];
-
-							if (CompactBoneIndex != INVALID_BONE_INDEX)
-							{
-								CompactToTransformStorage[CompactBoneIndex] = TransformIndex;
-							}
-						}
-
-						CompactToTransform = CompactToTransformStorage;
-					}
-					else
-					{
-						CompactToTransform = BoneTables.CompactToMesh;
-					}
+				for (uint32 TransformIndex = 0; TransformIndex < (uint32)BoneMap.Num(); ++TransformIndex)
+				{
+					const uint16 CompactBoneIndex = BoneTables.MeshToCompact[BoneMap[TransformIndex]];
 
-					for (uint32 CompactBoneIndex = 0; CompactBoneIndex < BoneTables.NumCompactBones; ++CompactBoneIndex)
+					if (CompactBoneIndex != INVALID_BONE_INDEX)
 					{
-						DstBuffer[CompactBoneIndex] = UE::HLSL::FBoneHierarchyData
-						{
-							  .ParentIndex    = (uint16)BoneTables.CompactToParent[CompactBoneIndex]
-							, .TransformIndex = (uint16)CompactToTransform[CompactBoneIndex]
-						};
+						CompactToTransform[CompactBoneIndex] = TransformIndex;
 					}
 				}
-				else
-				{
-					FLeafBonePathData& LeafBonePathData = LOD.LeafBonePathData;
-					TArrayView<uint32> DstBuffer = ScatterUploader.Add_GetRef<uint32>(LOD.LeafBonePathBufferAllocation.Offset, LOD.LeafBonePathBufferAllocation.Count);
 
-					const int32 NumLeafBonePaths = LeafBonePathData.Buffer.Num();
-					FMemory::Memcpy(DstBuffer.GetData(), LeafBonePathData.Buffer.GetData(), NumLeafBonePaths * sizeof(uint32));
-
-					// Release CPU side data now that it's been built.
-					LeafBonePathData.Buffer.Empty();
-
-					if (!BoneMap.IsEmpty())
-					{
-						// Transform bone space should match bone map space.
-						check(NumTransformBones == BoneMap.Num());
-
-						for (uint32 TransformIndex = 0; TransformIndex < (uint32)BoneMap.Num(); ++TransformIndex)
-						{
-							const uint16 MeshBoneIndex = BoneMap[TransformIndex];
-							const uint16 CompactBoneIndex = BoneTables.MeshToCompact[MeshBoneIndex];
-
-							DstBuffer[NumLeafBonePaths + TransformIndex] = CompactBoneIndex;
-						}
-					}
-					else
+				for (uint32 CompactBoneIndex = 0; CompactBoneIndex < BoneTables.NumCompactBones; ++CompactBoneIndex)
+				{
+					DstBuffer[CompactBoneIndex] = UE::HLSL::FBoneHierarchyData
 					{
-						// Transform bone space should match mesh bone space.
-						check(NumTransformBones == BoneTables.NumMeshBones);
-
-						for (uint32 MeshBoneIndex = 0; MeshBoneIndex < (uint32)BoneTables.MeshToCompact.Num(); ++MeshBoneIndex)
-						{
-							DstBuffer[NumLeafBonePaths + MeshBoneIndex] = BoneTables.MeshToCompact[MeshBoneIndex];
-						}
-					}
+						  .ParentIndex    = (uint16)BoneTables.CompactToParent[CompactBoneIndex]
+						, .TransformIndex = (uint16)CompactToTransform[CompactBoneIndex]
+					};
 				}
 			}
 		});
@@ -1884,583 +1601,49 @@ void FAnimSequenceTransformProvider::ProvideGPUTransforms(FSkinningTransformProv
 		Context.ScatterUploadBuilder.Execute(Context.GraphBuilder);
 	}
 
-	FRDGBuffer* LocalSpaceTransformBuffer = nullptr;
-	FRDGBuffer* GlobalSpaceTransformBuffer = nullptr;
-
-	if (NumTotalIntermediateTransforms > 0)
-	{
-		LocalSpaceTransformBuffer = GraphBuilder.CreateBuffer(FRDGBufferDesc::CreateByteAddressDesc(NumTotalIntermediateTransforms * sizeof(FBoneTransformWithScale)), TEXT("AnimSequenceTransformProvider.LocalSpaceBoneTransforms"));
-		GlobalSpaceTransformBuffer = GraphBuilder.CreateBuffer(FRDGBufferDesc::CreateByteAddressDesc(NumTotalIntermediateTransforms * sizeof(FBoneTransformWithScale)), TEXT("AnimSequenceTransformProvider.GlobalSpaceBoneTransforms"));
-	}
-
 	FGlobalShaderMap* ShaderMap = GetGlobalShaderMap(GMaxRHIFeatureLevel);
 
 	RDG_EVENT_SCOPE_STAT(GraphBuilder, AnimSequenceTransformProvider, "AnimSequenceTransformProvider");
 
-	// Single-pass evaluate dispatches
-	if (SinglePassEvaluateHeaderBuffer)
 	{
 		FEvaluateConcatenateScatterCS::FParameters* PassParameters = GraphBuilder.AllocParameters<FEvaluateConcatenateScatterCS::FParameters>();
-		PassParameters->RetargetingDataBuffer      = GraphBuilder.CreateSRV(RetargetingDataBuffer);
-		PassParameters->SequenceTransformBuffer    = GraphBuilder.CreateSRV(SequenceTransformBuffer);
-		PassParameters->EvaluateBlockHeaderBuffer  = GraphBuilder.CreateSRV(SinglePassEvaluateHeaderBuffer);
-		PassParameters->BoneMapBuffer              = GraphBuilder.CreateSRV(LeafBonePathBuffer);
-		PassParameters->SkinningTransformBufferUAV = GetCompressedBoneTransformUAV(GraphBuilder, Context.TransformBuffer);
-
-		GraphBuilder.AddPass(
-			RDG_EVENT_NAME("EvaluateConcatenateScatter"),
-			PassParameters,
-			ERDGPassFlags::Compute,
-			[ShaderMap, PassParameters, &SinglePassDispatchGroups](FRDGAsyncTask, FRHIComputeCommandList& RHICmdList)
-		{
-			for (const FSinglePassDispatchGroup& Group : SinglePassDispatchGroups)
-			{
-				if (Group.NumEvaluateHeaders > 0)
-				{
-					FEvaluateConcatenateScatterCS::FPermutationDomain Perm;
-					Perm.Set<FEvaluateConcatenateScatterCS::FBonesPerGroupDim>(Group.GroupSize);
-					TShaderMapRef<FEvaluateConcatenateScatterCS> ComputeShader(ShaderMap, Perm);
-
-					PassParameters->NumTotalHeaders  = Group.NumEvaluateHeaders;
-					PassParameters->HeaderBaseOffset = Group.EvaluateHeaderBaseOffset;
-
-					const uint32 NumGroups = FMath::DivideAndRoundUp(Group.NumEvaluateHeaders, Group.InstancesPerGroup);
-					FComputeShaderUtils::Dispatch(RHICmdList, ComputeShader, *PassParameters, FIntVector(NumGroups, 1, 1));
-				}
-			}
-		});
-	}
-
-	// Single-pass blend evaluate dispatches
-	if (SinglePassBlendEvaluateHeaderBuffer)
-	{
-		FEvaluateBlendConcatenateScatterCS::FParameters* PassParameters = GraphBuilder.AllocParameters<FEvaluateBlendConcatenateScatterCS::FParameters>();
-		PassParameters->RetargetingDataBuffer      = GraphBuilder.CreateSRV(RetargetingDataBuffer);
-		PassParameters->SequenceTransformBuffer    = GraphBuilder.CreateSRV(SequenceTransformBuffer);
-		PassParameters->EvaluateBlockHeaderBuffer  = GraphBuilder.CreateSRV(SinglePassBlendEvaluateHeaderBuffer);
-		PassParameters->BoneMapBuffer              = GraphBuilder.CreateSRV(LeafBonePathBuffer);
-		PassParameters->SkinningTransformBufferUAV = GetCompressedBoneTransformUAV(GraphBuilder, Context.TransformBuffer);
-
-		GraphBuilder.AddPass(
-			RDG_EVENT_NAME("EvaluateBlendConcatenateScatter"),
-			PassParameters,
-			ERDGPassFlags::Compute,
-			[ShaderMap, PassParameters, &SinglePassDispatchGroups](FRDGAsyncTask, FRHIComputeCommandList& RHICmdList)
-		{
-			for (const FSinglePassDispatchGroup& Group : SinglePassDispatchGroups)
-			{
-				if (Group.NumBlendEvaluateHeaders > 0)
-				{
-					FEvaluateBlendConcatenateScatterCS::FPermutationDomain Perm;
-					Perm.Set<FEvaluateBlendConcatenateScatterCS::FBonesPerGroupDim>(Group.GroupSize);
-					TShaderMapRef<FEvaluateBlendConcatenateScatterCS> ComputeShader(ShaderMap, Perm);
-
-					PassParameters->NumTotalHeaders  = Group.NumBlendEvaluateHeaders;
-					PassParameters->HeaderBaseOffset = Group.BlendEvaluateHeaderBaseOffset;
-
-					const uint32 NumGroups = FMath::DivideAndRoundUp(Group.NumBlendEvaluateHeaders, Group.InstancesPerGroup);
-					FComputeShaderUtils::Dispatch(RHICmdList, ComputeShader, *PassParameters, FIntVector(NumGroups, 1, 1));
-				}
-			}
-		});
-	}
-
-	// 3-pass evaluate dispatches
-	if (MultiPassEvaluateHeaderBuffer)
-	{
-		FEvaluateAnimSequenceCS::FParameters* PassParameters = GraphBuilder.AllocParameters<FEvaluateAnimSequenceCS::FParameters>();
 		PassParameters->RetargetingDataBuffer        = GraphBuilder.CreateSRV(RetargetingDataBuffer);
 		PassParameters->SequenceTransformBuffer      = GraphBuilder.CreateSRV(SequenceTransformBuffer);
-		PassParameters->EvaluateBlockHeaderBuffer    = GraphBuilder.CreateSRV(MultiPassEvaluateHeaderBuffer);
-		PassParameters->LocalSpaceTransformBufferUAV = GraphBuilder.CreateUAV(LocalSpaceTransformBuffer);
+		PassParameters->LayerHeaderBuffer            = GraphBuilder.CreateSRV(LayerHeaderBuffer);
+		PassParameters->SamplePoseBuffer             = GraphBuilder.CreateSRV(SamplePoseBuffer);
+		PassParameters->BoneMapBuffer                = GraphBuilder.CreateSRV(BoneHierarchyBuffer);
+		PassParameters->BoneMaskBuffer               = GraphBuilder.CreateSRV(BoneMaskBuffer);
+		PassParameters->SkinningTransformBufferUAV   = GetCompressedBoneTransformUAV(GraphBuilder, Context.TransformBuffer);
 
 		GraphBuilder.AddPass(
-			RDG_EVENT_NAME("EvaluateAnimSequence"),
+			RDG_EVENT_NAME("EvaluateConcatenateScatter"),
 			PassParameters,
 			ERDGPassFlags::Compute,
-			[ShaderMap, PassParameters, &MultiPassDispatchGroups](FRDGAsyncTask, FRHIComputeCommandList& RHICmdList)
+			[ShaderMap, PassParameters, &DispatchGroups](FRDGAsyncTask, FRHIComputeCommandList& RHICmdList)
 		{
-			for (const FMultiPassDispatchGroup& Group : MultiPassDispatchGroups)
-			{
-				if (Group.NumEvaluateHeaders > 0)
-				{
-					const uint32 NumBlocksPerInstance = FMath::DivideAndRoundUp(Group.NumCompactBones, FEvaluateAnimSequenceCS::BonesPerGroup);
-					TShaderMapRef<FEvaluateAnimSequenceCS> ComputeShader(ShaderMap);
-					PassParameters->HeaderBaseOffset = Group.EvaluateHeaderBaseOffset;
-					FComputeShaderUtils::Dispatch(RHICmdList, ComputeShader, *PassParameters, FIntVector(NumBlocksPerInstance, Group.NumEvaluateHeaders, 1));
-				}
-			}
-		});
-	}
+			RHICmdList.BeginUAVOverlap();
 
-	// 3-pass blend evaluate dispatches
-	if (MultiPassBlendEvaluateHeaderBuffer)
-	{
-		FEvaluateAnimSequenceBlendCS::FParameters* PassParameters = GraphBuilder.AllocParameters<FEvaluateAnimSequenceBlendCS::FParameters>();
-		PassParameters->RetargetingDataBuffer        = GraphBuilder.CreateSRV(RetargetingDataBuffer);
-		PassParameters->SequenceTransformBuffer      = GraphBuilder.CreateSRV(SequenceTransformBuffer);
-		PassParameters->EvaluateBlockHeaderBuffer    = GraphBuilder.CreateSRV(MultiPassBlendEvaluateHeaderBuffer);
-		PassParameters->LocalSpaceTransformBufferUAV = GraphBuilder.CreateUAV(LocalSpaceTransformBuffer);
-
-		GraphBuilder.AddPass(
-			RDG_EVENT_NAME("EvaluateAnimSequenceBlend"),
-			PassParameters,
-			ERDGPassFlags::Compute,
-			[ShaderMap, PassParameters, &MultiPassDispatchGroups](FRDGAsyncTask, FRHIComputeCommandList& RHICmdList)
-		{
-			for (const FMultiPassDispatchGroup& Group : MultiPassDispatchGroups)
+			for (const auto& [Key, Group] : DispatchGroups)
 			{
-				if (Group.NumBlendEvaluateHeaders > 0)
-				{
-					const uint32 NumBlocksPerInstance = FMath::DivideAndRoundUp(Group.NumCompactBones, FEvaluateAnimSequenceCS::BonesPerGroup);
-					TShaderMapRef<FEvaluateAnimSequenceBlendCS> ComputeShader(ShaderMap);
-					PassParameters->HeaderBaseOffset = Group.BlendEvaluateHeaderBaseOffset;
-					FComputeShaderUtils::Dispatch(RHICmdList, ComputeShader, *PassParameters, FIntVector(NumBlocksPerInstance, Group.NumBlendEvaluateHeaders, 1));
-				}
+				FEvaluateConcatenateScatterCS::FPermutationDomain Permutation;
+				Permutation.Set<FEvaluateConcatenateScatterCS::FBonesPerGroupDim>(Group.GroupSize);
+				Permutation.Set<FEvaluateConcatenateScatterCS::FSingleLayerDim>(Group.NumLayers == 1);
+				Permutation.Set<FEvaluateConcatenateScatterCS::FMeshSpaceBlend>(Group.MeshSpaceMask != 0);
+				TShaderMapRef<FEvaluateConcatenateScatterCS> ComputeShader(ShaderMap, Permutation);
+
+				PassParameters->NumTotalHeaders  = Group.NumHeaders;
+				PassParameters->HeaderBaseOffset = Group.HeaderBaseOffset;
+				PassParameters->SlotsPerInstance  = Group.SlotsPerInstance;
+				PassParameters->NumCompactBones  = Group.NumCompactBones;
+				PassParameters->NumLayers        = Group.NumLayers;
+				PassParameters->MeshSpaceMask    = Group.MeshSpaceMask;
+				PassParameters->LayerModeMask    = Group.LayerModeMask;
+
+				const uint32 NumGroups = FMath::DivideAndRoundUp(Group.NumHeaders, Group.InstancesPerGroup);
+				FComputeShaderUtils::Dispatch(RHICmdList, ComputeShader, *PassParameters, FIntVector(NumGroups, 1, 1));
 			}
-		});
-	}
 
-	// 3-pass dispatches: Concatenate (one RDG pass for all multi-pass groups)
-	if (!MultiPassDispatchGroups.IsEmpty() && LocalSpaceTransformBuffer)
-	{
-		FConcatenateLocalToGlobalCS::FParameters* PassParameters = GraphBuilder.AllocParameters<FConcatenateLocalToGlobalCS::FParameters>();
-		PassParameters->LocalSpaceTransformBuffer      = GraphBuilder.CreateSRV(LocalSpaceTransformBuffer);
-		PassParameters->ConcatenateBlockHeaderBuffer   = GraphBuilder.CreateSRV(ConcatenateHeaderBuffer);
-		PassParameters->LeafBonePathBuffer             = GraphBuilder.CreateSRV(LeafBonePathBuffer);
-		PassParameters->GlobalSpaceTransformBufferUAV  = GraphBuilder.CreateUAV(GlobalSpaceTransformBuffer);
-
-		GraphBuilder.AddPass(
-			RDG_EVENT_NAME("ConcatenateLeafBoneTransforms"),
-			PassParameters,
-			ERDGPassFlags::Compute,
-			[ShaderMap, PassParameters, &MultiPassDispatchGroups](FRDGAsyncTask, FRHIComputeCommandList& RHICmdList)
-		{
-			for (const FMultiPassDispatchGroup& Group : MultiPassDispatchGroups)
-			{
-				const uint32 NumTotalHeaders = Group.NumEvaluateHeaders + Group.NumBlendEvaluateHeaders;
-				if (NumTotalHeaders == 0)
-				{
-					continue;
-				}
-
-				const uint32 NumBlocksPerInstance = FMath::DivideAndRoundUp(Group.NumLeafBones, FEvaluateAnimSequenceCS::BonesPerGroup);
-
-				TShaderMapRef<FConcatenateLocalToGlobalCS> ComputeShader(ShaderMap);
-				PassParameters->HeaderBaseOffset = Group.ConcatenateHeaderBaseOffset;
-				FComputeShaderUtils::Dispatch(RHICmdList, ComputeShader, *PassParameters, FIntVector(NumBlocksPerInstance, NumTotalHeaders, 1));
-			}
+			RHICmdList.EndUAVOverlap();
 		});
 	}
-
-	// 3-pass dispatches: Scatter (one RDG pass for all multi-pass groups)
-	if (!MultiPassDispatchGroups.IsEmpty() && GlobalSpaceTransformBuffer)
-	{
-		FScatterTransformsCS::FParameters* PassParameters = GraphBuilder.AllocParameters<FScatterTransformsCS::FParameters>();
-		PassParameters->GlobalSpaceTransformBuffer   = GraphBuilder.CreateSRV(GlobalSpaceTransformBuffer);
-		PassParameters->SequenceTransformBuffer      = GraphBuilder.CreateSRV(SequenceTransformBuffer);
-		PassParameters->ScatterBlockHeaderBuffer     = GraphBuilder.CreateSRV(ScatterHeaderBuffer);
-		PassParameters->BoneMapBuffer                = GraphBuilder.CreateSRV(LeafBonePathBuffer);
-		PassParameters->SkinningTransformBufferUAV   = GetCompressedBoneTransformUAV(GraphBuilder, Context.TransformBuffer);
-
-		GraphBuilder.AddPass(
-			RDG_EVENT_NAME("ScatterTransforms"),
-			PassParameters,
-			ERDGPassFlags::Compute,
-			[ShaderMap, PassParameters, &MultiPassDispatchGroups](FRDGAsyncTask, FRHIComputeCommandList& RHICmdList)
-		{
-			for (const FMultiPassDispatchGroup& Group : MultiPassDispatchGroups)
-			{
-				const uint32 NumTotalHeaders = Group.NumEvaluateHeaders + Group.NumBlendEvaluateHeaders;
-				if (NumTotalHeaders == 0)
-				{
-					continue;
-				}
-
-				const uint32 NumBlocksPerInstance = FMath::DivideAndRoundUp(Group.MaxBoneTransformCount, FEvaluateAnimSequenceCS::BonesPerGroup);
-
-				TShaderMapRef<FScatterTransformsCS> ComputeShader(ShaderMap);
-				PassParameters->HeaderBaseOffset = Group.ScatterHeaderBaseOffset;
-				FComputeShaderUtils::Dispatch(RHICmdList, ComputeShader, *PassParameters, FIntVector(NumBlocksPerInstance, NumTotalHeaders, 1));
-			}
-		});
-	}
-}
-
-void FAnimSequenceTransformProvider::ProvideCPUTransforms(FSkinningTransformProvider::FProviderContext& Context)
-{
-	TRACE_CPUPROFILER_EVENT_SCOPE(FAnimSequenceTransformProvider::ProvideCPUTransforms);
-	FRDGBuilder& GraphBuilder = Context.GraphBuilder;
-
-	uint32 GlobalTransformCount = 0;
-
-	for (const FSkinningTransformProvider::FProviderIndirection Indirection : Context.Indirections)
-	{
-		const FInstancedSkinningSceneExtensionProxy* ExtensionProxy = static_cast<FInstancedSkinningSceneExtensionProxy*>(Context.Proxies[Indirection.Index]);
-		FAnimSequenceTransformProviderProxy* TransformProviderProxy = static_cast<FAnimSequenceTransformProviderProxy*>(ExtensionProxy->GetTransformProviderProxy());
-
-		const uint32 TransformCount = ExtensionProxy->GetMaxBoneTransformCount();
-		const uint32 ActiveAnimationCount = TransformProviderProxy->GetTracks().GetNumActiveTracks();
-		GlobalTransformCount += ActiveAnimationCount * TransformCount * 2; // Current and Previous
-	}
-
-	if (GlobalTransformCount == 0)
-	{
-		return;
-	}
-
-	Context.ScatterUploadBuilder.AddPass(GraphBuilder, SkinningTransformUploadBuffer, Context.TransformBuffer, GlobalTransformCount, sizeof(FCompressedBoneTransform), TEXT("Skinning.AnimTransforms"),
-		[this, Indirections = Context.Indirections, Proxies = Context.Proxies, CurrentTime = Context.CurrentTime, GlobalTransformCount] (FRDGScatterUploader& ScatterUploader)
-	{
-		TRACE_CPUPROFILER_EVENT_SCOPE(FAnimSequenceTransformProvider::ProvideCPUTransformsTask);
-
-		BoneTrackArray TrackToBoneIndexMap;
-		TArray<FTransform, SceneRenderingAllocator> SampledLocalAnimPose;
-		TBitArray<SceneRenderingBitArrayAllocator> SampledLocalAnimPoseBits;
-
-		const auto DecompressPose = [&TrackToBoneIndexMap, &SampledLocalAnimPose, &SampledLocalAnimPoseBits] (
-			  const FBoneTables& BoneTables
-			, const FMapping& Mapping
-			, const FSequence& Sequence
-			, float TimeOffset
-			, TArray<FTransform3f, SceneRenderingAllocator>& SampledGlobalMeshPose
-		)
-		{
-			const UAnimSequence* AnimSequence = Sequence.Sequence;
-			UAnimSequence::FScopedCompressedAnimSequence CompressedAnimSequence = AnimSequence->GetCompressedData();
-			const FCompressedAnimSequence& PlatformCompressedData = CompressedAnimSequence.Get();
-
-			FAnimSequenceDecompressionContext DecompressionContext(
-				  Sequence.SampleRate
-				, Sequence.NumFrames
-				, AnimSequence->Interpolation
-				, AnimSequence->GetRetargetTransformsSourceName()
-				, *PlatformCompressedData.CompressedDataStructure
-				, AnimSequence->GetSkeleton()->GetRefLocalPoses()
-				, PlatformCompressedData.CompressedTrackToSkeletonMapTable
-				, AnimSequence->GetSkeleton()
-				, AnimSequence->IsValidAdditive()
-				, AnimSequence->GetAdditiveAnimType()
-			);
-
-			const uint32 NumCompressedTracks = PlatformCompressedData.CompressedTrackToSkeletonMapTable.Num();
-			TrackToBoneIndexMap.Reset(NumCompressedTracks);
-			
-			TConstArrayView<FTransform> AnimRefLocalPoses = DecompressionContext.GetRefLocalPoses();
-			SampledLocalAnimPose.SetNumUninitialized(BoneTables.NumCompactBones);
-			SampledLocalAnimPoseBits.Init(false, BoneTables.NumCompactBones);
-
-			for (uint32 CompressedTrackIndex = 0; CompressedTrackIndex < NumCompressedTracks; ++CompressedTrackIndex)
-			{
-				const uint16 AnimBoneIndex = PlatformCompressedData.GetSkeletonIndexFromTrackIndex(CompressedTrackIndex);
-
-				// We only care about raw bones.
-				if (AnimBoneIndex < Sequence.NumAnimBones)
-				{
-					const uint16 CompactBoneIndex = Mapping.RetargetingTables.AnimToCompact[AnimBoneIndex];
-
-					if (CompactBoneIndex != INVALID_BONE_INDEX)
-					{
-						TrackToBoneIndexMap.Add(BoneTrackPair(CompactBoneIndex, CompressedTrackIndex));
-						SampledLocalAnimPose[CompactBoneIndex] = AnimRefLocalPoses[AnimBoneIndex];
-						SampledLocalAnimPoseBits[CompactBoneIndex] = true;
-					}
-				}
-			}
-
-			DecompressionContext.Seek((double)TimeOffset);
-
-			if (PlatformCompressedData.BoneCompressionCodec)
-			{
-				TArrayView<FTransform> SampledLocalAnimPoseView(SampledLocalAnimPose);
-				PlatformCompressedData.BoneCompressionCodec->DecompressPose(DecompressionContext, TrackToBoneIndexMap, TrackToBoneIndexMap, TrackToBoneIndexMap, SampledLocalAnimPoseView);
-			}
-
-			const FReferenceSkeleton& MeshRefSkeleton = Mapping.SkinnedAsset->GetRefSkeleton();
-			TConstArrayView<FTransform> MeshLocalRefPoses = MeshRefSkeleton.GetRawRefBonePose();
-
-			for (uint32 CompactBoneIndex = 0; CompactBoneIndex < Mapping.NumCompactBones; ++CompactBoneIndex)
-			{
-				const uint32 AnimBoneIndex = Mapping.RetargetingTables.CompactToAnim[CompactBoneIndex];
-				const uint32 MeshBoneIndex = BoneTables.CompactToMesh[CompactBoneIndex];
-
-				FTransform3f MeshLocalTransform;
-
-				if (AnimBoneIndex != INVALID_BONE_INDEX)
-				{
-					const FTransform3f RetargetTransform = Mapping.RetargetingTables.Transform[CompactBoneIndex];
-
-					if (SampledLocalAnimPoseBits[CompactBoneIndex])
-					{
-						MeshLocalTransform = FTransform3f(SampledLocalAnimPose[CompactBoneIndex]);
-					}
-					else
-					{
-						MeshLocalTransform = FTransform3f(AnimRefLocalPoses[AnimBoneIndex]);
-					}
-
-					if (Mapping.RetargetingTables.bRequiresRefPoseRetarget && CompactBoneIndex != 0)
-					{
-						const FQuat4f& Q0 = Mapping.RetargetingTables.RefPoseRetargetQuats[CompactBoneIndex].Get<0>();
-						const FQuat4f& Q1 = Mapping.RetargetingTables.RefPoseRetargetQuats[CompactBoneIndex].Get<1>();
-
-						FQuat4f RetargetedRotation = Q0 * MeshLocalTransform.GetRotation() * Q1;
-						RetargetedRotation.Normalize();
-						MeshLocalTransform.SetRotation(RetargetedRotation);
- 
-						if (Mapping.RetargetingTables.Mode[CompactBoneIndex] != EBoneTranslationRetargetingMode::Skeleton)
-						{
-							MeshLocalTransform.SetTranslation(Q0.RotateVector(MeshLocalTransform.GetTranslation()));
-						}
-					}
-
-					switch (Mapping.RetargetingTables.Mode[CompactBoneIndex])
-					{
-						case EBoneTranslationRetargetingMode::Animation:
-						{
-							// Identity
-						}
-						break;
-						case EBoneTranslationRetargetingMode::Skeleton:
-						{
-							MeshLocalTransform.SetTranslation(RetargetTransform.GetTranslation());
-							if (CompactBoneIndex == 0)
-							{
-								MeshLocalTransform.SetRotation(RetargetTransform.GetRotation());
-							}
-						}
-						break;
-						case EBoneTranslationRetargetingMode::AnimationScaled:
-						{
-							MeshLocalTransform.ScaleTranslation(RetargetTransform.GetScale3D().X);
-						}
-						break;
-						case EBoneTranslationRetargetingMode::AnimationRelative:
-						{
-							MeshLocalTransform.SetRotation(MeshLocalTransform.GetRotation() * RetargetTransform.GetRotation());
-							MeshLocalTransform.SetTranslation(MeshLocalTransform.GetTranslation() + RetargetTransform.GetTranslation());
-							MeshLocalTransform.SetScale3D(FVector3f(MeshLocalTransform.GetScale3D().X * RetargetTransform.GetScale3D().X));
-							MeshLocalTransform.NormalizeRotation();
-						}
-						break;
-						case EBoneTranslationRetargetingMode::OrientAndScale:
-						{
-							MeshLocalTransform.SetTranslation(RetargetTransform.GetRotation().RotateVector(MeshLocalTransform.GetTranslation()) * RetargetTransform.GetScale3D().X);
-							break;
-						}
-					}
-				}
-				else
-				{
-					MeshLocalTransform = FTransform3f(MeshLocalRefPoses[MeshBoneIndex]);
-				}
-
-				const uint32 ParentIndex = BoneTables.CompactToParent[CompactBoneIndex];
-
-				if (SampledGlobalMeshPose.IsValidIndex(ParentIndex) && ParentIndex < CompactBoneIndex)
-				{
-					SampledGlobalMeshPose[CompactBoneIndex] = MeshLocalTransform * SampledGlobalMeshPose[ParentIndex];
-				}
-				else
-				{
-					SampledGlobalMeshPose[CompactBoneIndex] = MeshLocalTransform;
-				}
-			}
-		};
-
-		TArray<FTransform3f, SceneRenderingAllocator> SourceSampledGlobalMeshPose;
-		TArray<FTransform3f, SceneRenderingAllocator> TargetSampledGlobalMeshPose;
-		TArray<FCompressedBoneTransform, SceneRenderingAllocator> SourceTransforms;
-		TBitArray<SceneRenderingBitArrayAllocator> SourceTransformMap;
-
-		const auto EvaluateAnimSequence = [&SourceTransforms, &SourceTransformMap, &SourceSampledGlobalMeshPose, &TargetSampledGlobalMeshPose, &DecompressPose] (
-			const FBoneTables& BoneTables,
-			const FInstancedSkinningSceneExtensionProxy* ExtensionProxy,
-			TArrayView<FCompressedBoneTransform> Transforms,
-			TArrayView<FCompressedBoneTransform> DuplicatePreviousTransforms,
-			const FSequence& SourceSequence,
-			const FSequence& TargetSequence,
-			const FMapping& SourceMapping,
-			const FMapping& TargetMapping,
-			float SourceTimeOffset,
-			float TargetTimeOffset,
-			float BlendWeight)
-		{
-			DecompressPose(BoneTables, SourceMapping, SourceSequence, SourceTimeOffset, SourceSampledGlobalMeshPose);
-
-			if (BlendWeight > 0.0f)
-			{
-				DecompressPose(BoneTables, TargetMapping, TargetSequence, TargetTimeOffset, TargetSampledGlobalMeshPose);
-			}
-
-			for (uint32 CompactBoneIndex = 0; CompactBoneIndex < TargetMapping.NumCompactBones; ++CompactBoneIndex)
-			{
-				if (BlendWeight > 0.0f)
-				{
-					SourceSampledGlobalMeshPose[CompactBoneIndex].BlendWith(TargetSampledGlobalMeshPose[CompactBoneIndex], BlendWeight);
-					SourceSampledGlobalMeshPose[CompactBoneIndex].NormalizeRotation();
-				}
-
-				SourceSampledGlobalMeshPose[CompactBoneIndex] = TargetMapping.InvGlobalRefPoseTransforms[CompactBoneIndex] * SourceSampledGlobalMeshPose[CompactBoneIndex];
-			}
-
-			if (!ExtensionProxy->GetBoneMap().IsEmpty())
-			{
-				SourceTransformMap.Init(false, TargetMapping.NumCompactBones);
-				SourceTransforms.SetNumUninitialized(TargetMapping.NumCompactBones, EAllowShrinking::No);
-				FCompressedBoneTransform IdentityTransform;
-				SetCompressedBoneTransformIdentity(IdentityTransform);
-				int32 DstTransformIndex = 0;
-
-				for (uint32 MeshBoneIndex : ExtensionProxy->GetBoneMap())
-				{
-					const uint32 CompactBoneIndex = BoneTables.MeshToCompact[MeshBoneIndex];
-					FCompressedBoneTransform* OutputTransform;
-
-					if (CompactBoneIndex != INVALID_BONE_INDEX)
-					{
-						OutputTransform = &SourceTransforms[CompactBoneIndex];
-
-						if (!SourceTransformMap[CompactBoneIndex])
-						{
-							SourceTransformMap[CompactBoneIndex] = true;
-							StoreCompressedBoneTransform(OutputTransform, SourceSampledGlobalMeshPose[CompactBoneIndex].ToMatrixWithScale());
-						}
-					}
-					else
-					{
-						OutputTransform = &IdentityTransform;
-					}
-
-					Transforms[DstTransformIndex] = *OutputTransform;
-
-					if (!DuplicatePreviousTransforms.IsEmpty())
-					{
-						DuplicatePreviousTransforms[DstTransformIndex] = *OutputTransform;
-					}
-
-					DstTransformIndex++;
-				}
-			}
-			else
-			{
-				for (uint32 MeshBoneIndex = 0; MeshBoneIndex < TargetMapping.NumMeshBones; ++MeshBoneIndex)
-				{
-					const uint32 CompactBoneIndex = BoneTables.MeshToCompact[MeshBoneIndex];
-					FCompressedBoneTransform OutputTransform;
-
-					if (CompactBoneIndex != INVALID_BONE_INDEX)
-					{
-						StoreCompressedBoneTransform(OutputTransform, SourceSampledGlobalMeshPose[CompactBoneIndex].ToMatrixWithScale());
-					}
-					else
-					{
-						SetCompressedBoneTransformIdentity(OutputTransform);
-					}
-
-					Transforms[MeshBoneIndex] = OutputTransform;
-
-					if (!DuplicatePreviousTransforms.IsEmpty())
-					{
-						DuplicatePreviousTransforms[MeshBoneIndex] = OutputTransform;
-					}
-				}
-			}
-		};
-
-		const float GlobalTimeScale = CVarAnimSequenceTransformProviderTimeScale.GetValueOnRenderThread();
-		const double WorldTime = CurrentTime.GetWorldTimeSeconds() * GlobalTimeScale;
-		const float DeltaTime = CurrentTime.GetDeltaWorldTimeSeconds() * GlobalTimeScale;
-
-		uint32 NumActiveTracks = 0;
-		uint32 NumCurrentTransformUpdates = 0;
-		uint32 NumPreviousTransformUpdates = 0;
-
-		for (const FSkinningTransformProvider::FProviderIndirection Indirection : Indirections)
-		{
-			const FInstancedSkinningSceneExtensionProxy* ExtensionProxy = static_cast<FInstancedSkinningSceneExtensionProxy*>(Proxies[Indirection.Index]);
-			FProxyUserData& ProxyUserData = static_cast<FProxyUserData&>(ExtensionProxy->GetUserDataChecked());
-			FAnimSequenceTrackRenderPool& Tracks = ProxyUserData.Tracks;
-			ProxyUserData.TrackHistories.SetNum(Tracks.GetNumTracks());
-
-			const uint32 LODIndex     = ExtensionProxy->GetLOD();
-			FProxyUserData::FLOD& LOD = ProxyUserData.LODs[LODIndex];
-
-			const uint32 UniqueAnimationCount   = ExtensionProxy->GetUniqueAnimationCount();
-			const uint32 NumTransforms          = ExtensionProxy->GetMaxBoneTransformCount();
-			const uint32 NumTotalTransforms     = NumTransforms * 2;
-			const uint32 CurrentTransformIndex  = Indirection.CurrentTransformOffset  / sizeof(FCompressedBoneTransform);
-			const uint32 PreviousTransformIndex = Indirection.PreviousTransformOffset / sizeof(FCompressedBoneTransform);
-			const uint32 CurrentTransformSlot   = CurrentTransformIndex < PreviousTransformIndex ? 0 : 1;
-
-			SourceSampledGlobalMeshPose.SetNumUninitialized(LOD.BoneTables.NumCompactBones, EAllowShrinking::No);
-			TargetSampledGlobalMeshPose.SetNumUninitialized(LOD.BoneTables.NumCompactBones, EAllowShrinking::No);
-	
-			FMemMark Mark(FMemStack::Get());
-
-			Tracks.EnumerateActiveTracks([&] (int32 TrackIndex)
-			{
-				FAnimSequenceTrackRenderData& Data = Tracks.GetData(TrackIndex);
-
-				const FTrackData SourceTrack = ProxyUserData.GetTrackData(Data.GetSource(), WorldTime, LODIndex);
-				const FTrackData TargetTrack = ProxyUserData.GetTrackData(Data.GetTarget(), WorldTime, LODIndex);
-
-				// Dirty previous transforms means the whole allocation is dirty. Rebuild everything.
-				const bool bDirtyTransformAllocation = EnumHasAnyFlags(Indirection.DirtyBoneTransforms, EDirtyBoneTransforms::Previous);
-				const FTrackHistory::FUpdateResult UpdateResult = ProxyUserData.TrackHistories[TrackIndex].Update(TargetTrack, CurrentTransformSlot, bDirtyTransformAllocation);
-	
-				const float BlendWeight = Data.GetBlendWeight();
-
-				EPreviousBoneTransformUpdateMode UpdateMode = EPreviousBoneTransformUpdateMode::None;
-				TArrayView<FCompressedBoneTransform> DstTransformsPrevious;
-
-				if (UpdateResult.bNeedsPrevious)
-				{
-					UpdateMode = UpdateResult.bNeedsCurrent && UpdateResult.bNeedsPrevious && TargetTrack.Position == TargetTrack.PreviousPosition
-						? EPreviousBoneTransformUpdateMode::DuplicateCurrentToPrevious
-						: EPreviousBoneTransformUpdateMode::UpdatePrevious;
-
-					DstTransformsPrevious = ScatterUploader.Add_GetRef<FCompressedBoneTransform>(PreviousTransformIndex + TrackIndex * NumTotalTransforms, NumTransforms);
-				}
-
-				if (UpdateResult.bNeedsCurrent)
-				{
-					TArrayView<FCompressedBoneTransform> DstTransforms = ScatterUploader.Add_GetRef<FCompressedBoneTransform>(CurrentTransformIndex + TrackIndex * NumTotalTransforms, NumTransforms);
-					EvaluateAnimSequence(
-						LOD.BoneTables,
-						ExtensionProxy,
-						DstTransforms,
-						UpdateMode == EPreviousBoneTransformUpdateMode::DuplicateCurrentToPrevious ? DstTransformsPrevious : TArrayView<FCompressedBoneTransform>{},
-						SourceTrack.Sequence, TargetTrack.Sequence,
-						SourceTrack.Mapping, TargetTrack.Mapping,
-						SourceTrack.Position, TargetTrack.Position,
-						BlendWeight);
-
-					NumCurrentTransformUpdates++;
-				}
-
-				if (UpdateMode == EPreviousBoneTransformUpdateMode::UpdatePrevious)
-				{
-					EvaluateAnimSequence(
-						LOD.BoneTables,
-						ExtensionProxy,
-						DstTransformsPrevious,
-						{},
-						SourceTrack.Sequence, TargetTrack.Sequence,
-						SourceTrack.Mapping, TargetTrack.Mapping,
-						SourceTrack.PreviousPosition, TargetTrack.PreviousPosition,
-						BlendWeight);
-
-					NumPreviousTransformUpdates++;
-				}
-
-				Data.Tick(DeltaTime, SourceTrack.Position, TargetTrack.Position);
-				NumActiveTracks++;
-			});
-		}
-
-		INC_DWORD_STAT_BY(STAT_AnimSequenceTransformProvider_ActiveTracks, NumActiveTracks);
-		INC_DWORD_STAT_BY(STAT_AnimSequenceTransformProvider_CurrentTransformUpdates, NumCurrentTransformUpdates);
-		INC_DWORD_STAT_BY(STAT_AnimSequenceTransformProvider_PreviousTransformUpdates, NumPreviousTransformUpdates);
-	});
-}
+}
\ No newline at end of file
diff --git a/Engine/Source/Runtime/Renderer/Private/Skinning/AnimSequenceTransformProvider.h b/Engine/Source/Runtime/Renderer/Private/Skinning/AnimSequenceTransformProvider.h
--- a/Engine/Source/Runtime/Renderer/Private/Skinning/AnimSequenceTransformProvider.h
+++ b/Engine/Source/Runtime/Renderer/Private/Skinning/AnimSequenceTransformProvider.h
@@ -19,6 +19,9 @@ class FAnimSequenceTransformProvider
 {
 	DECLARE_SCENE_EXTENSION(RENDERER_API, FAnimSequenceTransformProvider);
 
+	struct FDispatchGroupKey;
+	struct FDispatchGroup;
+	struct FTrackUpdateRequest;
 public:
 	FAnimSequenceTransformProvider(FScene& InScene);
 
@@ -29,24 +32,64 @@ public:
 	void RegisterProxy(FAnimSequenceTransformProviderProxy* Proxy) override;
 	void UnregisterProxy(FAnimSequenceTransformProviderProxy* Proxy) override;
 
-private:
-	void ProvideGPUTransforms(FSkinningTransformProvider::FProviderContext& Context);
-	void ProvideCPUTransforms(FSkinningTransformProvider::FProviderContext& Context);
+	struct FBoneTransformWithUniformScale
+	{
+		float Rotation[4];
+		float Translation[3];
+		float Scale;
+
+		void SetTransform(const FTransform3f& Xfm)
+		{
+			SetRotation(Xfm.GetRotation());
+			SetTranslation(Xfm.GetTranslation());
+			SetScale(Xfm.GetScale3D().X);
+		}
+
+		void SetRotation(FQuat4f Q)
+		{
+			Rotation[0] = Q.X;
+			Rotation[1] = Q.Y;
+			Rotation[2] = Q.Z;
+			Rotation[3] = Q.W;
+		}
+
+		void SetTranslation(FVector3f P)
+		{
+			Translation[0] = P.X;
+			Translation[1] = P.Y;
+			Translation[2] = P.Z;
+		}
 
-	static void ComputeSinglePassGroupParams(
+		void SetScale(float S)
+		{
+			Scale = S;
+		}
+	};
+
+private:
+	void ProvideTransforms(FSkinningTransformProvider::FProviderContext& Context);
+	void UploadSequences(FRDGScatterUploader& ScatterUploader, int32& NumTotalFramesRequested);
+
+	static void UploadBlockHeaders(
+		FRHICommandList& RHICmdList,
+		double WorldTime,
+		float DeltaTime,
+		TConstArrayView<FSkinningTransformProvider::FProviderIndirection> Indirections,
+		TConstArrayView<FSkinningSceneExtensionProxy*> Proxies,
+		int32& NumTotalFramesRequested,
+		TConstArrayView<FTrackUpdateRequest> TrackUpdateRequests,
+		FRDGBuffer* LayerHeaderBuffer,
+		FRDGBuffer* SamplePoseBuffer,
+		uint32 NumTotalEntries,
+		uint32 NumTotalSamples,
+		TMap<FDispatchGroupKey, FDispatchGroup>& DispatchGroups);
+
+	static void ComputeGroupParams(
 		uint32 NumCompactBones,
 		uint32& OutGroupSize,
 		uint32& OutInstancesPerGroup,
 		uint32& OutSlotsPerInstance)
 	{
-		if (NumCompactBones > 256)
-		{
-			OutGroupSize = 0;
-			OutInstancesPerGroup = 1;
-			OutSlotsPerInstance = 0;
-			return;
-		}
-
 		OutSlotsPerInstance = FMath::RoundUpToPowerOfTwo(NumCompactBones);
 
 		if (NumCompactBones <= 64)
@@ -57,10 +100,18 @@ private:
 		{
 			OutGroupSize = 128;
 		}
-		else
+		else if (NumCompactBones <= 256)
 		{
 			OutGroupSize = 256;
 		}
+		else if (NumCompactBones <= 512)
+		{
+			OutGroupSize = 512;
+		}
+		else
+		{
+			OutGroupSize = 1024;
+		}
 
 		OutInstancesPerGroup = OutGroupSize / OutSlotsPerInstance;
 	}
@@ -97,14 +148,6 @@ private:
 		const USkinnedAsset* MeshAsset,
 		const FBoneTables& BoneTables);
 
-	struct FLeafBonePathData
-	{
-		TArray<uint32> Buffer;
-		uint32 NumLeafBones;
-	};
-
-	void BuildLeafBonePathData(FLeafBonePathData& OutLeafBonePathData, const FBoneTables& RemappingTables);
-
 	struct FBufferAllocation
 	{
 		bool IsValid() const
@@ -116,40 +159,6 @@ private:
 		uint32 Count  = 0;
 	};
 
-	struct FBoneTransformWithScale
-	{
-		void SetTransform(const FTransform3f& Xfm)
-		{
-			SetRotation(Xfm.GetRotation());
-			SetTranslation(Xfm.GetTranslation());
-			SetScale(Xfm.GetScale3D().X);
-		}
-
-		void SetRotation(FQuat4f Q)
-		{
-			Rotation[0] = Q.X;
-			Rotation[1] = Q.Y;
-			Rotation[2] = Q.Z;
-			Rotation[3] = Q.W;
-		}
-
-		void SetTranslation(FVector3f P)
-		{
-			Translation[0] = P.X;
-			Translation[1] = P.Y;
-			Translation[2] = P.Z;
-		}
-
-		void SetScale(float S)
-		{
-			Scale = S;
-		}
-
-		float Rotation[4];
-		float Translation[3];
-		float Scale;
-	};
-
 	using FSequenceRegistryKey = const UAnimSequence*;
 
 	struct FMappingRegistryKey
@@ -240,108 +249,61 @@ private:
 	{
 		const FMapping& Mapping;
 		FSequence& Sequence;
-		int32 SequenceIndex;
+		uint16 SequenceIndex;
 		float Position;
-		float PreviousPosition;
-		bool bHasPreviousPosition;
-		bool bAutoPlay;
 	};
 
-	class FTrackHistory
+	class FProxyUserData : public FInstancedSkinningSceneExtensionProxyUserData
 	{
 	public:
-		static constexpr float Invalid = UE_MAX_FLT;
+		static constexpr uint32 InvalidBoneMaskOffset = SKINNING_INVALID_BONE_MASK_OFFSET;
 
-		struct FUpdateResult
-		{
-			uint8 bNeedsCurrent  : 1 = 0;
-			uint8 bNeedsPrevious : 1 = 0;
-		};
+		FProxyUserData(FAnimSequenceTransformProviderProxy* Proxy);
 
-		FUpdateResult Update(const FTrackData& Track, int32 SlotIndex, bool bForceUpdate)
+		FAnimSequenceTransformProviderRenderData* GetRenderData() const
 		{
-			check(SlotIndex == 0 || SlotIndex == 1);
-
-			FUpdateResult Result;
-			Result.bNeedsCurrent  = bForceUpdate || Track.bAutoPlay;
-			Result.bNeedsPrevious = bForceUpdate;
-
-			// Auto-play always rebuilds current transforms. Caching is disabled because blending
-			// between sequences with independent play rates makes reliable change detection impractical.
-			if (!Track.bAutoPlay)
-			{
-				if (SequenceIndex != Track.SequenceIndex)
-				{
-					SequenceIndex = Track.SequenceIndex;
-					Slots = FVector2f{ Invalid, Invalid };
-				}
-
-				Result.bNeedsCurrent  |= Slots[ SlotIndex] != Track.Position;
-				Result.bNeedsPrevious |= Slots[!SlotIndex] != Track.PreviousPosition && Track.bHasPreviousPosition;
-				Slots[SlotIndex] = Track.Position;
-
-				if (Track.bHasPreviousPosition)
-				{
-					Slots[!SlotIndex] = Track.PreviousPosition;
-				}
-			}
-			return Result;
+			return RenderData;
 		}
 
-	private:
-		int32 SequenceIndex = INDEX_NONE;
-		FVector2f Slots{ Invalid, Invalid };
-	};
-
-	class FProxyUserData : public FInstancedSkinningSceneExtensionProxyUserData
-	{
-	public:
-		FProxyUserData(FAnimSequenceTransformProviderProxy* Proxy);
-
-		FSequence& GetSequence(const FAnimSequenceTrackPackedData& Data) const
+		FSequence& GetSequence(uint16 InSequenceIndex) const
 		{
-			return *Sequences[Sequences.IsValidIndex(Data.GetSequenceIndex()) ? Data.GetSequenceIndex() : 0];
+			return *Sequences[Sequences.IsValidIndex(InSequenceIndex) ? InSequenceIndex : 0];
 		}
 
-		const FMapping& GetMapping(const FAnimSequenceTrackPackedData& Data, uint32 LODIndex) const
+		const FMapping& GetMapping(uint16 InSequenceIndex, uint32 LODIndex) const
 		{
 			check(LODIndex >= MinLODLevel && LODIndex < (uint32)LODs.Num());
-			const int32 SequenceIndex = Sequences.IsValidIndex(Data.GetSequenceIndex()) ? Data.GetSequenceIndex() : 0;
-			return *SequenceToMappingTable[SequenceIndex * LODs.Num() + LODIndex];
+			const int32 SeqIndex = Sequences.IsValidIndex(InSequenceIndex) ? InSequenceIndex : 0;
+			return *SequenceToMappingTable[SeqIndex * LODs.Num() + LODIndex];
 		}
 
-		FTrackData GetTrackData(const FAnimSequenceTrackPackedData& Data, double WorldTime, uint32 LODIndex) const
+		FTrackData GetTrackDataForSequence(uint16 InSequenceIndex, double Position, uint32 LODIndex, EAnimSequenceTrackLoopMode LoopMode) const
 		{
-			FSequence& Sequence = GetSequence(Data);
-			const float Position = Data.GetPosition(WorldTime, Sequence.PlayLength);
-			const float PreviousPosition = Data.HasPreviousPosition() ? Data.GetPreviousPosition() : Position;
+			FSequence& Sequence = GetSequence(InSequenceIndex);
+			const float WrappedPosition = FAnimSequenceTrackPackedData::WrapPosition(LoopMode, static_cast<float>(Position), Sequence.PlayLength);
 
 			return FTrackData
 			{
-				  .Mapping  = GetMapping(Data, LODIndex)
+				  .Mapping  = GetMapping(InSequenceIndex, LODIndex)
 				, .Sequence = Sequence
-				, .SequenceIndex = Data.GetSequenceIndex()
-				, .Position = Position
-				, .PreviousPosition = PreviousPosition
-				, .bHasPreviousPosition = Data.HasPreviousPosition()
-				, .bAutoPlay = Data.IsAutoPlay()
+				, .SequenceIndex = InSequenceIndex
+				, .Position = WrappedPosition
 			};
 		}
 
 		struct FLOD
 		{
 			FBoneTables BoneTables;
-			FLeafBonePathData LeafBonePathData;
-			FBufferAllocation LeafBonePathBufferAllocation;
+			FBufferAllocation BoneMapBufferAllocation;
 			int32 BoneMapBufferOffset = 0;
-			bool bSinglePass = false;
 			bool bUploaded = false;
 
-			// Dispatch grouping params (computed at registration)
-			uint32 GroupSize         = 0; // 64, 128, 256 for single-pass; 0 for 3-pass
-			uint32 InstancesPerGroup = 1; // Wave-packing factor (single-pass only)
-			uint32 SlotsPerInstance  = 0; // Power-of-two >= NumCompactBones (single-pass only)
-			void*  DispatchGroup     = nullptr; // Points to per-frame FSinglePassDispatchGroup or FMultiPassDispatchGroup
+			uint32 GroupSize         = 0;
+			uint32 InstancesPerGroup = 1;
+			uint32 SlotsPerInstance  = 0;
+			FDispatchGroup* DispatchGroup = nullptr;
+
+			TArray<uint32> BoneMaskBufferOffsets;
 
 			bool IsValid() const
 			{
@@ -349,8 +311,8 @@ private:
 			}
 		};
 
-		FAnimSequenceTrackRenderPool& Tracks;
-		TArray<FTrackHistory> TrackHistories;
+		FAnimSequenceTransformProviderRenderData* RenderData = nullptr;
+		TArray<const UBlendProfile*> BlendProfiles;
 		const FInstancedSkinningSceneExtensionProxy& SceneExtensionProxy;
 		TArray<FSequence*> Sequences;
 		TArray<FMapping*> Mappings;
@@ -358,7 +320,6 @@ private:
 		TArray<FLOD> LODs;
 		uint32 MinLODLevel = 0;
 		uint32 NumMeshBones = 0;
-		const EBoneTransformStorageMode TransformStorageMode;
 	};
 
 	TMap<FMappingRegistryKey, TUniquePtr<FMapping>> MappingRegistry;
@@ -370,15 +331,26 @@ private:
 
 	FSpanAllocator SequenceTransformAllocator;
 	FRDGAsyncScatterUploadBuffer SequenceTransformUploadBuffer;
-	TPersistentByteAddressBuffer<FBoneTransformWithScale> SequenceTransformPersistentBuffer;
+	TPersistentStructuredBuffer<FVector4f> SequenceTransformPersistentBuffer;
 
 	FSpanAllocator RetargetingDataAllocator;
 	FRDGAsyncScatterUploadBuffer RetargetingDataUploadBuffer;
 	TPersistentByteAddressBuffer<UE::HLSL::FBoneRetargetingData> RetargetingDataPersistentBuffer;
 
-	FSpanAllocator LeafBonePathAllocator;
-	FRDGAsyncScatterUploadBuffer LeafBonePathUploadBuffer;
-	TPersistentByteAddressBuffer<uint32> LeafBonePathPersistentBuffer;
+	FSpanAllocator BoneHierarchyAllocator;
+	FRDGAsyncScatterUploadBuffer BoneHierarchyUploadBuffer;
+	TPersistentByteAddressBuffer<UE::HLSL::FBoneHierarchyData> BoneHierarchyPersistentBuffer;
+
+	FSpanAllocator BoneMaskAllocator;
+	TPersistentByteAddressBuffer<float> BoneMaskPersistentBuffer;
+	FRDGAsyncScatterUploadBuffer BoneMaskUploadBuffer;
+
+	bool bHalfPrecisionSequences = false;
+
+	uint32 GetFloat4sPerBone() const
+	{
+		return bHalfPrecisionSequences ? 1 : 2;
+	}
 
-	FRDGAsyncScatterUploadBuffer SkinningTransformUploadBuffer;
+	void WriteBoneTransform(TArrayView<FVector4f>& DstFloat4s, uint32 Index, const FBoneTransformWithUniformScale& Transform) const;
 };
\ No newline at end of file
