Copy contents of instanced skinned meshes when allocations change. Trades expensive recalculation of previous state for a copy.

#Synced-CL 52519763

diff --git a/Engine/Shaders/Private/Skinning/CopyPreviousTransforms.usf b/Engine/Shaders/Private/Skinning/CopyPreviousTransforms.usf
new file mode 100644
--- /dev/null
+++ b/Engine/Shaders/Private/Skinning/CopyPreviousTransforms.usf
@@ -0,0 +1,75 @@
+// Copyright Epic Games, Inc. All Rights Reserved.
+
+#include "/Engine/Public/Platform.ush"
+#include "/Engine/Shared/SkinningDefinitions.h"
+
+// Copies bone transforms from old buffer offsets to new offsets.
+// Used when ISKM allocations move (defrag or realloc) to preserve motion vectors.
+// Copies the full allocation (both current and previous slots).
+//
+// Two modes:
+//   USE_STAGING_BUFFER=1: Read from staging SRV (defrag — old/new offsets may overlap in live buffer)
+//   USE_STAGING_BUFFER=0: Read/write same UAV (realloc — allocate-before-free guarantees no overlap)
+
+struct FCopyTransformsHeader
+{
+	uint SrcByteOffset;	// Byte offset in source buffer
+	uint DstByteOffset;	// Byte offset in destination buffer
+	uint Count;			// Number of compressed bone transforms to copy
+	uint _Pad0;
+};
+
+FCompressedBoneTransformRWBuffer DstTransformBuffer;
+StructuredBuffer<FCopyTransformsHeader> CopyHeaders;
+
+#if USE_STAGING_BUFFER
+FCompressedBoneTransformBuffer SrcTransformBuffer;
+#endif
+
+// Dispatch: (CeilDiv(MaxElementsAcrossAllEntries, COPY_THREAD_GROUP_SIZE), NumCopyEntries, 1)
+// GroupId.x = element group index, GroupId.y = copy entry index
+
+[numthreads(COPY_THREAD_GROUP_SIZE, 1, 1)]
+void CopyPreviousTransformsCS(uint3 GroupId : SV_GroupID, uint GroupIndex : SV_GroupIndex)
+{
+	const FCopyTransformsHeader Header = CopyHeaders[GroupId.y];
+	const uint ElementIndex = GroupId.x * COPY_THREAD_GROUP_SIZE + GroupIndex;
+
+	if (ElementIndex >= Header.Count)
+	{
+		return;
+	}
+
+#if USE_COMPRESSED_BONE_TRANSFORM
+	// ByteAddressBuffer: 32 bytes per transform
+	const uint SrcOffset = Header.SrcByteOffset + ElementIndex * 32u;
+	const uint DstOffset = Header.DstByteOffset + ElementIndex * 32u;
+
+#if USE_STAGING_BUFFER
+	const uint4 Data0 = SrcTransformBuffer.Load4(SrcOffset);
+	const uint4 Data1 = SrcTransformBuffer.Load4(SrcOffset + 16u);
+#else
+	const uint4 Data0 = DstTransformBuffer.Load4(SrcOffset);
+	const uint4 Data1 = DstTransformBuffer.Load4(SrcOffset + 16u);
+#endif
+	DstTransformBuffer.Store4(DstOffset, Data0);
+	DstTransformBuffer.Store4(DstOffset + 16u, Data1);
+#else
+	// RWBuffer<float4>: 3 float4s (48 bytes) per transform
+	const uint SrcBaseIndex = Header.SrcByteOffset / 16u + ElementIndex * 3u;
+	const uint DstBaseIndex = Header.DstByteOffset / 16u + ElementIndex * 3u;
+
+#if USE_STAGING_BUFFER
+	DstTransformBuffer[DstBaseIndex + 0] = SrcTransformBuffer[SrcBaseIndex + 0];
+	DstTransformBuffer[DstBaseIndex + 1] = SrcTransformBuffer[SrcBaseIndex + 1];
+	DstTransformBuffer[DstBaseIndex + 2] = SrcTransformBuffer[SrcBaseIndex + 2];
+#else
+	float4 T0 = DstTransformBuffer[SrcBaseIndex + 0];
+	float4 T1 = DstTransformBuffer[SrcBaseIndex + 1];
+	float4 T2 = DstTransformBuffer[SrcBaseIndex + 2];
+	DstTransformBuffer[DstBaseIndex + 0] = T0;
+	DstTransformBuffer[DstBaseIndex + 1] = T1;
+	DstTransformBuffer[DstBaseIndex + 2] = T2;
+#endif
+#endif
+}
diff --git a/Engine/Source/Runtime/Renderer/Private/Skinning/SkinningSceneExtension.cpp b/Engine/Source/Runtime/Renderer/Private/Skinning/SkinningSceneExtension.cpp
--- a/Engine/Source/Runtime/Renderer/Private/Skinning/SkinningSceneExtension.cpp
+++ b/Engine/Source/Runtime/Renderer/Private/Skinning/SkinningSceneExtension.cpp
@@ -165,6 +165,49 @@ private:
 
 IMPLEMENT_GLOBAL_SHADER(FRefPoseTransformProviderCS, "/Engine/Private/Skinning/TransformProviders.usf", "RefPoseProviderCS", SF_Compute);
 
+struct FCopyTransformsHeader
+{
+	uint32 SrcByteOffset;
+	uint32 DstByteOffset;
+	uint32 Count;
+	uint32 _Pad0;
+};
+
+class FCopyPreviousTransformsCS : public FGlobalShader
+{
+public:
+	static constexpr uint32 ThreadGroupSize = 64u;
+
+	class FUseStagingBuffer : SHADER_PERMUTATION_BOOL("USE_STAGING_BUFFER");
+	using FPermutationDomain = TShaderPermutationDomain<FUseStagingBuffer>;
+
+private:
+	DECLARE_GLOBAL_SHADER(FCopyPreviousTransformsCS);
+	SHADER_USE_PARAMETER_STRUCT(FCopyPreviousTransformsCS, FGlobalShader);
+
+	BEGIN_SHADER_PARAMETER_STRUCT(FParameters, )
+		SHADER_PARAMETER_RDG_COMPRESSED_BONE_TRANSFORM_UAV(DstTransformBuffer)
+		SHADER_PARAMETER_RDG_COMPRESSED_BONE_TRANSFORM_SRV(SrcTransformBuffer)
+		SHADER_PARAMETER_RDG_BUFFER_SRV(StructuredBuffer<FCopyTransformsHeader>, CopyHeaders)
+	END_SHADER_PARAMETER_STRUCT()
+
+	static bool ShouldCompilePermutation(const FGlobalShaderPermutationParameters& Parameters)
+	{
+		return true;
+	}
+
+	static void ModifyCompilationEnvironment(const FGlobalShaderPermutationParameters& Parameters, FShaderCompilerEnvironment& OutEnvironment)
+	{
+		FGlobalShader::ModifyCompilationEnvironment(Parameters, OutEnvironment);
+		OutEnvironment.CompilerFlags.Add(CFLAG_WarningsAsErrors);
+		OutEnvironment.CompilerFlags.Add(CFLAG_HLSL2021);
+		OutEnvironment.CompilerFlags.Add(CFLAG_AllowTypedUAVLoads);
+		OutEnvironment.SetDefine(TEXT("COPY_THREAD_GROUP_SIZE"), ThreadGroupSize);
+	}
+};
+
+IMPLEMENT_GLOBAL_SHADER(FCopyPreviousTransformsCS, "/Engine/Private/Skinning/CopyPreviousTransforms.usf", "CopyPreviousTransformsCS", SF_Compute);
+
 static FGuid RefPoseProviderId(REF_POSE_TRANSFORM_PROVIDER_GUID);
 static FGuid MeshObjectProviderId(ANIM_MESH_OBJECT_TRANSFORM_PROVIDER_GUID);
 
@@ -705,6 +748,9 @@ bool FSkinningSceneExtension::ProcessBufferDefragmentation()
 	{
 		if (Data.TransformBufferOffset != INDEX_NONE)
 		{
+			Data.PreviousTransformBufferOffset = Data.TransformBufferOffset;
+			Data.PreviousTransformBufferCount  = Data.TransformBufferCount;
+
 			Data.TransformBufferOffset = INDEX_NONE;
 			Data.TransformBufferCount = 0;
 		}
@@ -1136,7 +1182,7 @@ void FSkinningSceneExtension::FUpdater::PostMeshUpdate(FRDGBuilder& GraphBuilder
 	auto AllocSpaceForPrimitive = [this](const int32 HeaderDataIndex)
 	{
 		FHeaderData& Data = SceneData->HeaderDatas[HeaderDataIndex];
-		Data.DirtyBoneTransforms = EDirtyBoneTransforms::All;
+		Data.DirtyBoneTransforms = Data.bIsInstanced ? EDirtyBoneTransforms::Current : EDirtyBoneTransforms::All;
 
 		if (Data.Proxy->UseSkeletonBatching())
 		{
@@ -1246,6 +1292,13 @@ void FSkinningSceneExtension::FUpdater::PostMeshUpdate(FRDGBuilder& GraphBuilder
 		const uint32 TransformNeededSize = Data.UniqueAnimationCount * Data.MaxTransformCount * 2u; // Current and Previous
 		if (bRequireUpload || (TransformNeededSize != Data.TransformBufferCount))
 		{
+			// Don't overwrite if already set (e.g. by ProcessBufferDefragmentation)
+			if (Data.PreviousTransformBufferOffset == INDEX_NONE)
+			{
+				Data.PreviousTransformBufferOffset = Data.TransformBufferOffset;
+				Data.PreviousTransformBufferCount  = Data.TransformBufferCount;
+			}
+
 			if (Data.TransformBufferCount > 0)
 			{
 				SceneData->TransformAllocator.Free(Data.TransformBufferOffset, Data.TransformBufferCount);
@@ -1254,6 +1307,12 @@ void FSkinningSceneExtension::FUpdater::PostMeshUpdate(FRDGBuilder& GraphBuilder
 			Data.TransformBufferOffset = TransformNeededSize > 0 ? SceneData->TransformAllocator.Allocate(TransformNeededSize) : INDEX_NONE;
 			Data.TransformBufferCount = TransformNeededSize;
 
+			// First allocation — no previous data to copy, provider must write both slots
+			if (Data.PreviousTransformBufferOffset == INDEX_NONE)
+			{
+				Data.DirtyBoneTransforms = EDirtyBoneTransforms::All;
+			}
+
 			if (!bForceFullUpload)
 			{
 				bRequireUpload = true;
@@ -1668,9 +1727,100 @@ void FSkinningSceneExtension::FUpdater::PreGPUSceneUpdate(FRDGBuilder& GraphBuil
 	SCOPED_NAMED_EVENT(FSkinningSceneExtension_FUpdater_PreGPUSceneUpdate, FColor::Silver);
 	CSV_SCOPED_TIMING_STAT_EXCLUSIVE(STAT_SkinningSceneExtension);
 
+	int32 NumCopies = 0;
+	for (int32 PersistentIndex : SceneData->InstancedHeaderDataIndices)
+	{
+		FHeaderData& Data = SceneData->HeaderDatas[PersistentIndex];
+
+		if (Data.PreviousTransformBufferOffset != INDEX_NONE && Data.PreviousTransformBufferOffset != Data.TransformBufferOffset)
+		{
+			NumCopies++;
+		}
+	}
+
+	// Build copy commands for instanced meshes whose transform allocations moved. Collect before FinishSkinningBufferUpload since defrag may shrink the buffer.
+	FRDGBufferRef DefragStagingBuffer = nullptr;
+	TArrayView<FCopyTransformsHeader> CopyHeaders = GraphBuilder.AllocPODArrayView<FCopyTransformsHeader>(NumCopies);
+	uint32 MaxElementsPerCopy = 0;
+	int32 CopyIndex = 0;
+
+	for (int32 PersistentIndex : SceneData->InstancedHeaderDataIndices)
+	{
+		FHeaderData& Data = SceneData->HeaderDatas[PersistentIndex];
+
+		if (Data.PreviousTransformBufferOffset != INDEX_NONE && Data.PreviousTransformBufferOffset != Data.TransformBufferOffset)
+		{
+			const uint32 CopyCount = FMath::Min(Data.PreviousTransformBufferCount, Data.TransformBufferCount);
+			CopyHeaders[CopyIndex++] =
+			{
+				.SrcByteOffset = Data.PreviousTransformBufferOffset * (uint32)sizeof(FCompressedBoneTransform),
+				.DstByteOffset = Data.TransformBufferOffset * (uint32)sizeof(FCompressedBoneTransform),
+				.Count         = CopyCount
+			};
+			MaxElementsPerCopy = FMath::Max(MaxElementsPerCopy, CopyCount);
+			Data.DirtyBoneTransforms = EDirtyBoneTransforms::Current;
+		}
+
+		Data.PreviousTransformBufferOffset = INDEX_NONE;
+		Data.PreviousTransformBufferCount = 0;
+	}
+
+	// For defrag, snapshot the transform buffer before FinishSkinningBufferUpload shrinks it.
+	if (bDefragging && NumCopies > 0)
+	{
+		FRDGBufferRef OldTransformBuffer = SceneData->Buffers->TransformDataBuffer.Register(GraphBuilder);
+		DefragStagingBuffer = GraphBuilder.CreateBuffer(OldTransformBuffer->Desc, TEXT("Skinning.DefragStagingBuffer"));
+
+		FMemcpyResourceParams CopyParams;
+		CopyParams.Count = OldTransformBuffer->Desc.GetSize() / OldTransformBuffer->Desc.BytesPerElement;
+		CopyParams.SrcOffset = 0;
+		CopyParams.DstOffset = 0;
+		MemcpyResource(GraphBuilder, DefragStagingBuffer, OldTransformBuffer, CopyParams);
+	}
+
 	FSkinningSceneParameters Parameters;
 	const bool bUpdateStats = true;
 	SceneData->FinishSkinningBufferUpload(GraphBuilder, &Parameters, bUpdateStats);
+
+	// Dispatch the copy shader
+	if (NumCopies > 0)
+	{
+		RDG_EVENT_SCOPE(GraphBuilder, "CopyTransforms");
+
+		FRDGBufferRef HeaderBuffer = CreateStructuredBuffer(
+			GraphBuilder,
+			TEXT("Skinning.CopyTransformsHeaders"),
+			sizeof(FCopyTransformsHeader),
+			FMath::RoundUpToPowerOfTwo(NumCopies),
+			CopyHeaders.GetData(),
+			sizeof(FCopyTransformsHeader) * NumCopies,
+			ERDGInitialDataFlags::NoCopy
+		);
+
+		FRDGBufferRef TransformBuffer = Parameters.BoneTransforms->GetParent();
+
+		const bool bUseStagingBuffer = DefragStagingBuffer != nullptr;
+
+		FCopyPreviousTransformsCS::FParameters* PassParameters = GraphBuilder.AllocParameters<FCopyPreviousTransformsCS::FParameters>();
+		PassParameters->DstTransformBuffer = GetCompressedBoneTransformUAV(GraphBuilder, TransformBuffer);
+		PassParameters->SrcTransformBuffer = GetCompressedBoneTransformSRV(GraphBuilder, bUseStagingBuffer ? DefragStagingBuffer : GetDefaultCompressedBoneTransformBuffer(GraphBuilder));
+		PassParameters->CopyHeaders = GraphBuilder.CreateSRV(HeaderBuffer);
+
+		FCopyPreviousTransformsCS::FPermutationDomain PermutationVector;
+		PermutationVector.Set<FCopyPreviousTransformsCS::FUseStagingBuffer>(bUseStagingBuffer);
+		TShaderMapRef<FCopyPreviousTransformsCS> ComputeShader(GetGlobalShaderMap(GMaxRHIFeatureLevel), PermutationVector);
+
+		const uint32 GroupSize = FCopyPreviousTransformsCS::ThreadGroupSize;
+		const uint32 NumElementGroups = FMath::DivideAndRoundUp(MaxElementsPerCopy, GroupSize);
+		FComputeShaderUtils::AddPass(
+			GraphBuilder,
+			RDG_EVENT_NAME("CopyTransforms (%d entries, %d max elements, %s)", NumCopies, MaxElementsPerCopy, bUseStagingBuffer ? TEXT("staged") : TEXT("direct")),
+			ComputeShader,
+			PassParameters,
+			FIntVector(NumElementGroups, NumCopies, 1)
+		);
+	}
+
 	SceneData->PerformSkinning(Parameters, GraphBuilder, CurrentTime);
 }
 
@@ -2197,4 +2347,4 @@ void FSkinningSceneExtension::ResolveAttachments(FRDGBuilder& GraphBuilder, cons
 			return FIntVector(FMath::DivideAndRoundUp(NumActualAttachments, 64u), 1, 1);
 		}
 	);
-}
\ No newline at end of file
+}
diff --git a/Engine/Source/Runtime/Renderer/Private/Skinning/SkinningSceneExtension.h b/Engine/Source/Runtime/Renderer/Private/Skinning/SkinningSceneExtension.h
--- a/Engine/Source/Runtime/Renderer/Private/Skinning/SkinningSceneExtension.h
+++ b/Engine/Source/Runtime/Renderer/Private/Skinning/SkinningSceneExtension.h
@@ -54,6 +54,8 @@ public:
 		const bool bEnableAsync = true;
 		bool bForceFullUpload = false;
 		bool bDefragging = false;
+
+		friend FSkinningSceneExtension;
 	};
 
 	class FRenderer : public ISceneExtensionRenderer
@@ -124,6 +126,8 @@ private:
 		uint32 HierarchyBufferCount              = 0;
 		uint32 TransformBufferOffset             = INDEX_NONE;
 		uint32 TransformBufferCount              = 0;
+		uint32 PreviousTransformBufferOffset     = INDEX_NONE;
+		uint32 PreviousTransformBufferCount      = 0;
 		uint32 UniqueAnimationCount              = 1;
 		uint16 MaxTransformCount                 = 0;
 		uint16 MaxBoneMapCount                   = 0;
