Split screen debug editor artifact fix:  Logic that uses "HasBeenProduced" to gate rendering falls apart when the split screen debug mode is active, as it will be true for the second viewport when it shouldn't be.  In this specific case, the step that initializes render targets used for editor compositing (clearing the color buffer and initializing the depth buffer with the scene depth) are skipped, leading to uninitialized garbage being present.  This tends to be worse when resizing the viewport, as it invalidates cached render targets.  Fix is to use the Wireframe show flag instead -- in this context, HasBeenProduced was used to detect wireframe rendering, which forces the editor primitive buffer to be generated in the base pass, requiring the later editor compositing step to not reinitialize it.

#rnx

#Synced-CL 52567897

diff --git a/Engine/Source/Runtime/Renderer/Private/PostProcess/PostProcessCompositeEditorPrimitives.cpp b/Engine/Source/Runtime/Renderer/Private/PostProcess/PostProcessCompositeEditorPrimitives.cpp
--- a/Engine/Source/Runtime/Renderer/Private/PostProcess/PostProcessCompositeEditorPrimitives.cpp
+++ b/Engine/Source/Runtime/Renderer/Private/PostProcess/PostProcessCompositeEditorPrimitives.cpp
@@ -213,8 +213,10 @@ static FEditorPrimitivesDrawingContext InitializeEditorPrimitivesTextures(
 	const uint32 NumMSAASamples = SceneTextures.Config.EditorPrimitiveNumSamples;
 	const FViewInfo* EditorView = CreateCompositePrimitiveView(View, SceneColor.ViewRect, NumMSAASamples);	
 
-	// Load the color target if it already exists.
-	bool bProducedByPriorPass = HasBeenProduced(SceneTextures.EditorPrimitiveColor);
+	// In wireframe mode, the base pass has already rendered to EditorPrimitiveColor/Depth. We check the wireframe flag rather
+	// than using HasBeenProduced (as originally done here), because in split-screen debug mode, HasBeenProduced will be true
+	// for the second viewport, causing this initialization to be incorrectly skipped, compositing garbage onto the screen.
+	bool bProducedByPriorPass = View.Family->EngineShowFlags.Wireframe;
 	FIntPoint Extent = SceneColor.Texture->Desc.Extent;
 	FRDGTextureRef EditorPrimitiveColor = SceneTextures.EditorPrimitiveColor;
 	FRDGTextureRef EditorPrimitiveDepth = SceneTextures.EditorPrimitiveDepth;
