Scene Capture null render target resource crash fix.  Check for a null resource and do nothing, assuming this is a transient issue where the resource was somehow destroyed while a capture was in flight.

#rnx

#Synced-CL 52530005

diff --git a/Engine/Source/Runtime/Renderer/Private/SceneCaptureRendering.cpp b/Engine/Source/Runtime/Renderer/Private/SceneCaptureRendering.cpp
--- a/Engine/Source/Runtime/Renderer/Private/SceneCaptureRendering.cpp
+++ b/Engine/Source/Runtime/Renderer/Private/SceneCaptureRendering.cpp
@@ -1133,7 +1133,10 @@ void FScene::UpdateSceneCaptureContents(USceneCaptureComponent2D* CaptureCompone
 
 	const ESceneRenderGroupFlags SceneRenderGroupFlags = GetSceneCaptureGroupFlags(CaptureComponent);
 
-	if (UTextureRenderTarget2D* TextureRenderTarget = CaptureComponent->TextureTarget)
+	UTextureRenderTarget2D* TextureRenderTarget = CaptureComponent->TextureTarget;
+	FTextureRenderTargetResource* TextureRenderTargetResource = TextureRenderTarget ? TextureRenderTarget->GameThread_GetRenderTargetResource() : nullptr;
+
+	if (TextureRenderTargetResource)
 	{
 		FIntPoint CaptureSize;
 		FVector ViewLocation;
@@ -1477,8 +1480,6 @@ void FScene::UpdateSceneCaptureContents(USceneCaptureComponent2D* CaptureCompone
 		// Reset scene capture's camera cut.
 		CaptureComponent->bCameraCutThisFrame = false;
 
-		FTextureRenderTargetResource* TextureRenderTargetResource = TextureRenderTarget->GameThread_GetRenderTargetResource();
-
 		FString EventName;
 		if (!CaptureComponent->ProfilingEventName.IsEmpty())
 		{
