[HWRT] Hold reference to instance view UB in FRayTracingDynamicGeometryUpdateManager.
- also use TSet instead of TArray.
- and removed unused DispatchUpdates(...) member function.

#Synced-CL 52512372

diff --git a/Engine/Source/Runtime/Renderer/Private/RayTracing/RayTracingDynamicGeometryUpdateManager.cpp b/Engine/Source/Runtime/Renderer/Private/RayTracing/RayTracingDynamicGeometryUpdateManager.cpp
--- a/Engine/Source/Runtime/Renderer/Private/RayTracing/RayTracingDynamicGeometryUpdateManager.cpp
+++ b/Engine/Source/Runtime/Renderer/Private/RayTracing/RayTracingDynamicGeometryUpdateManager.cpp
@@ -381,6 +381,13 @@ void FRayTracingDynamicGeometryUpdateManager::AddDynamicGeometryToUpdate(
 	uint32 PrimitiveId
 )
 {
+	// Keep View and InstancedView UB alive until EndUpdate()
+	ReferencedUniformBuffers.Add(View->ViewUniformBuffer);
+	ReferencedUniformBuffers.Add(View->GetInstancedViewUniformBuffer());
+
+	DispatchCommandsPerView.FindOrAdd(View->ViewUniformBuffer, {});
+	InstancedViewUniformBuffers.FindOrAdd(View->ViewUniformBuffer, View->GetInstancedViewUniformBuffer());
+
 	FRayTracingGeometry& Geometry = *UpdateParams.Geometry;
 
 	uint32 NumVertices = UpdateParams.NumVertices;
@@ -393,22 +400,15 @@ void FRayTracingDynamicGeometryUpdateManager::AddDynamicGeometryToUpdate(
 
 		NumVertices = UpdateParams.NumTriangles * 3;
 		VertexBufferSize = NumVertices * (uint32)sizeof(FVector3f);
-	}	
+	}
 
 	FRWBuffer* RWBuffer = UpdateParams.Buffer;
 	uint32 VertexBufferOffset = 0;
 	bool bUseSharedVertexBuffer = false;
 
-	if (ReferencedUniformBuffers.Num() == 0 || ReferencedUniformBuffers.Last() != View->ViewUniformBuffer)
-	{
-		// Keep ViewUniformBuffer alive until EndUpdate()
-		ReferencedUniformBuffers.Add(View->ViewUniformBuffer);
-	}
-
 	FRayTracingDynamicGeometryBuildParams GeometryBuildParams;
 	GeometryBuildParams.ViewUniformBuffer = View->ViewUniformBuffer;
 	GeometryBuildParams.DispatchCommands.Reserve(UpdateParams.MeshBatches.Num());
-	InstancedViewUniformBuffers.FindOrAdd(View->ViewUniformBuffer, View->GetInstancedViewUniformBuffer());
 
 	if (CVarRTDynGeomPriorityBasedUpdate.GetValueOnRenderThread())
 	{
@@ -737,14 +737,13 @@ void FRayTracingDynamicGeometryUpdateManager::ScheduleUpdates(bool bUseTracingFe
 		return;
 	}
 
-	checkf(DispatchCommandsPerView.IsEmpty(), TEXT("DispatchCommandsPerView is not empty. Previous frame updates were not dispatched."));
 	checkf(BuildParams.IsEmpty(), TEXT("BuildParams is not empty. Previous frame updates were not dispatched."));
 
 	// reserve for worst case (might be wasteful if there are too many views)
-	for (FRHIUniformBuffer* ViewUniformBuffer : ReferencedUniformBuffers)
+	for (auto& ViewDispatchCommandsPair : DispatchCommandsPerView)
 	{
-		TArray<FMeshComputeDispatchCommand>& ViewDispatchCommands = DispatchCommandsPerView.FindOrAdd(ViewUniformBuffer);
-		ViewDispatchCommands.Reserve(TotalNumGeometryBuilds);
+		checkf(ViewDispatchCommandsPair.Value.IsEmpty(), TEXT("DispatchCommandsPerView is not empty. Previous frame updates were not dispatched."));
+		ViewDispatchCommandsPair.Value.Reserve(TotalNumGeometryBuilds);
 	}
 
 	BuildParams.Reserve(TotalNumGeometryBuilds);
@@ -1112,33 +1111,9 @@ void FRayTracingDynamicGeometryUpdateManager::DispatchUpdates(FRHICommandList& R
 	}
 }
 
-void FRayTracingDynamicGeometryUpdateManager::DispatchUpdates(FRHICommandList& RHICmdList, FRHIBuffer* ScratchBuffer, ERHIPipeline SrcResourceAccessPipelines, ERHIPipeline DstResourceAccessPipelines)
-{
-	for (TPair<FRHIUniformBuffer*, TArray<FMeshComputeDispatchCommand>>& ViewDispatchCommands : DispatchCommandsPerView)
-	{
-		DispatchUpdates(RHICmdList, ViewDispatchCommands.Value, SrcResourceAccessPipelines, DstResourceAccessPipelines);
-	}
-			
-	if (BuildParams.Num() > 0)
-	{
-		// Can't use parallel command list because we have to make sure we are not building BVH data
-		// on the same RTGeometry on multiple threads at the same time. Ideally move the build
-		// requests over to the RaytracingGeometry manager so they can be correctly scheduled
-		// with other build requests in the engine (see UE-106982)
-		SCOPED_DRAW_EVENT(RHICmdList, Build);
-
-		FRHIBufferRange ScratchBufferRange;
-		ScratchBufferRange.Buffer = ScratchBuffer;
-		ScratchBufferRange.Offset = 0;
-		RHICmdList.BuildAccelerationStructures(BuildParams, ScratchBufferRange);
-	}
-
-	EndUpdate();
-}
-
 void FRayTracingDynamicGeometryUpdateManager::EndUpdate()
 {
-	ReferencedUniformBuffers.Empty(ReferencedUniformBuffers.Max());
+	ReferencedUniformBuffers.Reset();
 
 	Clear();
 }
diff --git a/Engine/Source/Runtime/Renderer/Public/RayTracingDynamicGeometryUpdateManager.h b/Engine/Source/Runtime/Renderer/Public/RayTracingDynamicGeometryUpdateManager.h
--- a/Engine/Source/Runtime/Renderer/Public/RayTracingDynamicGeometryUpdateManager.h
+++ b/Engine/Source/Runtime/Renderer/Public/RayTracingDynamicGeometryUpdateManager.h
@@ -74,7 +74,6 @@ private:
 
 	void EndUpdate();
 	void ScheduleUpdates(bool bUseTracingFeedback);
-	void DispatchUpdates(FRHICommandList& RHICmdList, FRHIBuffer* ScratchBuffer, ERHIPipeline SrcResourceAccessPipelines = ERHIPipeline::Graphics, ERHIPipeline DstResourceAccessPipelines = ERHIPipeline::Graphics);
 
 	static void DispatchUpdates(FRHICommandList& RHICmdList, TConstArrayView<FMeshComputeDispatchCommand> DispatchCommands, ERHIPipeline SrcResourceAccessPipelines = ERHIPipeline::Graphics, ERHIPipeline DstResourceAccessPipelines = ERHIPipeline::Graphics);
 	
@@ -98,8 +97,8 @@ private:
 							 uint32 VertexBufferOffset,
 							 uint32 VertexBufferSize,
 	                         FRayTracingDynamicGeometryBuildParams& BuildParams);
-	FRWBuffer* AllocateSharedBuffer(FRHICommandListBase& RHICmdList, uint32 VertexBufferSize, uint32& OutVertexBufferOffset);
 
+	FRWBuffer* AllocateSharedBuffer(FRHICommandListBase& RHICmdList, uint32 VertexBufferSize, uint32& OutVertexBufferOffset);
 
 	TArray<FRayTracingDynamicGeometryBuildParams> DynamicGeometryBuilds;
 	TArray<FRayTracingDynamicGeometryBuildParams> DynamicGeometryUpdates;
@@ -119,7 +118,7 @@ private:
 	TArray<FVertexPositionBuffer*> VertexPositionBuffers;
 
 	// Any uniform buffers that must be kept alive until EndUpdate (after DispatchUpdates is called)
-	TArray<FUniformBufferRHIRef> ReferencedUniformBuffers;
+	TSet<FUniformBufferRHIRef> ReferencedUniformBuffers;
 
 	// Generation ID when the shared vertex buffers have been reset. The current generation ID is stored in the FRayTracingGeometry to keep track
 	// if the vertex buffer data is still valid for that frame - validated before generation the TLAS
