[HWRT] Don't need to recreate geometry when shared vertex buffer is resized or geometry doesn't support updates -> simply refit/rebuild as appropriate.

#Synced-CL 52512372

diff --git a/Engine/Source/Runtime/Renderer/Private/RayTracing/RayTracingDynamicGeometryUpdateManager.cpp b/Engine/Source/Runtime/Renderer/Private/RayTracing/RayTracingDynamicGeometryUpdateManager.cpp
--- a/Engine/Source/Runtime/Renderer/Private/RayTracing/RayTracingDynamicGeometryUpdateManager.cpp
+++ b/Engine/Source/Runtime/Renderer/Private/RayTracing/RayTracingDynamicGeometryUpdateManager.cpp
@@ -389,6 +389,7 @@ void FRayTracingDynamicGeometryUpdateManager::AddDynamicGeometryToUpdate(
 	InstancedViewUniformBuffers.FindOrAdd(View->ViewUniformBuffer, View->GetInstancedViewUniformBuffer());
 
 	FRayTracingGeometry& Geometry = *UpdateParams.Geometry;
+	check(Geometry.IsInitialized());
 
 	uint32 NumVertices = UpdateParams.NumVertices;
 	uint32 VertexBufferSize = UpdateParams.VertexBufferSize;
@@ -429,31 +430,24 @@ void FRayTracingDynamicGeometryUpdateManager::AddDynamicGeometryToUpdate(
 		AddDispatchCommands(RHICmdList, Scene, View, PrimitiveSceneProxy, UpdateParams, PrimitiveId, RWBuffer, NumVertices, VertexBufferOffset, VertexBufferSize, GeometryBuildParams);
 	}
 
-	bool bRefit = true;
-
 	// Optionally resize the buffer when not shared (could also be lazy allocated and still empty)
 	if (!bUseSharedVertexBuffer && RWBuffer && RWBuffer->NumBytes != VertexBufferSize)
 	{
 		RWBuffer->Initialize(RHICmdList, TEXT("FRayTracingDynamicGeometryUpdateManager::RayTracingDynamicVertexBuffer"), sizeof(float), VertexBufferSize / sizeof(float), PF_R32_FLOAT, BUF_UnorderedAccess | BUF_ShaderResource);
-		bRefit = false;
 	}
 
-	if (!Geometry.IsValid() || Geometry.IsEvicted())
-	{
-		bRefit = false;
-	}
+	bool bRecreateGeometry = false;
 
-	if (!Geometry.Initializer.bAllowUpdate)
+	if (!Geometry.IsValid() || Geometry.IsEvicted())
 	{
-		bRefit = false;
+		bRecreateGeometry = true;
 	}
 
-	check(Geometry.IsInitialized());
-
 	// Primitive count or triangle count changed then setup the initializer and segment data again
 	// TODO: Move this to higher level code since it only works with 1 segment. These "resizes" should be handled by the systems that request dynamic geometry updates.
+	const bool bTotalPrimitiveCountChanges = Geometry.Initializer.TotalPrimitiveCount != UpdateParams.NumTriangles;
 	const bool bVertexCountChanges = Geometry.Initializer.Segments.Num() == 1 ? Geometry.Initializer.Segments[0].MaxVertices != NumVertices : false;
-	if ((Geometry.Initializer.TotalPrimitiveCount != UpdateParams.NumTriangles || bVertexCountChanges) && UpdateParams.NumTriangles > 0)
+	if ((bTotalPrimitiveCountChanges || bVertexCountChanges) && UpdateParams.NumTriangles > 0)
 	{
 		checkf(Geometry.Initializer.Segments.Num() <= 1, TEXT("Dynamic ray tracing geometry '%s' has an unexpected number of segments."), *Geometry.Initializer.DebugName.ToString());
 		Geometry.Initializer.TotalPrimitiveCount = UpdateParams.NumTriangles;
@@ -462,7 +456,7 @@ void FRayTracingDynamicGeometryUpdateManager::AddDynamicGeometryToUpdate(
 		Segment.NumPrimitives = UpdateParams.NumTriangles;
 		Segment.MaxVertices = NumVertices;
 		Geometry.Initializer.Segments.Add(Segment);
-		bRefit = false;
+		bRecreateGeometry = true;
 	}
 
 	if (UpdateParams.bAlphaMasked)
@@ -488,16 +482,12 @@ void FRayTracingDynamicGeometryUpdateManager::AddDynamicGeometryToUpdate(
 	}
 #endif
 
-	if (!bRefit)
+	if (bRecreateGeometry)
 	{
 		checkf(Geometry.RawData.IsEmpty() && Geometry.Initializer.OfflineData == nullptr, TEXT("Dynamic ray tracing geometry '%s' is not expected to have offline acceleration structure data."), *Geometry.Initializer.DebugName.ToString());
 		Geometry.CreateRayTracingGeometry(RHICmdList, ERTAccelerationStructureBuildPriority::Skip);
 	}
 
-	EAccelerationStructureBuildMode BuildMode = Geometry.GetRequiresBuild()
-		? EAccelerationStructureBuildMode::Build
-		: EAccelerationStructureBuildMode::Update;
-
 	GeometryBuildParams.Geometry = UpdateParams.Geometry;
 
 	if (bUseSharedVertexBuffer)
@@ -506,9 +496,7 @@ void FRayTracingDynamicGeometryUpdateManager::AddDynamicGeometryToUpdate(
 		Segments.Append(Geometry.Initializer.Segments);
 	}
 
-	Geometry.SetRequiresBuild(false);
-
-	if (BuildMode == EAccelerationStructureBuildMode::Build)
+	if (Geometry.GetRequiresBuild() || !Geometry.Initializer.bAllowUpdate)
 	{
 		DynamicGeometryBuilds.Add(GeometryBuildParams);
 	}
@@ -516,6 +504,8 @@ void FRayTracingDynamicGeometryUpdateManager::AddDynamicGeometryToUpdate(
 	{
 		DynamicGeometryUpdates.Add(GeometryBuildParams);
 	}
+
+	Geometry.SetRequiresBuild(false);
 	
 	if (bUseSharedVertexBuffer)
 	{
