Fix missing light function on hair strands when using RT shadow.

#Synced-CL 52611221

diff --git a/Engine/Source/Runtime/Renderer/Private/LightRendering.cpp b/Engine/Source/Runtime/Renderer/Private/LightRendering.cpp
--- a/Engine/Source/Runtime/Renderer/Private/LightRendering.cpp
+++ b/Engine/Source/Runtime/Renderer/Private/LightRendering.cpp
@@ -2375,17 +2375,20 @@ void FDeferredShadingSceneRenderer::RenderLights(
 				{
 					if (bDrawLightFunction)
 					{
-						if (bIsHairPass)
+						if (!bIsHairPass)
 						{
-							if (CVarAppliedLightFunctionOnHair.GetValueOnRenderThread() > 0 && ScreenShadowMaskSubPixelTexture)
-							{
-								const bool bLightFunctionRendered = RenderLightFunction(GraphBuilder, SceneTextures, &LightSceneInfo, ScreenShadowMaskSubPixelTexture, bDrawShadows, false, true);
-								bUsedShadowMaskTexture |= bLightFunctionRendered;
-							}
+							const bool bLightFunctionRendered = RenderLightFunction(GraphBuilder, SceneTextures, &LightSceneInfo, ScreenShadowMaskTexture, bDrawShadows, false, false);
+							bUsedShadowMaskTexture |= bLightFunctionRendered;
 						}
-						else
+						
+						// * When using non-raytraced shadows, the light function for hair is evaluated during the hair pass (bIsHairPass=true).
+						// * When using ray traced shadows, hair lighting runs in the non-hair pass (bIsHairPass=false).
+						//   The light function must also be rendered to ScreenShadowMaskSubPixelTexture for hair,
+						//   since non-atlas-compatible light functions (e.g. those using AbsoluteWorldPosition)
+						//   are not sampled via the light function atlas in the hair deferred light shader.
+						if (bRunHairLighting && CVarAppliedLightFunctionOnHair.GetValueOnRenderThread() > 0 && ScreenShadowMaskSubPixelTexture)
 						{
-							const bool bLightFunctionRendered = RenderLightFunction(GraphBuilder, SceneTextures, &LightSceneInfo, ScreenShadowMaskTexture, bDrawShadows, false, false);
+							const bool bLightFunctionRendered = RenderLightFunction(GraphBuilder, SceneTextures, &LightSceneInfo, ScreenShadowMaskSubPixelTexture, bDrawShadows, false, true);
 							bUsedShadowMaskTexture |= bLightFunctionRendered;
 						}
 					}
