Fixed Nanite HZB occlusion on platforms with r.HZB.BuildUseCompute=0. The pixel shader path doesn't support VisBuffer, so force compute path for now.
Reconsider later if we can make the compute path faster or if we should add VisBuffer support to the pixel shader path.

#Synced-CL 52616188

diff --git a/Engine/Source/Runtime/Renderer/Private/SceneTextureReductions.cpp b/Engine/Source/Runtime/Renderer/Private/SceneTextureReductions.cpp
--- a/Engine/Source/Runtime/Renderer/Private/SceneTextureReductions.cpp
+++ b/Engine/Source/Runtime/Renderer/Private/SceneTextureReductions.cpp
@@ -140,7 +140,8 @@ void BuildHZB(
 	int32 NumMips = FMath::Max(FMath::FloorToInt(FMath::Log2(FMath::Max<float>(HZBSize.X, HZBSize.Y))), 1);
 
 	bool bReduceClosestDepth = OutClosestHZBTexture != nullptr;
-	bool bUseCompute = bReduceClosestDepth || HZBBuildUseCompute(ShaderPlatform);
+	// Force compute when a VisBuffer is provided, since HZBBuildPS does not support sampling it.
+	bool bUseCompute = bReduceClosestDepth || (VisBufferTexture != nullptr) || HZBBuildUseCompute(ShaderPlatform);
 
 	FRDGTextureDesc HZBDesc = FRDGTextureDesc::Create2D(
 		HZBSize, Format,
