[HWRT] Track RayTracingDynamicSharedVertexBufferMemory stat.

#Synced-CL 52610883

diff --git a/Engine/Source/Runtime/Renderer/Private/RayTracing/RayTracingDynamicGeometryUpdateManager.cpp b/Engine/Source/Runtime/Renderer/Private/RayTracing/RayTracingDynamicGeometryUpdateManager.cpp
--- a/Engine/Source/Runtime/Renderer/Private/RayTracing/RayTracingDynamicGeometryUpdateManager.cpp
+++ b/Engine/Source/Runtime/Renderer/Private/RayTracing/RayTracingDynamicGeometryUpdateManager.cpp
@@ -12,6 +12,7 @@
 #include "Rendering/RayTracingGeometryManager.h"
 #include "UnrealEngine.h"
 #include "RayTracing/RaytracingOptions.h"
+#include "ProfilingDebugging/CsvProfiler.h"
 
 #if RHI_RAYTRACING
 
@@ -24,6 +25,9 @@ DECLARE_GPU_STAT(RayTracingDynamicGeometry);
 DECLARE_DWORD_COUNTER_STAT(TEXT("Ray tracing dynamic build primitives"), STAT_RayTracingDynamicBuildPrimitives, STATGROUP_SceneRendering);
 DECLARE_DWORD_COUNTER_STAT(TEXT("Ray tracing dynamic update primitives"), STAT_RayTracingDynamicUpdatePrimitives, STATGROUP_SceneRendering);
 DECLARE_DWORD_COUNTER_STAT(TEXT("Ray tracing dynamic skipped primitives"), STAT_RayTracingDynamicSkippedPrimitives, STATGROUP_SceneRendering);
+DECLARE_MEMORY_STAT(TEXT("Ray tracing dynamic shared vertex buffer memory"), STAT_RayTracingDynamicSharedVertexBufferMemory, STATGROUP_SceneRendering);
+
+CSV_DEFINE_CATEGORY(RayTracing, true);
 
 static int32 GRTDynGeomSharedVertexBufferSizeInMB = 4;
 static FAutoConsoleVariableRef CVarRTDynGeomSharedVertexBufferSizeInMB(
@@ -718,7 +722,7 @@ END_SHADER_PARAMETER_STRUCT()
 
 void FRayTracingDynamicGeometryUpdateManager::ScheduleUpdates(bool bUseTracingFeedback)
 {
-	TRACE_CPUPROFILER_EVENT_SCOPE(FRayTracingDynamicGeometryUpdateManager::Update);
+	TRACE_CPUPROFILER_EVENT_SCOPE(FRayTracingDynamicGeometryUpdateManager::ScheduleUpdates);
 	
 	// Early out of nothing to do
 	const int32 TotalNumGeometryBuilds = DynamicGeometryBuilds.Num() + DynamicGeometryUpdates.Num();
@@ -876,6 +880,18 @@ void FRayTracingDynamicGeometryUpdateManager::ScheduleUpdates(bool bUseTracingFe
 	INC_DWORD_STAT_BY(STAT_RayTracingDynamicUpdatePrimitives, NumUpdatedPrimitives);
 	INC_DWORD_STAT_BY(STAT_RayTracingDynamicBuildPrimitives, NumBuildPrimitives);
 
+	// Track total shared vertex buffer memory (after all allocations are complete)
+	{
+		int64 TotalSharedVertexBufferMemory = 0;
+		for (const FVertexPositionBuffer* Buffer : VertexPositionBuffers)
+		{
+			TotalSharedVertexBufferMemory += Buffer->RWBuffer.NumBytes;
+		}
+
+		SET_MEMORY_STAT(STAT_RayTracingDynamicSharedVertexBufferMemory, TotalSharedVertexBufferMemory);
+		CSV_CUSTOM_STAT(RayTracing, DynamicGeometrySharedVertexBufferMB, float(TotalSharedVertexBufferMemory) / (1024.0f * 1024.0f), ECsvCustomStatOp::Set);
+	}
+
 	ScratchBufferSize = BLASScratchSize;
 }
 
