[HWRT] Migrated transient per-frame data in FRayTracingDynamicGeometryUpdateManager from persistent member arrays to RDG-scoped allocations.

#Synced-CL 52610883

diff --git a/Engine/Source/Runtime/Renderer/Private/RayTracing/RayTracingDynamicGeometryUpdateManager.cpp b/Engine/Source/Runtime/Renderer/Private/RayTracing/RayTracingDynamicGeometryUpdateManager.cpp
--- a/Engine/Source/Runtime/Renderer/Private/RayTracing/RayTracingDynamicGeometryUpdateManager.cpp
+++ b/Engine/Source/Runtime/Renderer/Private/RayTracing/RayTracingDynamicGeometryUpdateManager.cpp
@@ -332,10 +332,9 @@ void FRayTracingDynamicGeometryUpdateManager::Clear()
 	DispatchCommandsPerView = {};
 	InstancedViewUniformBuffers = {};
 
-	// Clear working arrays - keep max size allocated
-	BuildParams.Empty(BuildParams.Max());
-	Segments.Empty(Segments.Max());
+	BuildParams = nullptr;
 
+	// Clear working arrays - keep max size allocated
 	DynamicGeometryBuilds.Empty(DynamicGeometryBuilds.Max());
 	DynamicGeometryUpdates.Empty(DynamicGeometryUpdates.Max());
 
@@ -346,13 +345,12 @@ int64 FRayTracingDynamicGeometryUpdateManager::BeginUpdate()
 {
 	check(DispatchCommandsPerView.IsEmpty());
 	check(InstancedViewUniformBuffers.IsEmpty());
-	check(BuildParams.IsEmpty());
-	check(Segments.IsEmpty());
+	check(BuildParams == nullptr);
 	check(ReferencedUniformBuffers.IsEmpty());
 	check(DynamicGeometryBuilds.IsEmpty());
 	check(DynamicGeometryUpdates.IsEmpty());
 
-	// Vertex buffer data can be immediatly reused the next frame, because it's already 'consumed' for building the AccelerationStructure data
+	// Vertex buffer data can be immediately reused the next frame, because it's already 'consumed' for building the AccelerationStructure data
 	// Garbage collect unused buffers for n generations
 	for (int32 BufferIndex = 0; BufferIndex < VertexPositionBuffers.Num(); ++BufferIndex)
 	{
@@ -389,7 +387,7 @@ void FRayTracingDynamicGeometryUpdateManager::AddDynamicGeometryToUpdate(
 	ReferencedUniformBuffers.Add(View->ViewUniformBuffer);
 	ReferencedUniformBuffers.Add(View->GetInstancedViewUniformBuffer());
 
-	DispatchCommandsPerView.FindOrAdd(View->ViewUniformBuffer, {});
+	DispatchCommandsPerView.FindOrAdd(View->ViewUniformBuffer, nullptr);
 	InstancedViewUniformBuffers.FindOrAdd(View->ViewUniformBuffer, View->GetInstancedViewUniformBuffer());
 
 	FRayTracingGeometry& Geometry = *UpdateParams.Geometry;
@@ -431,7 +429,7 @@ void FRayTracingDynamicGeometryUpdateManager::AddDynamicGeometryToUpdate(
 		}
 		check(IsAligned(VertexBufferOffset, RHI_RAW_VIEW_ALIGNMENT));
 		
-		AddDispatchCommands(RHICmdList, Scene, View, PrimitiveSceneProxy, UpdateParams, PrimitiveId, RWBuffer, NumVertices, VertexBufferOffset, VertexBufferSize, GeometryBuildParams);
+		AddDispatchCommands(RHICmdList, Scene, View, PrimitiveSceneProxy, UpdateParams, PrimitiveId, RWBuffer, NumVertices, VertexBufferOffset, VertexBufferSize, GeometryBuildParams.DispatchCommands);
 	}
 
 	// Optionally resize the buffer when not shared (could also be lazy allocated and still empty)
@@ -494,12 +492,6 @@ void FRayTracingDynamicGeometryUpdateManager::AddDynamicGeometryToUpdate(
 
 	GeometryBuildParams.Geometry = UpdateParams.Geometry;
 
-	if (bUseSharedVertexBuffer)
-	{
-		GeometryBuildParams.SegmentOffset = Segments.Num();
-		Segments.Append(Geometry.Initializer.Segments);
-	}
-
 	if (Geometry.GetRequiresBuild() || !Geometry.Initializer.bAllowUpdate)
 	{
 		DynamicGeometryBuilds.Add(GeometryBuildParams);
@@ -571,7 +563,7 @@ void FRayTracingDynamicGeometryUpdateManager::AddDispatchCommands(
 	uint32 NumVertices,
 	uint32 VertexBufferOffset,
 	uint32 VertexBufferSize,
-	FRayTracingDynamicGeometryUpdateManager::FRayTracingDynamicGeometryBuildParams& GeometryBuildParams)
+	TArray<FMeshComputeDispatchCommand>& OutDispatchCommands)
 {	
 	const int32 PSOCollectorIndex = FPSOCollectorCreateManager::GetIndex(EShadingPath::Deferred, RayTracingDynamicGeometryPSOCollectorName);
 
@@ -702,7 +694,7 @@ void FRayTracingDynamicGeometryUpdateManager::AddDispatchCommands(
 					DispatchCmd.bApplyWorldPositionOffset = UpdateParams.bApplyWorldPositionOffset;
 #endif
 
-					GeometryBuildParams.DispatchCommands.Add(DispatchCmd);
+					OutDispatchCommands.Add(DispatchCmd);
 
 					break;
 				}
@@ -720,7 +712,7 @@ BEGIN_SHADER_PARAMETER_STRUCT(FRayTracingDynamicGeometryUpdatePassParams, )
 	SHADER_PARAMETER_RDG_UNIFORM_BUFFER(FSceneUniformParameters, Scene)
 END_SHADER_PARAMETER_STRUCT()
 
-void FRayTracingDynamicGeometryUpdateManager::ScheduleUpdates(bool bUseTracingFeedback)
+void FRayTracingDynamicGeometryUpdateManager::ScheduleUpdates(FRDGBuilder& GraphBuilder, bool bUseTracingFeedback)
 {
 	TRACE_CPUPROFILER_EVENT_SCOPE(FRayTracingDynamicGeometryUpdateManager::ScheduleUpdates);
 	
@@ -731,25 +723,24 @@ void FRayTracingDynamicGeometryUpdateManager::ScheduleUpdates(bool bUseTracingFe
 		return;
 	}
 
-	checkf(BuildParams.IsEmpty(), TEXT("BuildParams is not empty. Previous frame updates were not dispatched."));
+	checkf(BuildParams == nullptr, TEXT("BuildParams is not empty. Previous frame updates were not dispatched."));
+	BuildParams = &GraphBuilder.AllocArray<FRayTracingGeometryBuildParams>();
+	BuildParams->Reserve(TotalNumGeometryBuilds);
 
 	// reserve for worst case (might be wasteful if there are too many views)
 	for (auto& ViewDispatchCommandsPair : DispatchCommandsPerView)
 	{
-		checkf(ViewDispatchCommandsPair.Value.IsEmpty(), TEXT("DispatchCommandsPerView is not empty. Previous frame updates were not dispatched."));
-		ViewDispatchCommandsPair.Value.Reserve(TotalNumGeometryBuilds);
+		checkf(ViewDispatchCommandsPair.Value == nullptr, TEXT("DispatchCommandsPerView is not empty. Previous frame updates were not dispatched."));
+		ViewDispatchCommandsPair.Value = &GraphBuilder.AllocArray<FMeshComputeDispatchCommand>();
+		ViewDispatchCommandsPair.Value->Reserve(TotalNumGeometryBuilds);
 	}
 
-	BuildParams.Reserve(TotalNumGeometryBuilds);
-
-	FRayTracingGeometrySegment* SegmentData = Segments.GetData();
-
 	const uint32 ScratchAlignment = GRHIRayTracingScratchBufferAlignment;
 
 	uint32 BLASScratchSize = 0;
 	int32 NumBuildPrimitives = 0;
 
-	auto AddGeometryBuildParamsToBuildList = [this, SegmentData, &BLASScratchSize](const FRayTracingDynamicGeometryBuildParams& InBuildParams)
+	auto AddGeometryBuildParamsToBuildList = [this, &GraphBuilder, &BLASScratchSize](const FRayTracingDynamicGeometryBuildParams& InBuildParams)
 		{
 			FRHIRayTracingGeometry* RayTracingGeometry = InBuildParams.Geometry->GetRHI();
 			const FRayTracingGeometryInitializer& GeometryInitializer = InBuildParams.Geometry->GetInitializer();
@@ -763,24 +754,20 @@ void FRayTracingDynamicGeometryUpdateManager::ScheduleUpdates(bool bUseTracingFe
 			InBuildParams.Geometry->NumUpdatesSinceLastBuild = 0;
 			InBuildParams.Geometry->SetRequiresUpdate(false);
 
+			// Make a copy of GeometryInitializer.Segments that can be passed to RDG pass
+			TArray<FRayTracingGeometrySegment, SceneRenderingAllocator>& Segments = GraphBuilder.AllocArray<FRayTracingGeometrySegment>();
+			Segments.Append(GeometryInitializer.Segments);
+
 			FRayTracingGeometryBuildParams RTGeoBuildParams;
 			RTGeoBuildParams.Geometry = RayTracingGeometry;
 			RTGeoBuildParams.BuildMode = EAccelerationStructureBuildMode::Build;
+			RTGeoBuildParams.Segments = Segments;
 
-			if (InBuildParams.SegmentOffset >= 0)
-			{
-				RTGeoBuildParams.Segments = MakeArrayView(&SegmentData[InBuildParams.SegmentOffset], GeometryInitializer.Segments.Num());
-			}
-			else
-			{
-				RTGeoBuildParams.Segments = GeometryInitializer.Segments;
-			}
-
-			BuildParams.Add(MoveTemp(RTGeoBuildParams));
+			BuildParams->Add(MoveTemp(RTGeoBuildParams));
 
 			if (!InBuildParams.DispatchCommands.IsEmpty())
 			{
-				DispatchCommandsPerView[InBuildParams.ViewUniformBuffer].Append(InBuildParams.DispatchCommands);
+				DispatchCommandsPerView[InBuildParams.ViewUniformBuffer]->Append(InBuildParams.DispatchCommands);
 			}
 		};
 
@@ -853,22 +840,20 @@ void FRayTracingDynamicGeometryUpdateManager::ScheduleUpdates(bool bUseTracingFe
 		const uint32 ScratchSize = RayTracingGeometry->GetSizeInfo().UpdateScratchSize;
 		BLASScratchSize = Align(BLASScratchSize + ScratchSize, ScratchAlignment);
 
+		// Make a copy of GeometryInitializer.Segments that can be passed to RDG pass
+		TArray<FRayTracingGeometrySegment, SceneRenderingAllocator>& Segments = GraphBuilder.AllocArray<FRayTracingGeometrySegment>();
+		Segments.Append(GeometryInitializer.Segments);
+
 		FRayTracingGeometryBuildParams BuildParam;
 		BuildParam.Geometry = RayTracingGeometry;
 		BuildParam.BuildMode = EAccelerationStructureBuildMode::Update;
-		if (Update.SegmentOffset >= 0)
-		{
-			BuildParam.Segments = MakeArrayView(&SegmentData[Update.SegmentOffset], GeometryInitializer.Segments.Num());
-		}
-		else
-		{
-			BuildParam.Segments = GeometryInitializer.Segments;
-		}
-		BuildParams.Add(MoveTemp(BuildParam));
+		BuildParam.Segments = Segments;
+
+		BuildParams->Add(MoveTemp(BuildParam));
 
 		if (!Update.DispatchCommands.IsEmpty())
 		{
-			DispatchCommandsPerView[Update.ViewUniformBuffer].Append(Update.DispatchCommands);
+			DispatchCommandsPerView[Update.ViewUniformBuffer]->Append(Update.DispatchCommands);
 		}
 
 		if (NumUpdatedPrimitives > MaxUpdatePrimitivesPerFrame)
@@ -900,7 +885,7 @@ void FRayTracingDynamicGeometryUpdateManager::AddDynamicGeometryUpdatePass(FRDGB
 	RDG_GPU_MASK_SCOPE(GraphBuilder, FRHIGPUMask::All());
 	RDG_EVENT_SCOPE_STAT(GraphBuilder, RayTracingDynamicGeometry, "RayTracingDynamicGeometry");
 
-	ScheduleUpdates(bUseTracingFeedback);
+	ScheduleUpdates(GraphBuilder, bUseTracingFeedback);
 
 	const uint32 ScratchAlignment = GRHIRayTracingScratchBufferAlignment;
 	const uint32 BLASScratchSize = ScratchBufferSize;
@@ -928,9 +913,9 @@ void FRayTracingDynamicGeometryUpdateManager::AddDynamicGeometryUpdatePass(FRDGB
 
 	const ERHIPipeline SrcResourceAccessPipelines = ComputePassFlags == ERDGPassFlags::AsyncCompute ? ERHIPipeline::AsyncCompute : ERHIPipeline::Graphics;
 
-	for(TPair<FRHIUniformBuffer*, TArray<FMeshComputeDispatchCommand>>& ViewDispatchCommands : DispatchCommandsPerView)
+	for(TPair<FRHIUniformBuffer*, TArray<FMeshComputeDispatchCommand, SceneRenderingAllocator>*>& ViewDispatchCommands : DispatchCommandsPerView)
 	{
-		if (ViewDispatchCommands.Value.IsEmpty())
+		if (ViewDispatchCommands.Value == nullptr || ViewDispatchCommands.Value->IsEmpty())
 		{
 			continue;
 		}
@@ -945,13 +930,13 @@ void FRayTracingDynamicGeometryUpdateManager::AddDynamicGeometryUpdatePass(FRDGB
 		PassParams->DynamicGeometryScratchBuffer = OutDynamicGeometryScratchBuffer; 
 
 		GraphBuilder.AddPass(RDG_EVENT_NAME("RayTracingDynamicUpdate"), PassParams, ComputePassFlags | ERDGPassFlags::NeverCull,
-			[this, SrcResourceAccessPipelines, ResourceAccessPipelines, DispatchCommands = MakeArrayView(ViewDispatchCommands.Value)](FRHICommandList& RHICmdList)
+			[SrcResourceAccessPipelines, ResourceAccessPipelines, DispatchCommands = MakeConstArrayView(*ViewDispatchCommands.Value)](FRHICommandList& RHICmdList)
 			{
 				DispatchUpdates(RHICmdList, DispatchCommands, SrcResourceAccessPipelines, ResourceAccessPipelines);
 			});
 	}
 
-	if (BuildParams.Num() > 0)
+	if (BuildParams != nullptr && BuildParams->Num() > 0)
 	{
 		FRayTracingDynamicGeometryUpdatePassParams* PassParams = GraphBuilder.AllocParameters<FRayTracingDynamicGeometryUpdatePassParams>();
 		PassParams->View = {};
@@ -959,7 +944,7 @@ void FRayTracingDynamicGeometryUpdateManager::AddDynamicGeometryUpdatePass(FRDGB
 		PassParams->DynamicGeometryScratchBuffer = OutDynamicGeometryScratchBuffer;
 
 		GraphBuilder.AddPass(RDG_EVENT_NAME("RayTracingDynamicUpdateBuild"), PassParams, ComputePassFlags | ERDGPassFlags::NeverCull,
-			[this, PassParams, bSlowPath](FRHICommandList& RHICmdList)
+			[PassParams, BuildParams = MakeConstArrayView(*BuildParams), bSlowPath](FRHICommandList& RHICmdList)
 			{
 				// Can't use parallel command list because we have to make sure we are not building BVH data
 				// on the same RTGeometry on multiple threads at the same time. Ideally move the build
@@ -981,6 +966,8 @@ void FRayTracingDynamicGeometryUpdateManager::AddDynamicGeometryUpdatePass(FRDGB
 			});
 	}
 
+	Clear();
+
 	// TODO: Is it safe to use a regular task that waits on FRDGBuilder::GetAsyncExecuteTask() here instead?
 	// which would allow the passes above to be tagged with FRDGAsyncTask
 	GraphBuilder.AddPostExecuteCallback([this]
@@ -1120,8 +1107,6 @@ void FRayTracingDynamicGeometryUpdateManager::DispatchUpdates(FRHICommandList& R
 void FRayTracingDynamicGeometryUpdateManager::EndUpdate()
 {
 	ReferencedUniformBuffers.Reset();
-
-	Clear();
 }
 
 #undef USE_RAY_TRACING_DYNAMIC_GEOMETRY_PARALLEL_COMMAND_LISTS
diff --git a/Engine/Source/Runtime/Renderer/Public/RayTracingDynamicGeometryUpdateManager.h b/Engine/Source/Runtime/Renderer/Public/RayTracingDynamicGeometryUpdateManager.h
--- a/Engine/Source/Runtime/Renderer/Public/RayTracingDynamicGeometryUpdateManager.h
+++ b/Engine/Source/Runtime/Renderer/Public/RayTracingDynamicGeometryUpdateManager.h
@@ -73,7 +73,7 @@ public:
 private:
 
 	void EndUpdate();
-	void ScheduleUpdates(bool bUseTracingFeedback);
+	void ScheduleUpdates(FRDGBuilder& GraphBuilder, bool bUseTracingFeedback);
 
 	static void DispatchUpdates(FRHICommandList& RHICmdList, TConstArrayView<FMeshComputeDispatchCommand> DispatchCommands, ERHIPipeline SrcResourceAccessPipelines = ERHIPipeline::Graphics, ERHIPipeline DstResourceAccessPipelines = ERHIPipeline::Graphics);
 	
@@ -82,21 +82,21 @@ private:
 		TArray<FMeshComputeDispatchCommand> DispatchCommands;
 		FRHIUniformBuffer* ViewUniformBuffer = nullptr;
 		FRayTracingGeometry* Geometry = nullptr;
-		int32 SegmentOffset = -1;
 		float Priority = 0.0;
 	};
 
-	void AddDispatchCommands(FRHICommandListBase& RHICmdList,
-							 const FScene* Scene, 
-	                         const FSceneView* View, 
-	                         const FPrimitiveSceneProxy* PrimitiveSceneProxy, 
-	                         const FRayTracingDynamicGeometryUpdateParams& Params,
-	                         uint32 PrimitiveId,
-							 FRWBuffer* RWBuffer,
-							 uint32 NumVertices,
-							 uint32 VertexBufferOffset,
-							 uint32 VertexBufferSize,
-	                         FRayTracingDynamicGeometryBuildParams& BuildParams);
+	void AddDispatchCommands(
+		FRHICommandListBase& RHICmdList,
+		const FScene* Scene,
+		const FSceneView* View,
+		const FPrimitiveSceneProxy* PrimitiveSceneProxy,
+		const FRayTracingDynamicGeometryUpdateParams& Params,
+		uint32 PrimitiveId,
+		FRWBuffer* RWBuffer,
+		uint32 NumVertices,
+		uint32 VertexBufferOffset,
+		uint32 VertexBufferSize,
+		TArray<FMeshComputeDispatchCommand>& OutDispatchCommands);
 
 	FRWBuffer* AllocateSharedBuffer(FRHICommandListBase& RHICmdList, uint32 VertexBufferSize, uint32& OutVertexBufferOffset);
 
@@ -104,10 +104,9 @@ private:
 	TArray<FRayTracingDynamicGeometryBuildParams> DynamicGeometryUpdates;
 
 	// Group dispatch commands per view uniform buffer since it is specified when creating the RDG passes
-	TMap<FRHIUniformBuffer*, TArray<FMeshComputeDispatchCommand>> DispatchCommandsPerView;
+	TMap<FRHIUniformBuffer*, TArray<FMeshComputeDispatchCommand, SceneRenderingAllocator>*> DispatchCommandsPerView;
 	TMap<FRHIUniformBuffer*, FRHIUniformBuffer*> InstancedViewUniformBuffers;
-	TArray<FRayTracingGeometryBuildParams> BuildParams;
-	TArray<FRayTracingGeometrySegment> Segments;
+	TArray<FRayTracingGeometryBuildParams, SceneRenderingAllocator>* BuildParams = nullptr;
 
 	struct FVertexPositionBuffer
 	{
