[HWRT] Improvements to dynamic geometry update budgeting.
- !bAllowUpdate geometries are now budget-constrained.
- Forced-builds also count towards the update budget.
- LastUpdatedFrame is also updated when geometries are rebuilt to avoid being prioritized for update even though that were rebuilt.

#Synced-CL 52610883

diff --git a/Engine/Source/Runtime/Renderer/Private/RayTracing/RayTracingDynamicGeometryUpdateManager.cpp b/Engine/Source/Runtime/Renderer/Private/RayTracing/RayTracingDynamicGeometryUpdateManager.cpp
--- a/Engine/Source/Runtime/Renderer/Private/RayTracing/RayTracingDynamicGeometryUpdateManager.cpp
+++ b/Engine/Source/Runtime/Renderer/Private/RayTracing/RayTracingDynamicGeometryUpdateManager.cpp
@@ -59,6 +59,13 @@ static TAutoConsoleVariable<int32> CVarRTDynGeomForceBuildMaxUpdatePrimitivesPer
 	ECVF_RenderThreadSafe
 );
 
+static TAutoConsoleVariable<int32> CVarRTDynGeomForceBuildMinUpdatesSinceLastBuild(
+	TEXT("r.RayTracing.DynamicGeometry.ForceBuild.MinUpdatesSinceLastBuild"),
+	-1,
+	TEXT("Sets minimum number of updates before the dynamic geometry acceleration structure will be considered for rebuild (default INT_MAX)"),
+	ECVF_RenderThreadSafe
+);
+
 static TAutoConsoleVariable<int32> CVarRTDynGeomPriorityBasedUpdate(
 	TEXT("r.RayTracing.DynamicGeometry.PriorityBasedUpdate"),
 	0,
@@ -80,13 +87,6 @@ static TAutoConsoleVariable<float> CVarRTDynGeomPriorityBasedUpdateLastUpdateWei
 	ECVF_RenderThreadSafe
 );
 
-static TAutoConsoleVariable<int32> CVarRTDynGeomForceBuildMinUpdatesSinceLastBuild(
-	TEXT("r.RayTracing.DynamicGeometry.ForceBuild.MinUpdatesSinceLastBuild"),
-	-1,
-	TEXT("Sets minimum number of updates before the dynamic geometry acceleration structure will be considered for rebuild (default INT_MAX)"),
-	ECVF_RenderThreadSafe
-);
-
 // Lower value means higher priority - max priority is 0.0.
 float CalculatePriority(const FSceneView* View,	const FPrimitiveSceneProxy* PrimitiveSceneProxy, const FRayTracingDynamicGeometryUpdateParams& UpdateParams)
 {
@@ -492,7 +492,7 @@ void FRayTracingDynamicGeometryUpdateManager::AddDynamicGeometryToUpdate(
 
 	GeometryBuildParams.Geometry = UpdateParams.Geometry;
 
-	if (Geometry.GetRequiresBuild() || !Geometry.Initializer.bAllowUpdate)
+	if (Geometry.GetRequiresBuild())
 	{
 		DynamicGeometryBuilds.Add(GeometryBuildParams);
 	}
@@ -735,32 +735,42 @@ void FRayTracingDynamicGeometryUpdateManager::ScheduleUpdates(FRDGBuilder& Graph
 		ViewDispatchCommandsPair.Value->Reserve(TotalNumGeometryBuilds);
 	}
 
-	const uint32 ScratchAlignment = GRHIRayTracingScratchBufferAlignment;
-
 	uint32 BLASScratchSize = 0;
-	int32 NumBuildPrimitives = 0;
+	uint32 NumBuildPrimitives = 0;
 
-	auto AddGeometryBuildParamsToBuildList = [this, &GraphBuilder, &BLASScratchSize](const FRayTracingDynamicGeometryBuildParams& InBuildParams)
+	auto AddGeometryBuildParamsToBuildList = [this, &GraphBuilder, &BLASScratchSize, &NumBuildPrimitives](const FRayTracingDynamicGeometryBuildParams& InBuildParams, EAccelerationStructureBuildMode BuildMode)
 		{
-			FRHIRayTracingGeometry* RayTracingGeometry = InBuildParams.Geometry->GetRHI();
+			FRayTracingGeometry& RayTracingGeometry = *InBuildParams.Geometry;
+			FRHIRayTracingGeometry* RayTracingGeometryRHI = InBuildParams.Geometry->GetRHI();
 			const FRayTracingGeometryInitializer& GeometryInitializer = InBuildParams.Geometry->GetInitializer();
 
-			const uint32 ScratchAlignment = GRHIRayTracingScratchBufferAlignment;
-			const uint32 NumBuildPrimitives = GeometryInitializer.TotalPrimitiveCount;
+			if (BuildMode == EAccelerationStructureBuildMode::Build)
+			{
+				RayTracingGeometry.NumUpdatesSinceLastBuild = 0;
+				RayTracingGeometry.SetRequiresBuild(false);
+				RayTracingGeometry.SetRequiresUpdate(false);
+
+				const uint32 ScratchSize = RayTracingGeometryRHI->GetSizeInfo().BuildScratchSize;
+				BLASScratchSize = Align(BLASScratchSize + ScratchSize, GRHIRayTracingScratchBufferAlignment);
 
-			const uint32 ScratchSize = RayTracingGeometry->GetSizeInfo().BuildScratchSize;
-			BLASScratchSize = Align(BLASScratchSize + ScratchSize, ScratchAlignment);
+				NumBuildPrimitives += GeometryInitializer.TotalPrimitiveCount;
+			}
+			else
+			{
+				RayTracingGeometry.NumUpdatesSinceLastBuild++;
+				RayTracingGeometry.SetRequiresUpdate(false);
 
-			InBuildParams.Geometry->NumUpdatesSinceLastBuild = 0;
-			InBuildParams.Geometry->SetRequiresUpdate(false);
+				const uint32 ScratchSize = RayTracingGeometryRHI->GetSizeInfo().UpdateScratchSize;
+				BLASScratchSize = Align(BLASScratchSize + ScratchSize, GRHIRayTracingScratchBufferAlignment);
+			}
 
 			// Make a copy of GeometryInitializer.Segments that can be passed to RDG pass
 			TArray<FRayTracingGeometrySegment, SceneRenderingAllocator>& Segments = GraphBuilder.AllocArray<FRayTracingGeometrySegment>();
 			Segments.Append(GeometryInitializer.Segments);
 
 			FRayTracingGeometryBuildParams RTGeoBuildParams;
-			RTGeoBuildParams.Geometry = RayTracingGeometry;
-			RTGeoBuildParams.BuildMode = EAccelerationStructureBuildMode::Build;
+			RTGeoBuildParams.Geometry = RayTracingGeometryRHI;
+			RTGeoBuildParams.BuildMode = BuildMode;
 			RTGeoBuildParams.Segments = Segments;
 
 			BuildParams->Add(MoveTemp(RTGeoBuildParams));
@@ -773,8 +783,7 @@ void FRayTracingDynamicGeometryUpdateManager::ScheduleUpdates(FRDGBuilder& Graph
 
 	for (const FRayTracingDynamicGeometryBuildParams& Build : DynamicGeometryBuilds)
 	{
-		AddGeometryBuildParamsToBuildList(Build);
-		NumBuildPrimitives += Build.Geometry->GetInitializer().TotalPrimitiveCount;
+		AddGeometryBuildParamsToBuildList(Build, EAccelerationStructureBuildMode::Build);
 	}
 
 	const uint32 MaxUpdatePrimitivesPerFrame = uint32(CVarRTDynGeomMaxUpdatePrimitivesPerFrame.GetValueOnRenderThread());
@@ -820,42 +829,24 @@ void FRayTracingDynamicGeometryUpdateManager::ScheduleUpdates(FRDGBuilder& Graph
 			continue;
 		}
 
-		if (MaxForceBuildPrimitivesPerFrame > 0)
+		if (MaxForceBuildPrimitivesPerFrame > 0 && Update.Geometry->NumUpdatesSinceLastBuild > MinUpdatesSinceLastBuild && NumForceBuildPrimitives <= MaxForceBuildPrimitivesPerFrame)
 		{
-			if (Update.Geometry->NumUpdatesSinceLastBuild > MinUpdatesSinceLastBuild && NumForceBuildPrimitives <= MaxForceBuildPrimitivesPerFrame)
-			{
-				AddGeometryBuildParamsToBuildList(Update);
-				NumBuildPrimitives += TotalPrimitiveCount;
-				NumForceBuildPrimitives += TotalPrimitiveCount;
-				continue;
-			}
+			AddGeometryBuildParamsToBuildList(Update, EAccelerationStructureBuildMode::Build);
+			NumForceBuildPrimitives += TotalPrimitiveCount;
+		}
+		else if (!GeometryInitializer.bAllowUpdate)
+		{
+			AddGeometryBuildParamsToBuildList(Update, EAccelerationStructureBuildMode::Build);
+		}
+		else
+		{
+			AddGeometryBuildParamsToBuildList(Update, EAccelerationStructureBuildMode::Update);
 		}
 
 		Update.Geometry->LastUpdatedFrame = GFrameCounterRenderThread;
-		Update.Geometry->NumUpdatesSinceLastBuild++;
-		Update.Geometry->SetRequiresUpdate(false);
 
 		NumUpdatedPrimitives += TotalPrimitiveCount;
 
-		const uint32 ScratchSize = RayTracingGeometry->GetSizeInfo().UpdateScratchSize;
-		BLASScratchSize = Align(BLASScratchSize + ScratchSize, ScratchAlignment);
-
-		// Make a copy of GeometryInitializer.Segments that can be passed to RDG pass
-		TArray<FRayTracingGeometrySegment, SceneRenderingAllocator>& Segments = GraphBuilder.AllocArray<FRayTracingGeometrySegment>();
-		Segments.Append(GeometryInitializer.Segments);
-
-		FRayTracingGeometryBuildParams BuildParam;
-		BuildParam.Geometry = RayTracingGeometry;
-		BuildParam.BuildMode = EAccelerationStructureBuildMode::Update;
-		BuildParam.Segments = Segments;
-
-		BuildParams->Add(MoveTemp(BuildParam));
-
-		if (!Update.DispatchCommands.IsEmpty())
-		{
-			DispatchCommandsPerView[Update.ViewUniformBuffer]->Append(Update.DispatchCommands);
-		}
-
 		if (NumUpdatedPrimitives > MaxUpdatePrimitivesPerFrame)
 		{
 			break;
