Fix duplicate symbols when doing non-unity and header compilation
#rnx

diff --git a/Engine/Source/Runtime/Renderer/Private/VT/UniquePageList.h b/Engine/Source/Runtime/Renderer/Private/VT/UniquePageList.h
--- a/Engine/Source/Runtime/Renderer/Private/VT/UniquePageList.h
+++ b/Engine/Source/Runtime/Renderer/Private/VT/UniquePageList.h
@@ -97,7 +97,7 @@ inline void FUniquePageList::MergePages(const FUniquePageList* RESTRICT Other)
 	}
 }
 
-void FUniquePageList::SortPages()
+inline void FUniquePageList::SortPages()
 {
 	for (uint32 Index = 0; Index < NumPages; ++Index)
 	{
