FSSS - TAA stabilisation

#tests PC editor with FSSS
#rnx
#Synced-CL 52673527

diff --git a/Engine/Shaders/Private/HeightFogPixelShader.usf b/Engine/Shaders/Private/HeightFogPixelShader.usf
--- a/Engine/Shaders/Private/HeightFogPixelShader.usf
+++ b/Engine/Shaders/Private/HeightFogPixelShader.usf
@@ -69,7 +69,8 @@ RWTexture2D<float>  FogTexture1UAV;
 float2 ViewHalfResolution;
 float2 ViewHalfResolutionRectMin;
 float FogSeparateCompositionResolutionDivider;
-SamplerState SceneColorTextureSampler;
+Texture2D<float4> InputSceneColorTexture;
+SamplerState InputSceneColorTextureSampler;
 [numthreads(THREADGROUP_SIZE, THREADGROUP_SIZE, 1)]
 void RenderSeparateCompositionTexturesCS(uint3 ThreadId : SV_DispatchThreadID)
 {
@@ -225,7 +226,7 @@ void ExponentialPixelMain(
 	const float SceneColorToScatterAmount = saturate(FSSSParams.Coverage * FSSSParams.Albedo * FogStruct.FSSSSceneColorScatteringAmountScale); // *(1.0f / (4.0f * PI));		// Approximation, should use eq. 4 from paper?
 
 	// This is point sampling scene color when using full resolution. At half resolution this is average of 2x2 texels.
-	const float3 SceneColorToScatter = FogTransmittanceToScene * SceneColorToScatterAmount * Texture2DSampleLevel(SceneTexturesStruct.SceneColorTexture, SceneColorTextureSampler, TexCoordSceneTex, 0).rgb;
+	const float3 SceneColorToScatter = FogTransmittanceToScene * SceneColorToScatterAmount * Texture2DSampleLevel(InputSceneColorTexture, InputSceneColorTextureSampler, TexCoordSceneTex, 0).rgb;
 
 	const float FinalFogTransmittanceToScene = FogTransmittanceToScene * saturate(1.0f - SceneColorToScatterAmount);
 
diff --git a/Engine/Source/Runtime/Renderer/Private/DeferredShadingRenderer.cpp b/Engine/Source/Runtime/Renderer/Private/DeferredShadingRenderer.cpp
--- a/Engine/Source/Runtime/Renderer/Private/DeferredShadingRenderer.cpp
+++ b/Engine/Source/Runtime/Renderer/Private/DeferredShadingRenderer.cpp
@@ -3666,8 +3666,10 @@ void FDeferredShadingSceneRenderer::Render(FRDGBuilder& GraphBuilder, const FSce
 				RDG_CSV_STAT_EXCLUSIVE_SCOPE(GraphBuilder, RenderFog);
 				SCOPED_NAMED_EVENT(RenderFog, FColor::Emerald);
 				SCOPE_CYCLE_COUNTER(STAT_FDeferredShadingSceneRenderer_RenderFog);
+
 				if (bShouldRenderFogUsingSeparateComposition)
 				{
+					RDG_EVENT_SCOPE(GraphBuilder, "Fog::UpsampleFogSeparateCompositionTexture");
 					for (int32 ViewIndex = 0; ViewIndex < Views.Num(); ++ViewIndex)
 					{
 						FViewInfo& View = Views[ViewIndex];
diff --git a/Engine/Source/Runtime/Renderer/Private/FogRendering.cpp b/Engine/Source/Runtime/Renderer/Private/FogRendering.cpp
--- a/Engine/Source/Runtime/Renderer/Private/FogRendering.cpp
+++ b/Engine/Source/Runtime/Renderer/Private/FogRendering.cpp
@@ -15,6 +15,8 @@
 #include "EnvironmentComponentsFlags.h"
 #include "SceneProxies/SkyLightSceneProxy.h"
 #include "FogSeparateComposition.h"
+#include "PostProcess/TemporalAA.h"
+
 
 DECLARE_GPU_DRAWCALL_STAT(Fog);
 
@@ -59,6 +61,12 @@ static TAutoConsoleVariable<int32> CVarMobilePixelFogQuality(
 	TEXT("1 - all per-pixel fog features (second fog, directional inscattering, aerial perspective)"),
 	ECVF_Scalability | ECVF_RenderThreadSafe);
 
+static TAutoConsoleVariable<int32> CVarFogScreenSpaceScatteringTAA(
+	TEXT("r.Fog.ScreenSpaceScattering.TAA"),
+	1,
+	TEXT("Enable temporal anti-aliasing filtering in order to reduce flickering of FSSS blurred scene color."),
+	ECVF_Scalability | ECVF_RenderThreadSafe);
+
 bool PlatformUsesBasicFogFeatures(EShaderPlatform ShaderPlatform)
 {
 	// Disables AP, "SecondFog" and "DirectionalLight Inscattering"
@@ -295,7 +303,8 @@ class FRenderSeparateCompositionTexturesCS : public FGlobalShader
 		SHADER_PARAMETER_RDG_TEXTURE(Texture2D, HalfResDitheredDepth)
 		SHADER_PARAMETER_RDG_TEXTURE_UAV(RWTexture2D<float4>, FogTexture0UAV)
 		SHADER_PARAMETER_RDG_TEXTURE_UAV(RWTexture2D<float>,  FogTexture1UAV)
-		SHADER_PARAMETER_SAMPLER(SamplerState, SceneColorTextureSampler)
+		SHADER_PARAMETER_RDG_TEXTURE(Texture2D<float4>, InputSceneColorTexture)
+		SHADER_PARAMETER_SAMPLER(SamplerState, InputSceneColorTextureSampler)
 		SHADER_PARAMETER(FVector2f, ViewHalfResolution)
 		SHADER_PARAMETER(FVector2f, ViewHalfResolutionRectMin)
 		SHADER_PARAMETER(float, FogSeparateCompositionResolutionDivider)
@@ -727,12 +736,11 @@ void FDeferredShadingSceneRenderer::RenderFogSeparateCompositionTextures(
 {
 	if (ShouldRenderDeferredFog(Scene))
 	{
-		RDG_EVENT_SCOPE_STAT(GraphBuilder, Fog, "Fog::RenderSeparateCompositionTextures");
-
 		check(bFogComposeLocalFogVolumes == true);
 		check(SceneTextures.DitheredHalfResDepth);
 
 		const bool bShouldRenderVolumetricFog = ShouldRenderVolumetricFog();
+		const bool bSceneRequestsFSSS = SceneRequestsFSSS(Scene);
 
 		for(int32 ViewIndex = 0;ViewIndex < Views.Num();ViewIndex++)
 		{
@@ -747,73 +755,99 @@ void FDeferredShadingSceneRenderer::RenderFogSeparateCompositionTextures(
 			FFogSeparateCompositionViewResources& FogSeparateCompositionViewResourcesIt = FogSeparateCompositionViewResources[ViewIndex];
 			FogSeparateCompositionAllocateResources(GraphBuilder, FogSeparateCompositionViewResourcesIt, View, Scene);
 			const FIntPoint FogViewExtent = FogSeparateCompositionViewResourcesIt.ViewRect.Size();
-			
+
 			RDG_EVENT_SCOPE_CONDITIONAL(GraphBuilder, bShouldRenderFogForView && Views.Num() > 1, "View%d - %dx%d", ViewIndex, FogViewExtent.X, FogViewExtent.Y);
 			RDG_GPU_MASK_SCOPE(GraphBuilder, View.GPUMask);
 
-			TRDGUniformBufferRef<FFogUniformParameters> FogUniformBuffer = CreateFogUniformBuffer(GraphBuilder, View);
-
-			const FScreenPassTextureViewport SceneViewport(SceneTextures.Config.Extent, View.ViewRect);
-			const FScreenPassTextureViewport OutputViewport(GetDownscaledViewport(SceneViewport, GetLightShaftDownsampleFactor()));
-			const FScreenPassTextureViewportParameters LightShaftParameters = GetScreenPassTextureViewportParameters(OutputViewport);
-			 
-			FRenderSeparateCompositionTexturesCS::FParameters* PassParameters = GraphBuilder.AllocParameters<FRenderSeparateCompositionTexturesCS::FParameters>();
-
-			CreateDefaultCommonFogPassParameters(
-				GraphBuilder,
-				View,
-				FogUniformBuffer,
-				LightShaftOcclusionTexture,
-				LightShaftParameters,
-				PassParameters->FogCommon);
-
-			PassParameters->SceneTextures = CreateSceneTextureUniformBuffer(GraphBuilder, View, ESceneTextureSetupMode::SceneColor);
-			PassParameters->HalfResDitheredDepth = FogSeparateCompositionViewResourcesIt.ResolutionDivider == 2 ? SceneTextures.DitheredHalfResDepth : SceneTextures.Depth.Resolve;
-			PassParameters->FogTexture0UAV = GraphBuilder.CreateUAV(FRDGTextureUAVDesc(FogSeparateCompositionViewResourcesIt.FogSeparateCompositionTexture0));
-			PassParameters->FogTexture1UAV = GraphBuilder.CreateUAV(FRDGTextureUAVDesc(FogSeparateCompositionViewResourcesIt.FogSeparateCompositionTexture1));
-			PassParameters->SceneColorTextureSampler = TStaticSamplerState<SF_Bilinear, AM_Clamp, AM_Clamp, AM_Clamp>::GetRHI();
-			PassParameters->ViewHalfResolution = FVector2f(FogViewExtent.X, FogViewExtent.Y);
-			PassParameters->ViewHalfResolutionRectMin = FVector2f(FogSeparateCompositionViewResourcesIt.ViewRect.Min.X, FogSeparateCompositionViewResourcesIt.ViewRect.Min.Y);
-			PassParameters->FogSeparateCompositionResolutionDivider = float(FogSeparateCompositionViewResourcesIt.ResolutionDivider);
-			FIntVector NumGroups = FIntVector::DivideAndRoundUp(FIntVector(FogViewExtent.X, FogViewExtent.Y, 1), FIntVector(FRenderSeparateCompositionTexturesCS::ThreadGroupSize, FRenderSeparateCompositionTexturesCS::ThreadGroupSize, 1));
-
-			using FPermutationDomain = TShaderPermutationDomain<FSupportFogInScatteringTexture, FVolumetricFogFilteringQuality, FSupportFogDirectionalLightInScattering, FSupportVolumetricFog, FSupportLocalFogVolume, FSampleSkyAtmosphereOnOpaque>;
-
-			const bool bUseFogInscatteringColorCubemap = View.FogInscatteringColorCubemap != nullptr;
-			FRenderSeparateCompositionTexturesCS::FPermutationDomain PermutationVector;
-			PermutationVector.Set<FSupportFogInScatteringTexture>(bUseFogInscatteringColorCubemap);
-
-			int VolumetricFogFilteringQuality = View.CachedViewUniformShaderParameters->RenderingReflectionCaptureMask == 0.0f ? CVarVolumetricFogFilteringQuality.GetValueOnAnyThread() : 0;
-			PermutationVector.Set<FVolumetricFogFilteringQuality>(FMath::Clamp(VolumetricFogFilteringQuality, 0, 2));
-
-			PermutationVector.Set<FSupportFogDirectionalLightInScattering>(!bUseFogInscatteringColorCubemap && View.bUseDirectionalInscattering);
-			PermutationVector.Set<FSupportVolumetricFog>(bShouldRenderVolumetricFog);
-			PermutationVector.Set<FSupportLocalFogVolume>(bFogComposeLocalFogVolumes);
-			PermutationVector.Set<FSampleSkyAtmosphereOnOpaque>(bRenderSkyAtmosphereInLayer);
-			TShaderMapRef<FRenderSeparateCompositionTexturesCS> ComputeShader(GetGlobalShaderMap(FeatureLevel), PermutationVector);
-
-			ERDGPassFlags PassFlag = ERDGPassFlags::Compute;
-
-			// Dispatch with GenerateMips: reading from a slice through SRV and writing into lower mip through UAV.
-			ClearUnusedGraphResources(ComputeShader, PassParameters);
-			GraphBuilder.AddPass(
-				Forward<FRDGEventName>(RDG_EVENT_NAME("Fog::RenderSeparateCompositionTextures")),
-				PassParameters,
-				PassFlag,
-				[PassParameters, ComputeShader, NumGroups](FRDGAsyncTask, FRHICommandList& RHICmdList)
-				{
-					FComputeShaderUtils::Dispatch(RHICmdList, ComputeShader, *PassParameters, NumGroups);
-				});
+			FRDGTextureRef ViewFSSSInputSceneColorTexture = SceneTextures.Color.Target;
+			if (bSceneRequestsFSSS && View.ViewState && CVarFogScreenSpaceScatteringTAA.GetValueOnRenderThread() > 0)
+			{
+				FTAAPassParameters TAASettings(View);
+				TAASettings.SceneDepthTexture = SceneTextures.Depth.Resolve;
+
+				// We need a valid velocity buffer texture. Use black (no velocity) if it's not produced.
+				TAASettings.SceneVelocityTexture = GetIfProduced(SceneTextures.Velocity, GraphBuilder.RegisterExternalTexture(GSystemTextures.BlackDummy));
+				TAASettings.Pass = ETAAPassConfig::Main;		// Reusing main config for now.
+				TAASettings.SceneColorInput = ViewFSSSInputSceneColorTexture;
+				TAASettings.bOutputRenderTargetable = true;
+				TAASettings.CoCBilateralFilterStrength = 0.0f;
+				TAASettings.SetupViewRect(View);
+
+				FTAAOutputs TAAOutputs = AddTemporalAAPass(
+					GraphBuilder,
+					View,
+					TAASettings,
+					View.PrevViewInfo.FSSSHistory,
+					&View.ViewState->PrevFrameViewInfo.FSSSHistory);
+
+				ViewFSSSInputSceneColorTexture = TAAOutputs.SceneColor;
+			}
 
-			if (SceneRequestsFSSS(Scene))
 			{
-				// Now generate the blur mip chain.
-				const FExponentialHeightFogSceneInfo& HFog = Scene->ExponentialFogs[0];
-				FFogSSSBlurMipChainParameters BlurMipChainParameters;
-				BlurMipChainParameters.TextureWithMipChain = FogSeparateCompositionViewResourcesIt.FogSeparateCompositionTexture0;
-				BlurMipChainParameters.UpsampleBlurFactor = HFog.FSSSBlurControl;
-				// BlurMipChainParameters.BlurMaxExposedLuminance
-				GenerateFSSSBlurMipChain(GraphBuilder, BlurMipChainParameters);
+				TRDGUniformBufferRef<FFogUniformParameters> FogUniformBuffer = CreateFogUniformBuffer(GraphBuilder, View);
+
+				const FScreenPassTextureViewport SceneViewport(SceneTextures.Config.Extent, View.ViewRect);
+				const FScreenPassTextureViewport OutputViewport(GetDownscaledViewport(SceneViewport, GetLightShaftDownsampleFactor()));
+				const FScreenPassTextureViewportParameters LightShaftParameters = GetScreenPassTextureViewportParameters(OutputViewport);
+
+				FRenderSeparateCompositionTexturesCS::FParameters* PassParameters = GraphBuilder.AllocParameters<FRenderSeparateCompositionTexturesCS::FParameters>();
+
+				CreateDefaultCommonFogPassParameters(
+					GraphBuilder,
+					View,
+					FogUniformBuffer,
+					LightShaftOcclusionTexture,
+					LightShaftParameters,
+					PassParameters->FogCommon);
+
+				PassParameters->HalfResDitheredDepth = FogSeparateCompositionViewResourcesIt.ResolutionDivider == 2 ? SceneTextures.DitheredHalfResDepth : SceneTextures.Depth.Resolve;
+				PassParameters->FogTexture0UAV = GraphBuilder.CreateUAV(FRDGTextureUAVDesc(FogSeparateCompositionViewResourcesIt.FogSeparateCompositionTexture0));
+				PassParameters->FogTexture1UAV = GraphBuilder.CreateUAV(FRDGTextureUAVDesc(FogSeparateCompositionViewResourcesIt.FogSeparateCompositionTexture1));
+				PassParameters->InputSceneColorTexture = ViewFSSSInputSceneColorTexture;
+				PassParameters->InputSceneColorTextureSampler = TStaticSamplerState<SF_Bilinear, AM_Clamp, AM_Clamp, AM_Clamp>::GetRHI();
+				PassParameters->ViewHalfResolution = FVector2f(FogViewExtent.X, FogViewExtent.Y);
+				PassParameters->ViewHalfResolutionRectMin = FVector2f(FogSeparateCompositionViewResourcesIt.ViewRect.Min.X, FogSeparateCompositionViewResourcesIt.ViewRect.Min.Y);
+				PassParameters->FogSeparateCompositionResolutionDivider = float(FogSeparateCompositionViewResourcesIt.ResolutionDivider);
+				FIntVector NumGroups = FIntVector::DivideAndRoundUp(FIntVector(FogViewExtent.X, FogViewExtent.Y, 1), FIntVector(FRenderSeparateCompositionTexturesCS::ThreadGroupSize, FRenderSeparateCompositionTexturesCS::ThreadGroupSize, 1));
+
+				using FPermutationDomain = TShaderPermutationDomain<FSupportFogInScatteringTexture, FVolumetricFogFilteringQuality, FSupportFogDirectionalLightInScattering, FSupportVolumetricFog, FSupportLocalFogVolume, FSampleSkyAtmosphereOnOpaque>;
+
+				const bool bUseFogInscatteringColorCubemap = View.FogInscatteringColorCubemap != nullptr;
+				FRenderSeparateCompositionTexturesCS::FPermutationDomain PermutationVector;
+				PermutationVector.Set<FSupportFogInScatteringTexture>(bUseFogInscatteringColorCubemap);
+
+				int VolumetricFogFilteringQuality = View.CachedViewUniformShaderParameters->RenderingReflectionCaptureMask == 0.0f ? CVarVolumetricFogFilteringQuality.GetValueOnAnyThread() : 0;
+				PermutationVector.Set<FVolumetricFogFilteringQuality>(FMath::Clamp(VolumetricFogFilteringQuality, 0, 2));
+
+				PermutationVector.Set<FSupportFogDirectionalLightInScattering>(!bUseFogInscatteringColorCubemap && View.bUseDirectionalInscattering);
+				PermutationVector.Set<FSupportVolumetricFog>(bShouldRenderVolumetricFog);
+				PermutationVector.Set<FSupportLocalFogVolume>(bFogComposeLocalFogVolumes);
+				PermutationVector.Set<FSampleSkyAtmosphereOnOpaque>(bRenderSkyAtmosphereInLayer);
+				TShaderMapRef<FRenderSeparateCompositionTexturesCS> ComputeShader(GetGlobalShaderMap(FeatureLevel), PermutationVector);
+
+				ERDGPassFlags PassFlag = ERDGPassFlags::Compute;
+
+				// Dispatch with GenerateMips: reading from a slice through SRV and writing into lower mip through UAV.
+				ClearUnusedGraphResources(ComputeShader, PassParameters);
+				GraphBuilder.AddPass(
+					Forward<FRDGEventName>(RDG_EVENT_NAME("Fog::RenderSeparateCompositionTextures")),
+					PassParameters,
+					PassFlag,
+					[PassParameters, ComputeShader, NumGroups](FRDGAsyncTask, FRHICommandList& RHICmdList)
+					{
+						FComputeShaderUtils::Dispatch(RHICmdList, ComputeShader, *PassParameters, NumGroups);
+					});
+
+				if (bSceneRequestsFSSS)
+				{
+					// Now generate the blur mip chain.
+					const FExponentialHeightFogSceneInfo& HFog = Scene->ExponentialFogs[0];
+					FFogSSSBlurMipChainParameters BlurMipChainParameters;
+					BlurMipChainParameters.TextureWithMipChain = FogSeparateCompositionViewResourcesIt.FogSeparateCompositionTexture0;
+					BlurMipChainParameters.UpsampleBlurFactor = HFog.FSSSBlurControl;
+					// BlurMipChainParameters.BlurMaxExposedLuminance
+					GenerateFSSSBlurMipChain(GraphBuilder, BlurMipChainParameters);
+				}
 			}
 		}
 	}
diff --git a/Engine/Source/Runtime/Renderer/Private/SceneRendering.h b/Engine/Source/Runtime/Renderer/Private/SceneRendering.h
--- a/Engine/Source/Runtime/Renderer/Private/SceneRendering.h
+++ b/Engine/Source/Runtime/Renderer/Private/SceneRendering.h
@@ -1122,6 +1122,9 @@ struct FPreviousViewInfo
 	// Temporal AA history for Rough refraction
 	FTemporalAAHistory RoughRefractionHistory;
 
+	// Temporal AA history for Fog Screen Space Scattering
+	FTemporalAAHistory FSSSHistory;
+
 	// Temporal AA history for Hair
 	FTemporalAAHistory HairHistory;
 
diff --git a/Engine/Source/Runtime/Renderer/Private/SceneViewState.cpp b/Engine/Source/Runtime/Renderer/Private/SceneViewState.cpp
--- a/Engine/Source/Runtime/Renderer/Private/SceneViewState.cpp
+++ b/Engine/Source/Runtime/Renderer/Private/SceneViewState.cpp
@@ -188,6 +188,7 @@ uint64 FPreviousViewInfo::GetGPUSizeBytes(bool bLogSizes) const
 		SSRHistory.GetGPUSizeBytes(bLogSizes) +
 		WaterSSRHistory.GetGPUSizeBytes(bLogSizes) +
 		RoughRefractionHistory.GetGPUSizeBytes(bLogSizes) +
+		FSSSHistory.GetGPUSizeBytes(bLogSizes) +
 		HairHistory.GetGPUSizeBytes(bLogSizes) +
 #if UE_ENABLE_DEBUG_DRAWING
 		CompositePrimitiveDepthHistory.GetGPUSizeBytes(bLogSizes) +
