FSSS - fix for compilation with new ViewState.

#tests PC editor
#rnx
#Synced-CL 52675788

diff --git a/Engine/Source/Runtime/Renderer/Private/FogRendering.cpp b/Engine/Source/Runtime/Renderer/Private/FogRendering.cpp
--- a/Engine/Source/Runtime/Renderer/Private/FogRendering.cpp
+++ b/Engine/Source/Runtime/Renderer/Private/FogRendering.cpp
@@ -16,6 +16,7 @@
 #include "SceneProxies/SkyLightSceneProxy.h"
 #include "FogSeparateComposition.h"
 #include "PostProcess/TemporalAA.h"
+#include "SceneViewState.h"
 
 
 DECLARE_GPU_DRAWCALL_STAT(Fog);
