Exempt views using mobile multi-view (GetMultiVewCount() > 1) from the overlapping view rect check introduced in CL 52154273.
Views with GetMultiViewCount() > 1 are expected to have overlapping view rects as they are stored in different slices of a texture array.

#Synced-CL 52640403

diff --git a/Engine/Source/Runtime/Renderer/Private/SceneRendering.cpp b/Engine/Source/Runtime/Renderer/Private/SceneRendering.cpp
--- a/Engine/Source/Runtime/Renderer/Private/SceneRendering.cpp
+++ b/Engine/Source/Runtime/Renderer/Private/SceneRendering.cpp
@@ -3690,6 +3690,12 @@ void FSceneRenderer::PrepareViewRectsForRendering()
 			i, ViewRect.Min.X, ViewRect.Min.Y, ViewRect.Max.X, ViewRect.Max.Y,
 			UpperBound.Min.X, UpperBound.Min.Y, UpperBound.Max.X, UpperBound.Max.Y);
 
+		// ViewRect overlap is expected for views using mobile multi-view/texture arrays.
+		if (Views[i].GetMultiViewCount() > 1)
+		{
+			continue;
+		}
+
 		for (int32 j = i + 1; j < Views.Num(); j++)
 		{
 			const FIntRect& OtherUpperBound = Views[j].UpperBoundViewRect;
