Fixed race condition with previous frame copy pass. Added optimization to skip re-registering resources if they were done already.

#Synced-CL 52687437

diff --git a/Engine/Source/Runtime/Renderer/Private/Skinning/SkinningSceneExtension.cpp b/Engine/Source/Runtime/Renderer/Private/Skinning/SkinningSceneExtension.cpp
--- a/Engine/Source/Runtime/Renderer/Private/Skinning/SkinningSceneExtension.cpp
+++ b/Engine/Source/Runtime/Renderer/Private/Skinning/SkinningSceneExtension.cpp
@@ -125,6 +125,15 @@ END_UNIFORM_BUFFER_STRUCT()
 
 DECLARE_SCENE_UB_STRUCT(FSkinningSceneParameters, Skinning, RENDERER_API)
 
+// Blackboard struct to cache the result of FinishSkinningBufferUpload.
+// Avoids redundant buffer uploads when the function is called multiple
+// times within the same RDG builder (PreGPUSceneUpdate + UpdateSceneUniformBuffer).
+struct FSkinningBufferUploadResult
+{
+	FSkinningSceneParameters Parameters;
+};
+RDG_REGISTER_BLACKBOARD_STRUCT(FSkinningBufferUploadResult)
+
 // Reference pose transform provider
 struct FTransformBlockHeader
 {
@@ -321,6 +330,16 @@ void FSkinningSceneExtension::FinishSkinningBufferUpload(FRDGBuilder& GraphBuild
 		return;
 	}
 
+	// If we've already uploaded this builder, just return the cached parameters.
+	if (const FSkinningBufferUploadResult* Cached = GraphBuilder.Blackboard.Get<FSkinningBufferUploadResult>())
+	{
+		if (OutParams)
+		{
+			*OutParams = Cached->Parameters;
+		}
+		return;
+	}
+
 	FRDGBufferRef HeaderBuffer = nullptr;
 	FRDGBufferRef BoneMapBuffer = nullptr;
 	FRDGBufferRef BoneHierarchyBuffer = nullptr;
@@ -407,13 +426,18 @@ void FSkinningSceneExtension::FinishSkinningBufferUpload(FRDGBuilder& GraphBuild
 		bVertexFactoryFullUpload = true;
 	}
 
+	// Cache the result in the blackboard so subsequent calls within the same
+	// RDG builder can skip the upload work entirely.
+	FSkinningBufferUploadResult& CachedResult = GraphBuilder.Blackboard.Create<FSkinningBufferUploadResult>();
+	CachedResult.Parameters.Headers			= GraphBuilder.CreateSRV(HeaderBuffer);
+	CachedResult.Parameters.BoneMap			= GraphBuilder.CreateSRV(BoneMapBuffer);
+	CachedResult.Parameters.BoneHierarchy	= GraphBuilder.CreateSRV(BoneHierarchyBuffer);
+	CachedResult.Parameters.BoneObjectSpace	= GraphBuilder.CreateSRV(BoneObjectSpaceBuffer);
+	CachedResult.Parameters.BoneTransforms	= GetCompressedBoneTransformSRV(GraphBuilder, TransformBuffer);
+
 	if (OutParams != nullptr)
 	{
-		OutParams->Headers			= GraphBuilder.CreateSRV(HeaderBuffer);
-		OutParams->BoneMap			= GraphBuilder.CreateSRV(BoneMapBuffer);
-		OutParams->BoneHierarchy	= GraphBuilder.CreateSRV(BoneHierarchyBuffer);
-		OutParams->BoneObjectSpace	= GraphBuilder.CreateSRV(BoneObjectSpaceBuffer);
-		OutParams->BoneTransforms	= GetCompressedBoneTransformSRV(GraphBuilder, TransformBuffer);
+		*OutParams = CachedResult.Parameters;
 	}
 }
 
@@ -1727,6 +1751,10 @@ void FSkinningSceneExtension::FUpdater::PreGPUSceneUpdate(FRDGBuilder& GraphBuil
 	SCOPED_NAMED_EVENT(FSkinningSceneExtension_FUpdater_PreGPUSceneUpdate, FColor::Silver);
 	CSV_SCOPED_TIMING_STAT_EXCLUSIVE(STAT_SkinningSceneExtension);
 
+	// Wait for InitHeaderDataTask (which depends on FreeBufferSpaceTask) to ensure
+	// InstancedHeaderDataIndices is fully up to date before iterating.
+	SceneData->TaskHandles[InitHeaderDataTask].Wait();
+
 	int32 NumCopies = 0;
 	for (int32 PersistentIndex : SceneData->InstancedHeaderDataIndices)
 	{
