Renderer: Adds a standalone mode (bStandalonePass) to FPostProcessMaterialInputs for AddPostProcessMaterialPass, allowing it to execute outside or before the regular scene render.

* When true, the pass can execute outside of the regular post-processing render stack. This offers a leaner alternative to DrawMaterialToRenderTarget, where bindings can be assigned using regular post-process inputs instead of named texture parameters.
* Other bindings that require a fully initialized rendering pipeline are skipped. The caller is responsible for providing valid or dummy SceneTextures & ViewUniformBuffer.
* Also replaces FRDGSystemTextures::Get().Black with GSystemTextures in the SceneDepthWithoutWater default, removing an unnecessary dependency on the RDG blackboard.

#Synced-CL 52697243

diff --git a/Engine/Source/Runtime/Renderer/Private/PostProcess/PostProcessMaterial.cpp b/Engine/Source/Runtime/Renderer/Private/PostProcess/PostProcessMaterial.cpp
--- a/Engine/Source/Runtime/Renderer/Private/PostProcess/PostProcessMaterial.cpp
+++ b/Engine/Source/Runtime/Renderer/Private/PostProcess/PostProcessMaterial.cpp
@@ -534,14 +534,19 @@ FPostProcessMaterialParameters* GetPostProcessMaterialParameters(
 		PostProcessMaterialParameters->PathTracingPostProcessInput[InputIndex] = GetScreenPassTextureInput(Input, PointClampSampler);
 	}
 
-	PostProcessMaterialParameters->Substrate = Substrate::BindSubstrateGlobalUniformParameters(View);
+	// Substrate requires the full rendering pipeline to be initialized.
+	// Skip for standalone passes running outside the normal post-process pipeline.
+	if (!Inputs.bStandalonePass)
+	{
+		PostProcessMaterialParameters->Substrate = Substrate::BindSubstrateGlobalUniformParameters(View);
+	}
 
 	// SceneDepthWithoutWater
 	const bool bHasValidSceneDepthWithoutWater = Inputs.SceneWithoutWaterTextures && Inputs.SceneWithoutWaterTextures->DepthTexture;
 	const bool bShouldUseBilinearSamplerForDepth = bHasValidSceneDepthWithoutWater && ShouldUseBilinearSamplerForDepthWithoutSingleLayerWater(Inputs.SceneWithoutWaterTextures->DepthTexture->Desc.Format);
 	PostProcessMaterialParameters->bSceneDepthWithoutWaterTextureAvailable = bHasValidSceneDepthWithoutWater;
 	PostProcessMaterialParameters->SceneDepthWithoutSingleLayerWaterSampler = bShouldUseBilinearSamplerForDepth ? TStaticSamplerState<SF_Bilinear>::GetRHI() : TStaticSamplerState<SF_Point>::GetRHI();
-	PostProcessMaterialParameters->SceneDepthWithoutSingleLayerWaterTexture = FRDGSystemTextures::Get(GraphBuilder).Black;
+	PostProcessMaterialParameters->SceneDepthWithoutSingleLayerWaterTexture = GSystemTextures.GetBlackDummy(GraphBuilder);
 	PostProcessMaterialParameters->SceneWithoutSingleLayerWaterMinMaxUV = FVector4f(0.0f, 0.0f, 1.0f, 1.0f);
 	PostProcessMaterialParameters->SceneWithoutSingleLayerWaterTextureSize = FVector2f(0.0f, 0.0f);
 	PostProcessMaterialParameters->SceneWithoutSingleLayerWaterInvTextureSize = FVector2f(0.0f, 0.0f);
diff --git a/Engine/Source/Runtime/Renderer/Public/PostProcess/PostProcessMaterialInputs.h b/Engine/Source/Runtime/Renderer/Public/PostProcess/PostProcessMaterialInputs.h
--- a/Engine/Source/Runtime/Renderer/Public/PostProcess/PostProcessMaterialInputs.h
+++ b/Engine/Source/Runtime/Renderer/Public/PostProcess/PostProcessMaterialInputs.h
@@ -174,6 +174,16 @@ struct FPostProcessMaterialInputs
 	bool bUserSceneTextureOutput = false;
 	bool bUserSceneTextureFirstRender = false;
 	uint32 UserSceneTextureSceneColorInput = INDEX_NONE;
+
+	/**
+	 * When true, the AddPostProcessMaterialPass can execute outside of the regular render post-processing stack.
+	 * This offers a leaner alternative to DrawMaterialToRenderTarget, where bindings can be assigned using regular
+	 * post-process inputs instead of named texture parameters.
+	 * 
+	 * Other bindings that require a fully initialized rendering pipeline are skipped.
+	 * The caller remains responsible for providing valid (dummy) SceneTextures & ViewUniformBuffer.
+	 */
+	bool bStandalonePass = false;
 };
 
 class UMaterialInterface;
