[Mobile][Performance] Adding a mechanism for the experimental that does the following:
- Adds in a check for sphere radiuses that are too small or degenerate. If these objects by default are too small (< 1cm in normal UE units) we willl not culll them.
- Add in a near cutoff for the screen bounds check, where if the object is close enough we will not cull it even if the projected bounds woulld be too small.

#tests Tested on local iOS and Android devices. Tested with preflight built with supported build.

#Synced-CL 52608389

diff --git a/Engine/Source/Runtime/Renderer/Private/SceneVisibility.cpp b/Engine/Source/Runtime/Renderer/Private/SceneVisibility.cpp
--- a/Engine/Source/Runtime/Renderer/Private/SceneVisibility.cpp
+++ b/Engine/Source/Runtime/Renderer/Private/SceneVisibility.cpp
@@ -679,7 +679,7 @@ inline bool ShouldCullForRayTracing(const FScene& Scene, FViewInfo& View, int32
 // to reduce popin, as well as seeing what it would take to piggyback off of other features such as closeness to
 // screen center, view relevance, etc.
 float GMinScreenBoundsRadius = 0.0f;
-static FAutoConsoleVariableRef CVarFrustumCullMaxScreenSpaceBoundsRadius(
+static FAutoConsoleVariableRef CVarFrustumCullMinScreenBoundsRadius(
 	TEXT("r.MinScreenBoundsRadius.Experimental"),
 	GMinScreenBoundsRadius,
 	TEXT("Specifies a minimum value for the screen space radius of the projected bounding sphere for visibility tests. ")
@@ -687,6 +687,25 @@ static FAutoConsoleVariableRef CVarFrustumCullMaxScreenSpaceBoundsRadius(
 	ECVF_Scalability | ECVF_RenderThreadSafe
 );
 
+float GMinScreenBoundsSphereRadiusCutoff = 1.0f;
+static FAutoConsoleVariableRef CVarFrustumCullMinScreenBoundsSphereRadiusCutoff(
+	TEXT("r.MinScreenBoundsSphereRadiusCutoff.Experimental"),
+	GMinScreenBoundsSphereRadiusCutoff,
+	TEXT("Specifies a minimum value for the minimum size of a sphere radius needed for testing objects. ")
+	TEXT("This helps to account for objects with near zero bounds that have not been properly updated or ")
+	TEXT("are setup to not be dependent on their bounds for rendering checks. Defaults to 1.0 (in world units)."),
+	ECVF_Scalability | ECVF_RenderThreadSafe
+);
+
+float GMinScreenBoundsIgnoreDistance = 0.0f;
+static FAutoConsoleVariableRef CVarFrustumCullMinScreenBoundsIgnoreDistance(
+	TEXT("r.MinScreenBoundsIgnoreDistance.Experimental"),
+	GMinScreenBoundsIgnoreDistance,
+	TEXT("Specifies a minimum distance at which screen bounds radius testing will begin and anything less will be ignored. ")
+	TEXT("Value must be from 0.0 (disabled) and upwards (in world units). "),
+	ECVF_Scalability | ECVF_RenderThreadSafe
+);
+
 float GMinScreenBoundsRadiusRuntimeReflection = 0.02f;
 static FAutoConsoleVariableRef CVarMinScreenBoundsRadiusRuntimeReflection(
 	TEXT("r.MinScreenBoundsRadius.RuntimeReflection"),
@@ -804,7 +823,8 @@ static void UpdateAlwaysVisible(const FScene& Scene, FViewInfo& View, FFrustumCu
 float ScreenVisibilityGetMinScreenBoundsRadiusSquared(const FSceneView& View)
 {
 	// Values above 1 are considered invalid, and disable size culling
-	float MinScreenBoundsRadiusSquared = FMath::Square(View.IsRuntimeReflectionCapture() ? GMinScreenBoundsRadiusRuntimeReflection : GMinScreenBoundsRadius);
+	const float MinRadiusToUse = View.IsRuntimeReflectionCapture() ? GMinScreenBoundsRadiusRuntimeReflection : GMinScreenBoundsRadius;
+	float MinScreenBoundsRadiusSquared = FMath::Square(MinRadiusToUse);
 	if (MinScreenBoundsRadiusSquared > 1.0f)
 	{
 		MinScreenBoundsRadiusSquared = 0.0f;
@@ -813,6 +833,30 @@ float ScreenVisibilityGetMinScreenBoundsRadiusSquared(const FSceneView& View)
 	return MinScreenBoundsRadiusSquared;
 }
 
+bool IsScreenSizeCullingElligible(FViewInfo& View, const FBoxSphereBounds& Bounds, const float MinScreenBoundsRadiusSquared)
+{
+	// Never cull if we do not specify the value
+	if (MinScreenBoundsRadiusSquared <= 0.0f)
+	{
+		return false;
+	}
+
+	// Never cull objects whose sphere radius is not properly set or too small
+	if (Bounds.SphereRadius < GMinScreenBoundsSphereRadiusCutoff)
+	{
+		return false;
+	}
+
+	// Never cull if we are close enough to the object to be culled
+	if (GMinScreenBoundsIgnoreDistance > 0.0f &&
+		((Bounds.Origin - View.CullingOrigin).Length() - Bounds.SphereRadius) < GMinScreenBoundsIgnoreDistance)
+	{
+		return false;
+	}
+
+	return true;
+}
+
 static int32 FrustumCull(const FScene& Scene, FViewInfo& View, FFrustumCullingFlags Flags, float MaxDrawDistanceScale, const FHLODVisibilityState* const HLODState, const FSceneBitArray* VisibleNodes, const FVisibilityTaskConfig& TaskConfig, int32 TaskIndex)
 {
 	QUICK_SCOPE_CYCLE_COUNTER(STAT_FrustumCull_Loop);
@@ -839,7 +883,7 @@ static int32 FrustumCull(const FScene& Scene, FViewInfo& View, FFrustumCullingFl
 	const bool bRayTracingEnabled = IsRayTracingEnabled(View.GetShaderPlatform()) && View.IsRayTracingAllowedForView();
 #endif
 
-	float MinScreenBoundsRadiusSquared = ScreenVisibilityGetMinScreenBoundsRadiusSquared(View);
+	const float MinScreenBoundsRadiusSquared = ScreenVisibilityGetMinScreenBoundsRadiusSquared(View);
 
 	for (int32 WordIndex = TaskWordOffset; WordIndex < TaskWordOffset + int32(TaskConfig.FrustumCull.NumWordsPerTask) && WordIndex * NumBitsPerDWORD < BitArrayNumInner; WordIndex++)
 	{
@@ -926,7 +970,7 @@ static int32 FrustumCull(const FScene& Scene, FViewInfo& View, FFrustumCullingFl
 				}
 
 				// Screen size culling next
-				if (bIsVisible)
+				if (bIsVisible && IsScreenSizeCullingElligible(View, Bounds.BoxSphereBounds, MinScreenBoundsRadiusSquared))
 				{
 					bIsVisible = !IsScreenSizeCulled(View, Bounds.BoxSphereBounds, MinScreenBoundsRadiusSquared);
 				}
