Mobile: Allow access to GBuffer for NiagaraFX when deferred multi-pass is used
Also fixes DepthBuffer collisions for NiagaraFX
#Synced-CL 52440158

diff --git a/Engine/Plugins/FX/Niagara/Shaders/Private/NiagaraDataInterfaceCollisionQuery.ush b/Engine/Plugins/FX/Niagara/Shaders/Private/NiagaraDataInterfaceCollisionQuery.ush
--- a/Engine/Plugins/FX/Niagara/Shaders/Private/NiagaraDataInterfaceCollisionQuery.ush
+++ b/Engine/Plugins/FX/Niagara/Shaders/Private/NiagaraDataInterfaceCollisionQuery.ush
@@ -3,13 +3,15 @@
 #pragma once
 
 #include "/Engine/Private/Common.ush"
+#include "/Engine/Private/DeferredShadingCommon.ush"
 
 #define NEEDS_SCENE_TEXTURES 1
 #define SCENE_TEXTURES_DISABLED 0
 #define SubstratePublicStruct SubstratePublic
 #include "/Engine/Public/SubstratePublic.ush"
 
-#define NDICOLLISIONQUERY_USE_GBUFFER_NORMAL    (FEATURE_LEVEL >= FEATURE_LEVEL_SM5)
+// Mobile requires multi-pass deferred to access GBuffer
+#define NDICOLLISIONQUERY_USE_GBUFFER_NORMAL    (FEATURE_LEVEL >= FEATURE_LEVEL_SM5) || (MOBILE_DEFERRED_SHADING && !ALLOW_FRAMEBUFFER_FETCH)
 
 float NDICollisionQuery_GetCustomSceneDepth(float2 ScreenUV)
 {
@@ -57,7 +59,8 @@ void NDICollisionQuery_QuerySceneDepthGPU(in float3 In_SamplePos, in float3 In_L
 		uint2 ScreenPos = uint2(ScreenUV * View.BufferSizeAndInvSize.xy);
 		float3 WorldNormal = SubstratePublic_GetWorldNormal(ScreenPos);
 #else
-		float3 WorldNormal = Texture2DSampleLevel(SceneTexturesStruct.GBufferATexture, SceneTexturesStruct_GBufferATextureSampler, ScreenUV, 0).xyz * 2.0 - 1.0;
+		FGBufferData GBuffer = GetGBufferData(ScreenUV, false);
+		float3 WorldNormal = GBuffer.WorldNormal;
 #endif
 #else
 		float CollisionDepthBounds = 500.0f;
diff --git a/Engine/Plugins/FX/Niagara/Shaders/Private/NiagaraDataInterfaceGBufferTemplate.ush b/Engine/Plugins/FX/Niagara/Shaders/Private/NiagaraDataInterfaceGBufferTemplate.ush
--- a/Engine/Plugins/FX/Niagara/Shaders/Private/NiagaraDataInterfaceGBufferTemplate.ush
+++ b/Engine/Plugins/FX/Niagara/Shaders/Private/NiagaraDataInterfaceGBufferTemplate.ush
@@ -2,7 +2,8 @@
 
 #include "/Engine/Private/DeferredShadingCommon.ush"
 
-#define NDIGBUFFER_HAS_GBUFFER				(FEATURE_LEVEL >= FEATURE_LEVEL_SM5)
+// Mobile requires multi-pass deferred to access GBuffer
+#define NDIGBUFFER_HAS_GBUFFER				(FEATURE_LEVEL >= FEATURE_LEVEL_SM5) || (MOBILE_DEFERRED_SHADING && !ALLOW_FRAMEBUFFER_FETCH)
 #define NDIGBUFFER_USE_GBUFFER_VELOCITY		(false && NDIGBUFFER_HAS_GBUFFER)
 #define NDIGBUFFER_IS_MOBILE				(FEATURE_LEVEL <= FEATURE_LEVEL_ES3_1)
 
diff --git a/Engine/Source/Runtime/Renderer/Private/MobileShadingRenderer.cpp b/Engine/Source/Runtime/Renderer/Private/MobileShadingRenderer.cpp
--- a/Engine/Source/Runtime/Renderer/Private/MobileShadingRenderer.cpp
+++ b/Engine/Source/Runtime/Renderer/Private/MobileShadingRenderer.cpp
@@ -1047,7 +1047,8 @@ void FMobileSceneRenderer::RenderCustomRenderPassBasePass(FRDGBuilder& GraphBuil
 		};
 		auto ParameterCollection = GraphBuilder.AllocParameters<FCustomPassParameterCollection>();
 
-		EMobileSceneTextureSetupMode SetupMode = bIsFullDepthPrepassEnabled ? EMobileSceneTextureSetupMode::SceneDepth : EMobileSceneTextureSetupMode::None;
+		EMobileSceneTextureSetupMode SetupMode = bIsFullDepthPrepassEnabled ? 
+			(EMobileSceneTextureSetupMode::SceneDepth | EMobileSceneTextureSetupMode::SceneVelocity) : EMobileSceneTextureSetupMode::None;
 		auto PassParameters = &ParameterCollection->PassParameters;
 		PassParameters->View = View.GetShaderParameters();
 		PassParameters->MobileBasePass = CreateMobileBasePassUniformBuffer(GraphBuilder, View, EMobileBasePass::Opaque, SetupMode, BasePassTextures);
@@ -1664,7 +1665,7 @@ void FMobileSceneRenderer::Render(FRDGBuilder& GraphBuilder, const FSceneRenderU
 			AddResolveSceneDepthPass(GraphBuilder, Views, SceneTextures.Depth);
 		}
 
-		SceneTextures.MobileSetupMode = EMobileSceneTextureSetupMode::SceneDepth;
+		SceneTextures.MobileSetupMode = EMobileSceneTextureSetupMode::SceneDepth | EMobileSceneTextureSetupMode::SceneVelocity;
 		SceneTextures.MobileUniformBuffer = CreateMobileSceneTextureUniformBuffer(GraphBuilder, &SceneTextures, SceneTextures.MobileSetupMode);
 	}
 
@@ -1787,7 +1788,6 @@ void FMobileSceneRenderer::Render(FRDGBuilder& GraphBuilder, const FSceneRenderU
 		}
 
 		SceneTextures.MobileSetupMode = EMobileSceneTextureSetupMode::All;
-		SceneTextures.MobileSetupMode &= ~EMobileSceneTextureSetupMode::SceneVelocity;
 		SceneTextures.MobileUniformBuffer = CreateMobileSceneTextureUniformBuffer(GraphBuilder, &SceneTextures, SceneTextures.MobileSetupMode);
 
 		if (bShouldRenderVelocities)
@@ -1804,8 +1804,13 @@ void FMobileSceneRenderer::Render(FRDGBuilder& GraphBuilder, const FSceneRenderU
 			SceneTextures.MobileUniformBuffer = CreateMobileSceneTextureUniformBuffer(GraphBuilder, &SceneTextures, SceneTextures.MobileSetupMode);
 		}
 
-		FRendererModule& RendererModule = static_cast<FRendererModule&>(GetRendererModule());
-		RendererModule.RenderPostOpaqueExtensions(GraphBuilder, Views, SceneTextures);
+		if (!(bDeferredShading && bRequiresMultiPass))
+		{
+			FRendererModule& RendererModule = static_cast<FRendererModule&>(GetRendererModule());
+			RendererModule.RenderPostOpaqueExtensions(GraphBuilder, Views, SceneTextures);
+			
+			RenderOpaqueFX(GraphBuilder, GetSceneViews(), GetSceneUniforms(), FXSystem, SceneTextures.MobileUniformBuffer);
+		}
 
 		if (Views[0].IsRuntimeReflectionCapture())
 		{
@@ -1815,9 +1820,7 @@ void FMobileSceneRenderer::Render(FRDGBuilder& GraphBuilder, const FSceneRenderU
 		}
 
 		CopySceneCaptureComponentToTarget(GraphBuilder, SceneTextures, ViewFamilyTexture, ViewFamilyDepthTexture, ViewFamily, Views);
-
-		RenderOpaqueFX(GraphBuilder, GetSceneViews(), GetSceneUniforms(), FXSystem, SceneTextures.MobileUniformBuffer);
-
+		
 		if (ViewFamily.EngineShowFlags.VisualizeMeshDistanceFields || ViewFamily.EngineShowFlags.VisualizeGlobalDistanceField)
 		{
 			RenderMeshDistanceFieldVisualization(GraphBuilder, SceneTextures);
@@ -2870,7 +2873,7 @@ void FMobileSceneRenderer::RenderDeferredMultiPass(FRDGBuilder& GraphBuilder, FS
 			});
 	}
 
-	SceneTextures.MobileSetupMode = EMobileSceneTextureSetupMode::SceneDepth | EMobileSceneTextureSetupMode::GBuffers;
+	SceneTextures.MobileSetupMode = EMobileSceneTextureSetupMode::SceneDepth | EMobileSceneTextureSetupMode::GBuffers | EMobileSceneTextureSetupMode::SceneVelocity;
 	SceneTextures.MobileUniformBuffer = CreateMobileSceneTextureUniformBuffer(GraphBuilder, &SceneTextures, SceneTextures.MobileSetupMode);
 
 	if (bShouldRenderHZB && !bIsFullDepthPrepassEnabled)
@@ -2878,6 +2881,11 @@ void FMobileSceneRenderer::RenderDeferredMultiPass(FRDGBuilder& GraphBuilder, FS
 		RenderHZB(GraphBuilder, SceneTextures.Depth.Target);
 	}
 
+	FRendererModule& RendererModule = static_cast<FRendererModule&>(GetRendererModule());
+	RendererModule.RenderPostOpaqueExtensions(GraphBuilder, Views, SceneTextures);
+
+	RenderOpaqueFX(GraphBuilder, GetSceneViews(), GetSceneUniforms(), FXSystem, SceneTextures.MobileUniformBuffer);
+
 	if (bShouldRenderVolumetricFog && !bIsFullDepthPrepassEnabled)
 	{
 		ComputeVolumetricFog(GraphBuilder, SceneTextures);
@@ -2963,7 +2971,8 @@ void FMobileSceneRenderer::RenderDeferredMultiPass(FRDGBuilder& GraphBuilder, FS
 
 		EMobileSceneTextureSetupMode SetupMode = 
 			EMobileSceneTextureSetupMode::SceneDepth | 
-			EMobileSceneTextureSetupMode::CustomDepth | 
+			EMobileSceneTextureSetupMode::CustomDepth |
+			EMobileSceneTextureSetupMode::SceneVelocity |
 			EMobileSceneTextureSetupMode::GBuffers;
 
 		FMobileBasePassUniformParameters* BasePassParameters = GraphBuilder.AllocParameters<FMobileBasePassUniformParameters>();
