Change to HasRayTracedOverlay to ensure ray tracing is actually supported by the RHI at runtime, otherwise render passes are skipped and SceneColor is invalid

#Synced-CL 52187599

diff --git a/Engine/Source/Runtime/Renderer/Private/DeferredShadingRenderer.cpp b/Engine/Source/Runtime/Renderer/Private/DeferredShadingRenderer.cpp
--- a/Engine/Source/Runtime/Renderer/Private/DeferredShadingRenderer.cpp
+++ b/Engine/Source/Runtime/Renderer/Private/DeferredShadingRenderer.cpp
@@ -4361,8 +4361,14 @@ bool ShouldRenderRayTracingEffect(bool bEffectEnabled, ERayTracingPipelineCompat
 bool HasRaytracingDebugViewModeRaytracedOverlay(const FSceneViewFamily& ViewFamily);
 bool HasRayTracedOverlay(const FSceneViewFamily& ViewFamily)
 {
-	// Return true if a full screen ray tracing pass will be displayed on top of the raster pass
-	// This can be used to skip certain calculations
+	// Return true if a full screen ray tracing pass will be displayed on top of the raster pass.
+	// This can be used to skip certain calculations.
+	// Must also verify ray tracing is actually supported by the RHI at runtime.
+	if (ViewFamily.Views.IsEmpty() || !IsRayTracingEnabled(ViewFamily.Views[0]->GetShaderPlatform()))
+	{
+		return false;
+	}
+
 	return
 		ViewFamily.EngineShowFlags.PathTracing ||
 		(ViewFamily.EngineShowFlags.RayTracingDebug && HasRaytracingDebugViewModeRaytracedOverlay(ViewFamily));
