Removed PreSceneUpdate finish call and only cache the RDG result when we request parameters.

#Synced-CL 52714022

diff --git a/Engine/Source/Runtime/Renderer/Private/Skinning/SkinningSceneExtension.cpp b/Engine/Source/Runtime/Renderer/Private/Skinning/SkinningSceneExtension.cpp
--- a/Engine/Source/Runtime/Renderer/Private/Skinning/SkinningSceneExtension.cpp
+++ b/Engine/Source/Runtime/Renderer/Private/Skinning/SkinningSceneExtension.cpp
@@ -426,17 +426,16 @@ void FSkinningSceneExtension::FinishSkinningBufferUpload(FRDGBuilder& GraphBuild
 		bVertexFactoryFullUpload = true;
 	}
 
-	// Cache the result in the blackboard so subsequent calls within the same
-	// RDG builder can skip the upload work entirely.
-	FSkinningBufferUploadResult& CachedResult = GraphBuilder.Blackboard.Create<FSkinningBufferUploadResult>();
-	CachedResult.Parameters.Headers			= GraphBuilder.CreateSRV(HeaderBuffer);
-	CachedResult.Parameters.BoneMap			= GraphBuilder.CreateSRV(BoneMapBuffer);
-	CachedResult.Parameters.BoneHierarchy	= GraphBuilder.CreateSRV(BoneHierarchyBuffer);
-	CachedResult.Parameters.BoneObjectSpace	= GraphBuilder.CreateSRV(BoneObjectSpaceBuffer);
-	CachedResult.Parameters.BoneTransforms	= GetCompressedBoneTransformSRV(GraphBuilder, TransformBuffer);
-
+	// Cache the result in the blackboard so subsequent calls within the same RDG builder can skip the upload work entirely.
 	if (OutParams != nullptr)
 	{
+		FSkinningBufferUploadResult& CachedResult = GraphBuilder.Blackboard.Create<FSkinningBufferUploadResult>();
+		CachedResult.Parameters.Headers			= GraphBuilder.CreateSRV(HeaderBuffer);
+		CachedResult.Parameters.BoneMap			= GraphBuilder.CreateSRV(BoneMapBuffer);
+		CachedResult.Parameters.BoneHierarchy	= GraphBuilder.CreateSRV(BoneHierarchyBuffer);
+		CachedResult.Parameters.BoneObjectSpace	= GraphBuilder.CreateSRV(BoneObjectSpaceBuffer);
+		CachedResult.Parameters.BoneTransforms	= GetCompressedBoneTransformSRV(GraphBuilder, TransformBuffer);
+
 		*OutParams = CachedResult.Parameters;
 	}
 }
@@ -872,11 +871,6 @@ void FSkinningSceneExtension::FUpdater::End()
 
 void FSkinningSceneExtension::FUpdater::PreSceneUpdate(FRDGBuilder& GraphBuilder, const FScenePreUpdateChangeSet& ChangeSet)
 {
-	// If there was a pending upload from a prior update (due to the buffer never being used), finish the upload now.
-	// This keeps the upload entries from growing unbounded and prevents any undefined behavior caused by any
-	// updates that overlap primitives.
-	SceneData->FinishSkinningBufferUpload(GraphBuilder);
-
 	if (!SceneData->IsEnabled())
 	{
 		return;
