FSSS - now applied on volumetric cloud when late fog is enable.
Fixed nan when forcing separate fog composite but FSSS is disabled.

#rnx

#Synced-CL 52723731

#Synced-CL 52741690

diff --git a/Engine/Shaders/Private/FogSeparateComposition.usf b/Engine/Shaders/Private/FogSeparateComposition.usf
--- a/Engine/Shaders/Private/FogSeparateComposition.usf
+++ b/Engine/Shaders/Private/FogSeparateComposition.usf
@@ -58,8 +58,6 @@ void DownsampleColorCS(uint3 ThreadId : SV_DispatchThreadID)
 
 
 
-
-
 #ifdef UPSAMPLECOLORCS
 
 int SrcMipLevel;
@@ -101,69 +99,28 @@ void UpsampleColorCS(uint3 ThreadId : SV_DispatchThreadID)
 
 
 
-#ifdef HALFRESOLUTIONUPSAMPLERPS
-
-int4 HalfResDepthMinMaxCoord;			// min and max of depth inside the hald res buffer were many views lives.
-float2 TexelCoordToUVs;					// Convert a texel to UVs according to view resolution and texture resolution.
-
-float FSSSSpreadScale;
-float FullResToFogResScale; // either 1 for full or 0.5 for half resolution
 
-Texture2D<float4> HalfResolutionSourceTexture;
-SamplerState HalfResolutionSourceSampler;
 
-Texture2D<float> HalfResolutionSource2Texture;
-SamplerState HalfResolutionSource2Sampler;
+#ifdef HALFRESOLUTIONUPSAMPLERPS
 
+#include "FogSeparateComposition.ush"
 
+#if !PERMUTATION_UPSAMPLE_FULL_RESOLUTION
 Texture2D<float> HalfResDepthTexture;
 SamplerState HalfResDepthSampler;
-
 Texture2D<float> FullResDepthTexture;
 SamplerState FullResDepthSampler;
-
-float4 CombineSourceTextures(float4 SourceTextureColor, float SourceTextureColor2)
-{
-	return float4(SourceTextureColor.rgb, SourceTextureColor2);
-}
+#endif
 
 void HalfResolutionUpsamplerPS(
 	float4 SvPosition : SV_POSITION,
 	out float4 OutColor : SV_Target0)
 {
-	float4 SourceTextureColor = 0;
-	float SourceTextureColor2 = 0;
-
 #if PERMUTATION_UPSAMPLE_FULL_RESOLUTION
 
 	float2 FullResLocalCoord = SvPosition.xy - View.ViewRectMinAndSize.xy;
-	SourceTextureColor = HalfResolutionSourceTexture[FullResLocalCoord];
 
-	const float W = SourceTextureColor.a;
-	if (W <= 0.0f)
-	{
-		// Such value will lead to no difference when upsampling
-		OutColor = 0.0;
-		clip(-1);
-		return;
-	}
-
-	SourceTextureColor2 = HalfResolutionSource2Texture[FullResLocalCoord].r;
-
-	if (false)
-	{
-		OutColor = CombineSourceTextures(SourceTextureColor, SourceTextureColor2);
-	}
-	else
-	{
-		const float Mip = max(0.0, FSSSSpreadScale * log2(W / 2));
-
-		// This is a custom fade based on W for the clip above to note cause discontinuity.
-		const float FadeBasedOnW = saturate(W);
-
-		SourceTextureColor = HalfResolutionSourceTexture.SampleLevel(HalfResolutionSourceSampler, FullResLocalCoord * TexelCoordToUVs, Mip).rgba * FadeBasedOnW;
-		OutColor = CombineSourceTextures(SourceTextureColor, SourceTextureColor2);
-	}
+	OutColor = FSSSFullResolutionSample(FullResLocalCoord);
 
 #else // PERMUTATION_UPSAMPLE_FULL_RESOLUTION
 
@@ -171,16 +128,20 @@ void HalfResolutionUpsamplerPS(
 
 	float2 FullResLocalCoord = SvPosition.xy - View.ViewRectMinAndSize.xy;
 
-	float2 HalfResLocalCoord = FullResLocalCoord * FullResToFogResScale;
+	float2 HalfResLocalCoord = FullResLocalCoord * FSSSGlobals.FullResToFogResScale;
 
 	// Debug raw buffer
 #if 0
-	SourceTextureColor = HalfResolutionSourceTexture[HalfResLocalCoord].rgba;
-	SourceTextureColor2 = HalfResolutionSource2Texture[HalfResLocalCoord].r;
-	OutColor = CombineSourceTextures(SourceTextureColor, SourceTextureColor2);
+	SourceTextureColor = FSSSGlobals.FogSeparateCompositionTexture0[HalfResLocalCoord].rgba;
+	SourceTextureColor2 = FSSSGlobals.FogSeparateCompositionTexture1[HalfResLocalCoord].r;
+	OutColor = FSSSCombineSourceTextures(SourceTextureColor, SourceTextureColor2);
 	return;
 #endif
 
+	float4 SourceTextureColor = 0;
+	float SourceTextureColor2 = 0;
+
+
 	int2 HalfResLocalCoordI = int2(HalfResLocalCoord);
 	const int OffsetX = (frac(HalfResLocalCoord.x) > 0.5f ? 1 : -1);
 	const int OffsetY = (frac(HalfResLocalCoord.y) > 0.5f ? 1 : -1);
@@ -191,10 +152,10 @@ void HalfResolutionUpsamplerPS(
 
 	// TODO try GatherRed, but need to make sure UV fetch the 4 valid texel
 	float4 HalfResDepths;
-	HalfResDepths.x = ConvertFromDeviceZ(HalfResDepthTexture[clamp(HalfResDepthMinMaxCoord.xy + HalfResLocalCoordI							, HalfResDepthMinMaxCoord.xy, HalfResDepthMinMaxCoord.zw)]);
-	HalfResDepths.y = ConvertFromDeviceZ(HalfResDepthTexture[clamp(HalfResDepthMinMaxCoord.xy + HalfResLocalCoordI + int2(OffsetX,	0)		, HalfResDepthMinMaxCoord.xy, HalfResDepthMinMaxCoord.zw)]);
-	HalfResDepths.z = ConvertFromDeviceZ(HalfResDepthTexture[clamp(HalfResDepthMinMaxCoord.xy + HalfResLocalCoordI + int2(OffsetX,	OffsetY), HalfResDepthMinMaxCoord.xy, HalfResDepthMinMaxCoord.zw)]);
-	HalfResDepths.w = ConvertFromDeviceZ(HalfResDepthTexture[clamp(HalfResDepthMinMaxCoord.xy + HalfResLocalCoordI + int2(0,		OffsetY), HalfResDepthMinMaxCoord.xy, HalfResDepthMinMaxCoord.zw)]);
+	HalfResDepths.x = ConvertFromDeviceZ(HalfResDepthTexture[clamp(FSSSGlobals.HalfResDepthMinMaxCoord.xy + HalfResLocalCoordI							, FSSSGlobals.HalfResDepthMinMaxCoord.xy, FSSSGlobals.HalfResDepthMinMaxCoord.zw)]);
+	HalfResDepths.y = ConvertFromDeviceZ(HalfResDepthTexture[clamp(FSSSGlobals.HalfResDepthMinMaxCoord.xy + HalfResLocalCoordI + int2(OffsetX,	0)		, FSSSGlobals.HalfResDepthMinMaxCoord.xy, FSSSGlobals.HalfResDepthMinMaxCoord.zw)]);
+	HalfResDepths.z = ConvertFromDeviceZ(HalfResDepthTexture[clamp(FSSSGlobals.HalfResDepthMinMaxCoord.xy + HalfResLocalCoordI + int2(OffsetX,	OffsetY), FSSSGlobals.HalfResDepthMinMaxCoord.xy, FSSSGlobals.HalfResDepthMinMaxCoord.zw)]);
+	HalfResDepths.w = ConvertFromDeviceZ(HalfResDepthTexture[clamp(FSSSGlobals.HalfResDepthMinMaxCoord.xy + HalfResLocalCoordI + int2(0,		OffsetY), FSSSGlobals.HalfResDepthMinMaxCoord.xy, FSSSGlobals.HalfResDepthMinMaxCoord.zw)]);
 
 	float4 DeltaDepths = abs(HalfResDepths - FullResDepth.xxxx);
 
@@ -232,18 +193,18 @@ void HalfResolutionUpsamplerPS(
 			UpsamplingCoord = HalfResLocalCoordI + int2(0,			OffsetY);
 		}
 
-		SourceTextureColor = HalfResolutionSourceTexture[UpsamplingCoord].rgba;
-		SourceTextureColor2 = HalfResolutionSource2Texture[UpsamplingCoord].r;
-		OutColor = CombineSourceTextures(SourceTextureColor, SourceTextureColor2);
+		SourceTextureColor = FSSSGlobals.FogSeparateCompositionTexture0[UpsamplingCoord].rgba;
+		SourceTextureColor2 = FSSSGlobals.FogSeparateCompositionTexture1[UpsamplingCoord].r;
+		OutColor = FSSSCombineSourceTextures(SourceTextureColor, SourceTextureColor2);
 		return;
 		//HalfResLocalCoord = UpsamplingCoord;
 	}
 
 
 
-	SourceTextureColor = HalfResolutionSourceTexture.SampleLevel(HalfResolutionSourceSampler, HalfResLocalCoord * TexelCoordToUVs, 0).rgba;
-	SourceTextureColor2 = HalfResolutionSource2Texture.SampleLevel(HalfResolutionSource2Sampler, HalfResLocalCoord * TexelCoordToUVs, 0).r;
-	OutColor = CombineSourceTextures(SourceTextureColor, SourceTextureColor2);
+	SourceTextureColor = FSSSGlobals.FogSeparateCompositionTexture0.SampleLevel(FSSSGlobals.FogSeparateCompositionTexture0Sampler, HalfResLocalCoord * FSSSGlobals.TexelCoordToUVs, 0).rgba;
+	SourceTextureColor2 = FSSSGlobals.FogSeparateCompositionTexture1.SampleLevel(FSSSGlobals.FogSeparateCompositionTexture1Sampler, HalfResLocalCoord * FSSSGlobals.TexelCoordToUVs, 0).r;
+	OutColor = FSSSCombineSourceTextures(SourceTextureColor, SourceTextureColor2);
 
 
 
diff --git a/Engine/Shaders/Private/FogSeparateComposition.ush b/Engine/Shaders/Private/FogSeparateComposition.ush
new file mode 100644
--- /dev/null
+++ b/Engine/Shaders/Private/FogSeparateComposition.ush
@@ -0,0 +1,44 @@
+// Copyright Epic Games, Inc. All Rights Reserved.
+
+/*=============================================================================
+	FogSeparateComposition.ush: Common code for composing fog from separate composition textures
+	This code can be used by multiple passes (fog upsampling and cloud composition) to apply
+	fog textures onto the scene.
+=============================================================================*/
+
+#pragma once
+
+
+
+float4 FSSSCombineSourceTextures(float4 SourceTextureColor, float SourceTextureColor2)
+{
+	return float4(SourceTextureColor.rgb, SourceTextureColor2);
+}
+
+float4 FSSSFullResolutionSample(float2 FullResLocalCoord)
+{
+	float4 SourceTextureColor = 0;
+	float SourceTextureColor2 = 0;
+
+	SourceTextureColor = FSSSGlobals.FogSeparateCompositionTexture0[FullResLocalCoord];
+
+	const float W = SourceTextureColor.a;
+	if (W <= 0.0f)
+	{
+		clip(-1);
+		// Such value will lead to no difference when upsampling
+		return 0.0f;
+	}
+
+	SourceTextureColor2 = FSSSGlobals.FogSeparateCompositionTexture1[FullResLocalCoord].r;
+
+	const float Mip = max(0.0, FSSSGlobals.FSSSSpreadScale * log2(W / 2));
+
+	// This is a custom fade based on W for the clip above to note cause discontinuity.
+	const float FadeBasedOnW = saturate(W);
+
+	SourceTextureColor = FSSSGlobals.FogSeparateCompositionTexture0.SampleLevel(FSSSGlobals.FogSeparateCompositionTexture0Sampler, FullResLocalCoord * FSSSGlobals.TexelCoordToUVs, Mip).rgba * FadeBasedOnW;
+	return FSSSCombineSourceTextures(SourceTextureColor, SourceTextureColor2);
+}
+
+
diff --git a/Engine/Shaders/Private/VolumetricRenderTarget.usf b/Engine/Shaders/Private/VolumetricRenderTarget.usf
--- a/Engine/Shaders/Private/VolumetricRenderTarget.usf
+++ b/Engine/Shaders/Private/VolumetricRenderTarget.usf
@@ -27,6 +27,10 @@
 #include "LocalFogVolumes/LocalFogVolumeCommon.ush"
 #endif // PERMUTATION_SUPPORT_LOCAL_FOG_VOLUME
 
+#if PERMUTATION_APPLY_FOG_SEPARATE_COMPOSITION
+#include "FogSeparateComposition.ush"
+#endif // PERMUTATION_APPLY_FOG_SEPARATE_COMPOSITION
+
 // BufferSVPos is needed when using view matrices operation relying on VolumetricRenderTargetViewUniformBuffer such as GetScreenWorldDir.
 float4 ApplyFogToCloudSampleFromVPosAndDepth(float2 FullResViewSVPos, float2 BufferSVPos, float DepthKilometer, float4 CloudRGBT, float2 View_ViewRectMin)
 {
@@ -79,7 +83,6 @@ float4 ApplyFogToCloudSampleFromVPosAndDepth(float2 FullResViewSVPos, float2 Buf
 	HeightFogInscatteringAndTransmittance.rgb *= PrimaryView.PreExposure;
 
 	CloudRGBT.rgb = CloudRGBT.rgb * HeightFogInscatteringAndTransmittance.a + CloudCoverage * HeightFogInscatteringAndTransmittance.rgb;
-
 #endif // APPLY_FOG_LATE
 
 	return CloudRGBT;
@@ -1414,6 +1417,17 @@ void ComposeVolumetricRTOverScenePS(
 	if (ValidSampleCount > 0.0f)
 	{
 		OutputRt0 = DataAcc / ValidSampleCount;
+
+	#if PERMUTATION_APPLY_FOG_SEPARATE_COMPOSITION
+		const float CloudCoverage = saturate(1.0 - OutputRt0.a);
+		if (CloudCoverage > 0.0f)
+		{
+			// Apply FSSS as for the depth buffer, we cannot account for the cloud front depth. So it only work for distant clouds.
+			float2 FullResLocalCoord = SVPos.xy - View_ViewRectMin;
+			float4 FogComp = FSSSFullResolutionSample(FullResLocalCoord);
+			OutputRt0.rgb = OutputRt0.rgb * FogComp.a + CloudCoverage * FogComp.rgb;
+		}
+	#endif
 	}
 	else
 	{
diff --git a/Engine/Source/Runtime/Renderer/Private/DeferredShadingRenderer.cpp b/Engine/Source/Runtime/Renderer/Private/DeferredShadingRenderer.cpp
--- a/Engine/Source/Runtime/Renderer/Private/DeferredShadingRenderer.cpp
+++ b/Engine/Source/Runtime/Renderer/Private/DeferredShadingRenderer.cpp
@@ -3732,9 +3732,10 @@ void FDeferredShadingSceneRenderer::Render(FRDGBuilder& GraphBuilder, const FSce
 			{
 				const bool bComposeWithWater = bIsCameraUnderWater ? false : bShouldRenderSingleLayerWater;
 				ComposeVolumetricRenderTargetOverScene(
-					GraphBuilder, Views, SceneTextures.Color.Target, SceneTextures.Depth.Target,
+					GraphBuilder, Scene, Views, SceneTextures.Color.Target, SceneTextures.Depth.Target,
 					bComposeWithWater,
-					SceneWithoutWaterTextures, SceneTextures);
+					SceneWithoutWaterTextures, SceneTextures,
+					&FogSeparateCompositionViewResources);
 			}
 		};
 
diff --git a/Engine/Source/Runtime/Renderer/Private/FogRendering.cpp b/Engine/Source/Runtime/Renderer/Private/FogRendering.cpp
--- a/Engine/Source/Runtime/Renderer/Private/FogRendering.cpp
+++ b/Engine/Source/Runtime/Renderer/Private/FogRendering.cpp
@@ -205,7 +205,7 @@ void SetupFogUniformParameters(FRDGBuilder& GraphBuilder, const FViewInfo& View,
 
 	OutParameters.EnableFSSS							= 0.0f;
 	OutParameters.FSSSSceneColorScatteringAmountScale	= 0.0f;
-	OutParameters.FSSSSceneColorScatteringAmountPower	= 0.0f;
+	OutParameters.FSSSSceneColorScatteringAmountPower	= 1.0f;
 	OutParameters.FSSSSpreadScale						= 0.0f;
 	OutParameters.FSSSBlurControl						= 0.0f;
 	if(Scene && Scene->HasAnyExponentialHeightFog())
@@ -213,7 +213,7 @@ void SetupFogUniformParameters(FRDGBuilder& GraphBuilder, const FViewInfo& View,
 		const FExponentialHeightFogSceneInfo& HFog = Scene->ExponentialFogs[0];
 		OutParameters.EnableFSSS							= !HFog.bEnableFSSS ? 0.0f : 1.0f;
 		OutParameters.FSSSSceneColorScatteringAmountScale	= !HFog.bEnableFSSS ? 0.0f : HFog.FSSSSceneColorScatteringAmountScale;
-		OutParameters.FSSSSceneColorScatteringAmountPower	= !HFog.bEnableFSSS ? 0.0f : HFog.FSSSSceneColorScatteringAmountPower;
+		OutParameters.FSSSSceneColorScatteringAmountPower	= !HFog.bEnableFSSS ? 1.0f : HFog.FSSSSceneColorScatteringAmountPower;
 		OutParameters.FSSSSpreadScale						= !HFog.bEnableFSSS ? 0.0f : HFog.FSSSSpreadScale;
 		OutParameters.FSSSBlurControl						= !HFog.bEnableFSSS ? 0.0f : HFog.FSSSBlurControl;
 	}
diff --git a/Engine/Source/Runtime/Renderer/Private/FogSeparateComposition.cpp b/Engine/Source/Runtime/Renderer/Private/FogSeparateComposition.cpp
--- a/Engine/Source/Runtime/Renderer/Private/FogSeparateComposition.cpp
+++ b/Engine/Source/Runtime/Renderer/Private/FogSeparateComposition.cpp
@@ -244,6 +244,25 @@ void GenerateFSSSBlurMipChain(FRDGBuilder& GraphBuilder, FFogSSSBlurMipChainPara
 
 
 
+IMPLEMENT_GLOBAL_SHADER_PARAMETER_STRUCT(FSSSGlobalParameters, "FSSSGlobals");
+
+TRDGUniformBufferRef<FSSSGlobalParameters> GetDefaultFSSSGlobalParameters(FRDGBuilder& GraphBuilder)
+{
+	FSSSGlobalParameters* FSSSGlobals = GraphBuilder.AllocParameters<FSSSGlobalParameters>();
+
+	FSSSGlobals->HalfResDepthMinMaxCoord = FIntVector4::ZeroValue;
+	FSSSGlobals->TexelCoordToUVs = FVector2f::One();
+	FSSSGlobals->FullResToFogResScale = 1.0f;
+	FSSSGlobals->FSSSSpreadScale = 1.0f;
+
+	FSSSGlobals->FogSeparateCompositionTexture0 = GSystemTextures.GetBlackDummy(GraphBuilder);
+	FSSSGlobals->FogSeparateCompositionTexture1 = GSystemTextures.GetBlackDummy(GraphBuilder);
+	FSSSGlobals->FogSeparateCompositionTexture0Sampler = TStaticSamplerState<SF_Trilinear>::GetRHI();
+	FSSSGlobals->FogSeparateCompositionTexture1Sampler = TStaticSamplerState<SF_Trilinear>::GetRHI();
+
+	return GraphBuilder.CreateUniformBuffer(FSSSGlobals);
+}
+
 
 class FUpsampleHalfResTexturePS : public FGlobalShader
 {
@@ -254,16 +273,7 @@ class FUpsampleHalfResTexturePS : public FGlobalShader
 	using FPermutationDomain = TShaderPermutationDomain<FUpsampleFullResolution>;
 
 	BEGIN_SHADER_PARAMETER_STRUCT(FParameters, )
-		SHADER_PARAMETER(FIntVector4, HalfResDepthMinMaxCoord)
-		SHADER_PARAMETER(FVector2f, TexelCoordToUVs)
-		SHADER_PARAMETER(float, FullResToFogResScale)
-		SHADER_PARAMETER(float, FSSSSpreadScale)
-
-		SHADER_PARAMETER_RDG_TEXTURE(Texture2D, HalfResolutionSourceTexture)
-		SHADER_PARAMETER_SAMPLER(SamplerState, HalfResolutionSourceSampler)
-
-		SHADER_PARAMETER_RDG_TEXTURE(Texture2D, HalfResolutionSource2Texture)
-		SHADER_PARAMETER_SAMPLER(SamplerState, HalfResolutionSource2Sampler)
+		SHADER_PARAMETER_RDG_UNIFORM_BUFFER(FSSSGlobalParameters, FSSSGlobals)
 
 		SHADER_PARAMETER_RDG_TEXTURE(Texture2D, HalfResDepthTexture)
 		SHADER_PARAMETER_SAMPLER(SamplerState, HalfResDepthSampler)
@@ -357,26 +367,29 @@ void UpsampleFogSeparateCompositionTextureForView(
 	const FVector2f SourceTextureResolution = FVector2f(FogSeparateCompositionViewResources.FogSeparateCompositionTexture0->Desc.Extent.X, FogSeparateCompositionViewResources.FogSeparateCompositionTexture0->Desc.Extent.Y);
 	const bool bSceneRequestsFSSS = SceneRequestsFSSS(Scene);
 
-	PassParameters->HalfResDepthMinMaxCoord = FIntVector4(HalfResViewRect.Min.X, HalfResViewRect.Min.Y, HalfResViewRect.Max.X, HalfResViewRect.Max.Y);
+	// Create global shader parameters for fog composition
+	FSSSGlobalParameters* FSSSGlobals = GraphBuilder.AllocParameters<FSSSGlobalParameters>();
+	FSSSGlobals->HalfResDepthMinMaxCoord = FIntVector4(HalfResViewRect.Min.X, HalfResViewRect.Min.Y, HalfResViewRect.Max.X, HalfResViewRect.Max.Y);
 
 	const FVector2f TextureFullResolution = FVector2f(FogSeparateCompositionViewResources.FogSeparateCompositionTexture0->Desc.Extent);
 	const FVector2f TextureViewResolution = FVector2f(View.ViewRect.Size());
-	PassParameters->TexelCoordToUVs = (TextureViewResolution / TextureFullResolution) / TextureFullResolution;
-	
-	PassParameters->FullResToFogResScale = 1.0f / float(ResolutionDivider);
-	PassParameters->FSSSSpreadScale = bSceneRequestsFSSS ? FogInfo.FSSSSpreadScale : 0.0f; // Spread of 0 if FSSS is not enabled and mip chain not generated.
+	FSSSGlobals->TexelCoordToUVs = (TextureViewResolution / TextureFullResolution) / TextureFullResolution;
+
+	FSSSGlobals->FullResToFogResScale = 1.0f / float(ResolutionDivider);
+	FSSSGlobals->FSSSSpreadScale = bSceneRequestsFSSS ? FogInfo.FSSSSpreadScale : 0.0f; // Spread of 0 if FSSS is not enabled and mip chain not generated.
+	FSSSGlobals->FogSeparateCompositionTexture0 = FogSeparateCompositionViewResources.FogSeparateCompositionTexture0;
+	FSSSGlobals->FogSeparateCompositionTexture1 = FogSeparateCompositionViewResources.FogSeparateCompositionTexture1;
+	FSSSGlobals->FogSeparateCompositionTexture0Sampler = TStaticSamplerState<SF_Trilinear>::GetRHI();
+	FSSSGlobals->FogSeparateCompositionTexture1Sampler = TStaticSamplerState<SF_Trilinear>::GetRHI();
+
+	FogSeparateCompositionViewResources.FSSSGlobalsUniformBuffer = GraphBuilder.CreateUniformBuffer(FSSSGlobals);
+	PassParameters->FSSSGlobals = FogSeparateCompositionViewResources.FSSSGlobalsUniformBuffer;
 
 	PassParameters->HalfResDepthTexture = SceneTextures.DitheredHalfResDepth;
 	PassParameters->HalfResDepthSampler = TStaticSamplerState<SF_Point>::GetRHI();
 	PassParameters->FullResDepthTexture = SceneTextures.Depth.Resolve;
 	PassParameters->FullResDepthSampler = TStaticSamplerState<SF_Point>::GetRHI();
 
-	PassParameters->HalfResolutionSourceTexture = FogSeparateCompositionViewResources.FogSeparateCompositionTexture0;
-	PassParameters->HalfResolutionSourceSampler = TStaticSamplerState<SF_Trilinear>::GetRHI();
-
-	PassParameters->HalfResolutionSource2Texture = FogSeparateCompositionViewResources.FogSeparateCompositionTexture1;
-	PassParameters->HalfResolutionSource2Sampler = TStaticSamplerState<SF_Trilinear>::GetRHI();
-
 	PassParameters->ViewUniformBuffer = View.ViewUniformBuffer;
 
 	PassParameters->RenderTargets[0] = FRenderTargetBinding(SceneTextures.Color.Target, ERenderTargetLoadAction::ELoad);
diff --git a/Engine/Source/Runtime/Renderer/Private/FogSeparateComposition.h b/Engine/Source/Runtime/Renderer/Private/FogSeparateComposition.h
--- a/Engine/Source/Runtime/Renderer/Private/FogSeparateComposition.h
+++ b/Engine/Source/Runtime/Renderer/Private/FogSeparateComposition.h
@@ -11,6 +11,7 @@
 
 #include "CoreMinimal.h"
 #include "RenderGraphFwd.h"
+#include "RenderGraphResources.h"
 
 
 struct FFogSSSBlurMipChainParameters
@@ -24,6 +25,20 @@ struct FFogSSSBlurMipChainParameters
 void GenerateFSSSBlurMipChain(FRDGBuilder& GraphBuilder, FFogSSSBlurMipChainParameters& Parameters);
 
 
+BEGIN_GLOBAL_SHADER_PARAMETER_STRUCT(FSSSGlobalParameters,)
+	SHADER_PARAMETER(FIntVector4, HalfResDepthMinMaxCoord)
+	SHADER_PARAMETER(FVector2f, TexelCoordToUVs)
+	SHADER_PARAMETER(float, FullResToFogResScale)
+	SHADER_PARAMETER(float, FSSSSpreadScale)
+	SHADER_PARAMETER_RDG_TEXTURE(Texture2D<float4>, FogSeparateCompositionTexture0)
+	SHADER_PARAMETER_RDG_TEXTURE(Texture2D<float>,  FogSeparateCompositionTexture1)
+	SHADER_PARAMETER_SAMPLER(SamplerState,			FogSeparateCompositionTexture0Sampler)
+	SHADER_PARAMETER_SAMPLER(SamplerState,			FogSeparateCompositionTexture1Sampler)
+END_GLOBAL_SHADER_PARAMETER_STRUCT()
+
+TRDGUniformBufferRef<FSSSGlobalParameters> GetDefaultFSSSGlobalParameters(FRDGBuilder& GraphBuilder);
+
+
 class FFogSeparateCompositionViewResources
 {
 public:
@@ -37,6 +52,9 @@ public:
 
 	// Texture for side data used for FSSS (Fog Screen Space Scattering) when enabled
 	FRDGTextureRef FogSeparateCompositionTexture1 = nullptr;
+
+	// Uniform buffer containing global fog composition parameters
+	TRDGUniformBufferRef<FSSSGlobalParameters> FSSSGlobalsUniformBuffer = nullptr;
 };
 
 bool ShouldRenderFogUsingSeparateComposition(
diff --git a/Engine/Source/Runtime/Renderer/Private/VolumetricRenderTarget.cpp b/Engine/Source/Runtime/Renderer/Private/VolumetricRenderTarget.cpp
--- a/Engine/Source/Runtime/Renderer/Private/VolumetricRenderTarget.cpp
+++ b/Engine/Source/Runtime/Renderer/Private/VolumetricRenderTarget.cpp
@@ -17,6 +17,7 @@
 #include "LogRenderer.h"
 #include "SceneViewState.h"
 #include "RenderTargetPool.h"
+#include "FogSeparateComposition.h"
 
 static TAutoConsoleVariable<int32> CVarVolumetricRenderTarget(
 	TEXT("r.VolumetricRenderTarget"), 1,
@@ -73,6 +74,11 @@ static TAutoConsoleVariable<int32> CVarVolumetricRenderTargetSimulateNullResolut
 	TEXT("Simulate a view port resolution of 0x0 for debugging."),
 	ECVF_RenderThreadSafe);
 
+static TAutoConsoleVariable<int32> CVarFogScreenSpaceScatteringApplyOnCloud(
+	TEXT("r.Fog.ScreenSpaceScattering.ApplyOnCloud"), 1,
+	TEXT("Enable FSSS on volumetirc clouds, applied during the upsampling stage on the full resolution buffer when using late fog only. This only works for distant clouds"),
+	ECVF_Scalability | ECVF_RenderThreadSafe);
+
 static float GetUvNoiseSampleAcceptanceWeight()
 {
 	return FMath::Max(0.0f, CVarVolumetricRenderTargetUvNoiseSampleAcceptanceWeight.GetValueOnRenderThread());
@@ -929,17 +935,19 @@ class FComposeVolumetricRTOverScenePS : public FGlobalShader
 
 	class FUpsamplingMode : SHADER_PERMUTATION_RANGE_INT("PERMUTATION_UPSAMPLINGMODE", 0, 5);
 	class FRenderUnderWaterBuffer : SHADER_PERMUTATION_BOOL("PERMUTATION_RENDER_UNDERWATER_BUFFER");	// Render into the water scene color buffer (used when rendering from water system)
-	class FRenderCameraComposeWithWater : SHADER_PERMUTATION_BOOL("PERMUTATION_COMPOSE_WITH_WATER");	// When water us used and the camera is under water, use that permutation (to handle camera intersection with water and double cloud composition) 
+	class FRenderCameraComposeWithWater : SHADER_PERMUTATION_BOOL("PERMUTATION_COMPOSE_WITH_WATER");	// When water us used and the camera is under water, use that permutation (to handle camera intersection with water and double cloud composition)
 	class FMSAASampleCount : SHADER_PERMUTATION_SPARSE_INT("PERMUTATION_MSAA_SAMPLE_COUNT", 1, 2, 4, 8);
 	class FApplyFog : SHADER_PERMUTATION_BOOL("PERMUTATION_APPLY_FOG");
 	class FApplyLocalFogVolume : SHADER_PERMUTATION_BOOL("PERMUTATION_APPLY_LOCAL_FOG_VOLUME");
 	class FMinMaxDepthAvailable : SHADER_PERMUTATION_BOOL("PERMUTATION_MINMAXDEPTH_AVAILABLE");
-	using FPermutationDomain = TShaderPermutationDomain<FUpsamplingMode, FRenderUnderWaterBuffer, FRenderCameraComposeWithWater, FMSAASampleCount, FApplyFog, FApplyLocalFogVolume, FMinMaxDepthAvailable>;
+	class FApplyFogSeparateComposition : SHADER_PERMUTATION_BOOL("PERMUTATION_APPLY_FOG_SEPARATE_COMPOSITION");	// Apply fog from separate composition textures on top of clouds
+	using FPermutationDomain = TShaderPermutationDomain<FUpsamplingMode, FRenderUnderWaterBuffer, FRenderCameraComposeWithWater, FMSAASampleCount, FApplyFog, FApplyLocalFogVolume, FMinMaxDepthAvailable, FApplyFogSeparateComposition>;
 
 	BEGIN_SHADER_PARAMETER_STRUCT(FParameters, )
 		SHADER_PARAMETER_STRUCT_REF(FViewUniformShaderParameters, ViewUniformBuffer)
 		SHADER_PARAMETER_STRUCT(FLocalFogVolumeUniformParameters, LFV)
 		SHADER_PARAMETER_RDG_UNIFORM_BUFFER(FFogUniformParameters, FogStruct)
+		SHADER_PARAMETER_RDG_UNIFORM_BUFFER(FSSSGlobalParameters, FSSSGlobals)
 		SHADER_PARAMETER_RDG_TEXTURE(Texture2D, VolumetricTexture)
 		SHADER_PARAMETER_RDG_TEXTURE(Texture2D, VolumetricSecondaryTexture)
 		SHADER_PARAMETER_RDG_TEXTURE(Texture2D, VolumetricDepthTexture)
@@ -965,6 +973,12 @@ class FComposeVolumetricRTOverScenePS : public FGlobalShader
 
 	static FPermutationDomain RemapPermutation(FPermutationDomain PermutationVector)
 	{
+		if (PermutationVector.Get<FApplyFogSeparateComposition>())
+		{
+			// When we sample separate fog composition buffer, it already contains all the fog contributions.
+			PermutationVector.Set<FApplyFog>(false);
+			PermutationVector.Set<FApplyLocalFogVolume>(false);
+		}
 		return PermutationVector;
 	}
 
@@ -977,6 +991,10 @@ class FComposeVolumetricRTOverScenePS : public FGlobalShader
 			// We only compile the MSAA support when Forward shading is enabled because MSAA can only be used in this case.
 			return false;
 		}
+		if (PermutationVector.Get<FApplyFogSeparateComposition>() && (PermutationVector.Get<FApplyFog>() || PermutationVector.Get<FApplyLocalFogVolume>()))
+		{
+			return false; // When we sample separate fog composition buffer, it already contains all the fog contributions.
+		}
 
 		return ShouldPipelineCompileVolumetricRenderTargetShaders(Parameters.Platform);
 	}
@@ -1030,14 +1048,16 @@ static int32 GetCompositionUpsamplingMode(uint32 VRTMode, int32 UpsamplingMode)
 
 void ComposeVolumetricRenderTargetOverScene(
 	FRDGBuilder& GraphBuilder,
+	const FScene* Scene,
 	TArrayView<FViewInfo> Views,
 	FRDGTextureRef SceneColorTexture,
 	FRDGTextureRef SceneDepthTexture,
 	bool bComposeWithWater,
 	const FSceneWithoutWaterTextures& WaterPassData,
-	const FMinimalSceneTextures& SceneTextures)
+	const FMinimalSceneTextures& SceneTextures,
+	TArray<FFogSeparateCompositionViewResources>* FogSeparateCompositionViewResources)
 {
-	if (!AnyViewRequiresProcessing(Views))
+	if (!Scene || !AnyViewRequiresProcessing(Views))
 	{
 		return;
 	}
@@ -1096,14 +1116,25 @@ void ComposeVolumetricRenderTargetOverScene(
 
 		const bool bShouldVolumetricCloudsApplyFogDuringUpsampling = ShouldVolumetricCloudsApplyFogDuringUpsampling(ViewInfo);
 
+		// Check if we should apply fog separate composition on clouds
+		check(!FogSeparateCompositionViewResources || FogSeparateCompositionViewResources->Num() > ViewIndex);
+		const bool bApplyFogSeparateComposition = FogSeparateCompositionViewResources 
+			&& (*FogSeparateCompositionViewResources)[ViewIndex].FogSeparateCompositionTexture0 != nullptr 
+			&& (*FogSeparateCompositionViewResources)[ViewIndex].FogSeparateCompositionTexture1 != nullptr
+			&& CVarFogScreenSpaceScatteringApplyOnCloud.GetValueOnAnyThread() > 0;
+
 		FComposeVolumetricRTOverScenePS::FPermutationDomain PermutationVector;
 		PermutationVector.Set<FComposeVolumetricRTOverScenePS::FUpsamplingMode>(UpsamplingMode);
 		PermutationVector.Set<FComposeVolumetricRTOverScenePS::FRenderUnderWaterBuffer>(0);
 		PermutationVector.Set<FComposeVolumetricRTOverScenePS::FRenderCameraComposeWithWater>((bComposeWithWater && ViewInfo.IsUnderwater()) ? 1 : 0);
 		PermutationVector.Set<FComposeVolumetricRTOverScenePS::FMSAASampleCount>(SceneDepthTexture->Desc.NumSamples);
-		PermutationVector.Set<FComposeVolumetricRTOverScenePS::FApplyFog>(bShouldVolumetricCloudsApplyFogDuringUpsampling && (ViewInfo.Family->Scene->HasAnyExponentialHeightFog() || ViewInfo.LocalFogVolumeViewData.GPUInstanceCount > 0) && ShouldRenderFog(*ViewInfo.Family));
-		PermutationVector.Set<FComposeVolumetricRTOverScenePS::FApplyLocalFogVolume>(bShouldVolumetricCloudsApplyFogDuringUpsampling && ViewInfo.LocalFogVolumeViewData.GPUInstanceCount > 0);
+		if (!bApplyFogSeparateComposition)
+		{
+			PermutationVector.Set<FComposeVolumetricRTOverScenePS::FApplyFog>(bShouldVolumetricCloudsApplyFogDuringUpsampling && (ViewInfo.Family->Scene->HasAnyExponentialHeightFog() || ViewInfo.LocalFogVolumeViewData.GPUInstanceCount > 0) && ShouldRenderFog(*ViewInfo.Family));
+			PermutationVector.Set<FComposeVolumetricRTOverScenePS::FApplyLocalFogVolume>(bShouldVolumetricCloudsApplyFogDuringUpsampling && ViewInfo.LocalFogVolumeViewData.GPUInstanceCount > 0);
+		}
 		PermutationVector.Set<FComposeVolumetricRTOverScenePS::FMinMaxDepthAvailable>(bMinMaxDepth);
+		PermutationVector.Set<FComposeVolumetricRTOverScenePS::FApplyFogSeparateComposition>(bApplyFogSeparateComposition);
 		TShaderMapRef<FComposeVolumetricRTOverScenePS> PixelShader(ViewInfo.ShaderMap, PermutationVector);
 
 		FComposeVolumetricRTOverScenePS::FParameters* PassParameters = GraphBuilder.AllocParameters<FComposeVolumetricRTOverScenePS::FParameters>();
@@ -1141,6 +1172,18 @@ void ComposeVolumetricRenderTargetOverScene(
 		PassParameters->FogStruct = CreateFogUniformBuffer(GraphBuilder, ViewInfo);
 		PassParameters->LFV = ViewInfo.LocalFogVolumeViewData.UniformParametersStruct;
 
+		// Bind fog separate composition parameters if available
+		if (bApplyFogSeparateComposition)
+		{
+			const FFogSeparateCompositionViewResources& FogCompositionRes = (*FogSeparateCompositionViewResources)[ViewIndex];
+			check(FogCompositionRes.FSSSGlobalsUniformBuffer);
+			PassParameters->FSSSGlobals = FogCompositionRes.FSSSGlobalsUniformBuffer;
+		}
+		else
+		{
+			PassParameters->FSSSGlobals = GetDefaultFSSSGlobalParameters(GraphBuilder);
+		}
+
 		FPixelShaderUtils::AddFullscreenPass<FComposeVolumetricRTOverScenePS>(
 			GraphBuilder, ViewInfo.ShaderMap, RDG_EVENT_NAME("VolumetricComposeCloudOverScene"), PixelShader, PassParameters, ViewInfo.ViewRect,
 			PreMultipliedColorTransmittanceBlend);
diff --git a/Engine/Source/Runtime/Renderer/Private/VolumetricRenderTarget.h b/Engine/Source/Runtime/Renderer/Private/VolumetricRenderTarget.h
--- a/Engine/Source/Runtime/Renderer/Private/VolumetricRenderTarget.h
+++ b/Engine/Source/Runtime/Renderer/Private/VolumetricRenderTarget.h
@@ -16,6 +16,7 @@ class FScene;
 class FViewInfo;
 struct FSceneWithoutWaterTextures;
 struct FMinimalSceneTextures;
+class FFogSeparateCompositionViewResources;
 
 void InitVolumetricRenderTargetForViews(FRDGBuilder& GraphBuilder, TArrayView<FViewInfo> Views, const FMinimalSceneTextures& SceneTextures);
 void ResetVolumetricRenderTargetForViews(FRDGBuilder& GraphBuilder, TArrayView<FViewInfo> Views);
@@ -29,12 +30,14 @@ void ReconstructVolumetricRenderTarget(
 
 void ComposeVolumetricRenderTargetOverScene(
 	FRDGBuilder& GraphBuilder,
+	const FScene* Scene,
 	TArrayView<FViewInfo> Views,
 	FRDGTextureRef SceneColorTexture,
 	FRDGTextureRef SceneDepthResolveTexture,
 	bool bShouldRenderSingleLayerWater,
 	const FSceneWithoutWaterTextures& WaterPassData,
-	const FMinimalSceneTextures& SceneTextures);
+	const FMinimalSceneTextures& SceneTextures,
+	TArray<FFogSeparateCompositionViewResources>* FogSeparateCompositionViewResources = nullptr);
 
 void ComposeVolumetricRenderTargetOverSceneUnderWater(
 	FRDGBuilder& GraphBuilder,
